[Teds] Move UI file to TedsUI from TedsCore module, and remove shortname from TedsCore module

#rnx

diff --git a/Engine/Plugins/Experimental/EditorDataStorage/Source/TedsCore/Private/UI/TypedElementSlateWidgetReferenceColumnUpdateProcessor.cpp b/Engine/Plugins/Experimental/EditorDataStorage/Source/TedsCore/Private/UI/TypedElementSlateWidgetReferenceColumnUpdateProcessor.cpp
deleted file mode 100644
diff --git a/Engine/Plugins/Experimental/EditorDataStorage/Source/TedsCore/Private/UI/TypedElementSlateWidgetReferenceColumnUpdateProcessor.h b/Engine/Plugins/Experimental/EditorDataStorage/Source/TedsCore/Private/UI/TypedElementSlateWidgetReferenceColumnUpdateProcessor.h
deleted file mode 100644
diff --git a/Engine/Plugins/Experimental/EditorDataStorage/Source/TedsCore/TedsCore.Build.cs b/Engine/Plugins/Experimental/EditorDataStorage/Source/TedsCore/TedsCore.Build.cs
--- a/Engine/Plugins/Experimental/EditorDataStorage/Source/TedsCore/TedsCore.Build.cs
+++ b/Engine/Plugins/Experimental/EditorDataStorage/Source/TedsCore/TedsCore.Build.cs
@@ -30,7 +30,5 @@ public class TedsCore : ModuleRules
 			PrivateDependencyModuleNames.AddRange(new string[] {});
 			DynamicallyLoadedModuleNames.AddRange(new string[] {});
 		}
-
-		ShortName = "TElmntsDataStor";
 	}
 }
diff --git a/Engine/Plugins/Experimental/EditorDataStorage/Source/TedsUI/Private/Processors/WidgetReferenceColumnUpdateProcessor.cpp b/Engine/Plugins/Experimental/EditorDataStorage/Source/TedsUI/Private/Processors/WidgetReferenceColumnUpdateProcessor.cpp
new file mode 100644
--- /dev/null
+++ b/Engine/Plugins/Experimental/EditorDataStorage/Source/TedsUI/Private/Processors/WidgetReferenceColumnUpdateProcessor.cpp
@@ -0,0 +1,58 @@
+// Copyright Epic Games, Inc. All Rights Reserved.
+
+#include "Processors/WidgetReferenceColumnUpdateProcessor.h"
+
+#include "Elements/Columns/TypedElementSlateWidgetColumns.h"
+#include "Elements/Framework/TypedElementQueryBuilder.h"
+
+void UWidgetReferenceColumnUpdateFactory::RegisterQueries(ITypedElementDataStorageInterface& DataStorage)
+{
+	RegisterDeleteRowOnWidgetDeleteQuery(DataStorage);
+	RegisterDeleteColumnOnWidgetDeleteQuery(DataStorage);
+}
+
+void UWidgetReferenceColumnUpdateFactory::RegisterDeleteRowOnWidgetDeleteQuery(ITypedElementDataStorageInterface& DataStorage) const
+{
+	using namespace TypedElementQueryBuilder;
+	using DSI = ITypedElementDataStorageInterface;
+
+	DataStorage.RegisterQuery(
+    	Select(
+    		TEXT("Delete row with deleted widget"),
+    		FPhaseAmble(FPhaseAmble::ELocation::Preamble, DSI::EQueryTickPhase::FrameEnd),
+    		[](DSI::IQueryContext& Context, TypedElementRowHandle Row, const FTypedElementSlateWidgetReferenceColumn& WidgetReference)
+    		{
+    			if (!WidgetReference.TedsWidget.IsValid())
+    			{
+    				Context.RemoveRow(Row);
+    			}
+    		}
+    	)
+    	.Where()
+    		.All<FTypedElementSlateWidgetReferenceDeletesRowTag>()
+    	.Compile()
+    	);
+}
+
+void UWidgetReferenceColumnUpdateFactory::RegisterDeleteColumnOnWidgetDeleteQuery(ITypedElementDataStorageInterface& DataStorage) const
+{
+	using namespace TypedElementQueryBuilder;
+	using DSI = ITypedElementDataStorageInterface;
+
+	DataStorage.RegisterQuery(
+		Select(
+			TEXT("Delete widget columns for deleted widget"),
+			FPhaseAmble(FPhaseAmble::ELocation::Preamble, DSI::EQueryTickPhase::FrameEnd),
+			[](DSI::IQueryContext& Context, TypedElementRowHandle Row, const FTypedElementSlateWidgetReferenceColumn& WidgetReference)
+			{
+				if (!WidgetReference.TedsWidget.IsValid())
+				{
+					Context.RemoveColumns<FTypedElementSlateWidgetReferenceColumn>(Row);
+				}
+			}
+		)
+		.Where()
+			.None<FTypedElementSlateWidgetReferenceDeletesRowTag>()
+		.Compile()
+	);
+}
diff --git a/Engine/Plugins/Experimental/EditorDataStorage/Source/TedsUI/Private/Processors/WidgetReferenceColumnUpdateProcessor.h b/Engine/Plugins/Experimental/EditorDataStorage/Source/TedsUI/Private/Processors/WidgetReferenceColumnUpdateProcessor.h
new file mode 100644
--- /dev/null
+++ b/Engine/Plugins/Experimental/EditorDataStorage/Source/TedsUI/Private/Processors/WidgetReferenceColumnUpdateProcessor.h
@@ -0,0 +1,28 @@
+// Copyright Epic Games, Inc. All Rights Reserved.
+
+#pragma once
+
+#include "Elements/Interfaces/TypedElementDataStorageFactory.h"
+#include "UObject/ObjectMacros.h"
+
+#include "WidgetReferenceColumnUpdateProcessor.generated.h"
+
+/**
+ * Queries that check whether or not a widget still exists. If it has been deleted
+ * then it will remove the column from the Data Storage or deletes the entire row if
+ * the FTypedElementSlateWidgetReferenceDeletesRowTag was found.
+ */
+UCLASS()
+class UWidgetReferenceColumnUpdateFactory : public UTypedElementDataStorageFactory
+{
+	GENERATED_BODY()
+
+public:
+	~UWidgetReferenceColumnUpdateFactory() override = default;
+
+	void RegisterQueries(ITypedElementDataStorageInterface& DataStorage) override;
+
+private:
+	void RegisterDeleteRowOnWidgetDeleteQuery(ITypedElementDataStorageInterface& DataStorage) const;
+	void RegisterDeleteColumnOnWidgetDeleteQuery(ITypedElementDataStorageInterface& DataStorage) const;
+};
