[Teds] Rename TedsUI files and classes to Teds naming scheme

#rnx

diff --git a/Engine/Plugins/Experimental/EditorDataStorage/Source/TedsUI/Private/TedsUIModule.cpp b/Engine/Plugins/Experimental/EditorDataStorage/Source/TedsUI/Private/TedsUIModule.cpp
new file mode 100644
--- /dev/null
+++ b/Engine/Plugins/Experimental/EditorDataStorage/Source/TedsUI/Private/TedsUIModule.cpp
@@ -0,0 +1,60 @@
+// Copyright Epic Games, Inc. All Rights Reserved.
+
+#include "TedsUIModule.h"
+
+#include "Elements/Framework/TypedElementRegistry.h"
+#include "Elements/Interfaces/TypedElementDataStorageUiInterface.h"
+#include "HAL/IConsoleManager.h"
+#include "Misc/CoreDelegates.h"
+#include "UObject/UObjectGlobals.h"
+
+#define LOCTEXT_NAMESPACE "TedsUIModule"
+
+namespace UE::TypedElementsDataStorageUI::Private
+{
+	FAutoConsoleCommandWithOutputDevice PrintWidgetPurposesConsoleCommand(
+	TEXT("TEDS.UI.PrintWidgetPurposes"),
+	TEXT("Prints a list of all the known widget purposes."),
+	FConsoleCommandWithOutputDeviceDelegate::CreateLambda([](FOutputDevice& Output)
+		{
+			UTypedElementRegistry* Registry = UTypedElementRegistry::GetInstance();
+			if (ITypedElementDataStorageUiInterface* UiStorage = Registry->GetMutableDataStorageUi())
+			{
+				Output.Log(TEXT("The Typed Elements Data Storage has recorded the following widget purposes:"));
+				UiStorage->ListWidgetPurposes(
+					[&Output](FName Purpose, ITypedElementDataStorageUiInterface::EPurposeType, const FText& Description)
+					{
+						Output.Logf(TEXT("    %s - %s"), *Purpose.ToString(), *Description.ToString());
+					});
+				Output.Log(TEXT("End of Typed Elements Data Storage widget purpose list."));
+			}
+		}));
+
+	static bool bUseNewTEDSUIWidgets = false;
+	FAutoConsoleVariableRef UseNewWidgetsCvar(
+		TEXT("TEDS.UI.UseNewWidgets"),
+		bUseNewTEDSUIWidgets,
+		TEXT("If true, TEDS UI will use new attribute binding driven widgets (needs to be set at startup)")
+	);
+}
+
+void FTedsUIModule::StartupModule()
+{
+}
+
+void FTedsUIModule::ShutdownModule()
+{
+}
+
+void FTedsUIModule::AddReferencedObjects(FReferenceCollector& Collector)
+{
+}
+
+FString FTedsUIModule::GetReferencerName() const
+{
+	return TEXT("TEDS: Editor Data Storage UI Module");
+}
+
+#undef LOCTEXT_NAMESPACE
+	
+IMPLEMENT_MODULE(FTedsUIModule, TedsUI)
\ No newline at end of file
diff --git a/Engine/Plugins/Experimental/EditorDataStorage/Source/TedsUI/Private/TedsUIModule.h b/Engine/Plugins/Experimental/EditorDataStorage/Source/TedsUI/Private/TedsUIModule.h
new file mode 100644
--- /dev/null
+++ b/Engine/Plugins/Experimental/EditorDataStorage/Source/TedsUI/Private/TedsUIModule.h
@@ -0,0 +1,29 @@
+// Copyright Epic Games, Inc. All Rights Reserved.
+
+#pragma once
+
+#include "Modules/ModuleManager.h"
+#include "UObject/GCObject.h"
+
+class FReferenceCollector;
+
+class FTedsUIModule : public IModuleInterface, public FGCObject
+{
+public:
+	~FTedsUIModule() override = default;
+
+	//
+	// IModuleInterface
+	//
+
+	virtual void StartupModule() override;
+	virtual void ShutdownModule() override;
+
+	
+	//
+	// FGCObject
+	//
+	
+	void AddReferencedObjects(FReferenceCollector& Collector) override;
+	FString GetReferencerName() const override;
+};
diff --git a/Engine/Plugins/Experimental/EditorDataStorage/Source/TedsUI/Private/TypedElementCounterWidget.cpp b/Engine/Plugins/Experimental/EditorDataStorage/Source/TedsUI/Private/TypedElementCounterWidget.cpp
deleted file mode 100644
diff --git a/Engine/Plugins/Experimental/EditorDataStorage/Source/TedsUI/Private/TypedElementExportedTextWidget.cpp b/Engine/Plugins/Experimental/EditorDataStorage/Source/TedsUI/Private/TypedElementExportedTextWidget.cpp
deleted file mode 100644
diff --git a/Engine/Plugins/Experimental/EditorDataStorage/Source/TedsUI/Private/TypedElementGeneralRegistration.cpp b/Engine/Plugins/Experimental/EditorDataStorage/Source/TedsUI/Private/TypedElementGeneralRegistration.cpp
deleted file mode 100644
diff --git a/Engine/Plugins/Experimental/EditorDataStorage/Source/TedsUI/Private/TypedElementLabelWidget.cpp b/Engine/Plugins/Experimental/EditorDataStorage/Source/TedsUI/Private/TypedElementLabelWidget.cpp
deleted file mode 100644
diff --git a/Engine/Plugins/Experimental/EditorDataStorage/Source/TedsUI/Private/TypedElementPackagePathWidget.cpp b/Engine/Plugins/Experimental/EditorDataStorage/Source/TedsUI/Private/TypedElementPackagePathWidget.cpp
deleted file mode 100644
diff --git a/Engine/Plugins/Experimental/EditorDataStorage/Source/TedsUI/Private/TypedElementsDataStorageUI.cpp b/Engine/Plugins/Experimental/EditorDataStorage/Source/TedsUI/Private/TypedElementsDataStorageUI.cpp
deleted file mode 100644
diff --git a/Engine/Plugins/Experimental/EditorDataStorage/Source/TedsUI/Private/Widgets/CounterWidget.cpp b/Engine/Plugins/Experimental/EditorDataStorage/Source/TedsUI/Private/Widgets/CounterWidget.cpp
new file mode 100644
--- /dev/null
+++ b/Engine/Plugins/Experimental/EditorDataStorage/Source/TedsUI/Private/Widgets/CounterWidget.cpp
@@ -0,0 +1,191 @@
+// Copyright Epic Games, Inc. All Rights Reserved.
+
+#include "Widgets/CounterWidget.h"
+
+#include "Elements/Columns/TypedElementSlateWidgetColumns.h"
+#include "Elements/Columns/TypedElementValueCacheColumns.h"
+#include "Elements/Framework/TypedElementQueryBuilder.h"
+#include "Elements/Framework/TypedElementRegistry.h"
+#include "Elements/Interfaces/TypedElementDataStorageInterface.h"
+#include "Interfaces/IMainFrameModule.h"
+#include "Layout/Margin.h"
+#include "Misc/CoreDelegates.h"
+#include "ToolMenu.h"
+#include "ToolMenus.h"
+#include "Widgets/DeclarativeSyntaxSupport.h"
+#include "Widgets/SWindow.h"
+#include "Widgets/Text/STextBlock.h"
+
+#define LOCTEXT_NAMESPACE "TedsUI_CounterWidget"
+
+
+FAutoConsoleCommand EnableCounterWidgetsConsoleCommand(
+	TEXT("TEDS.UI.EnableCounterWidgets"),
+	TEXT("Adds registered counter widgets to the bottom right status bar of the main editor window."),
+	FConsoleCommandDelegate::CreateLambda([]()
+		{
+			UCounterWidgetFactory::EnableCounterWidgets();
+		}));
+
+//
+// UCounterWidgetFactory
+//
+
+FName UCounterWidgetFactory::WigetPurpose(TEXT("LevelEditor.StatusBar.ToolBar"));
+bool UCounterWidgetFactory::bAreCounterWidgetsEnabled{ false };
+bool UCounterWidgetFactory::bHasBeenSetup{ false };
+
+UCounterWidgetFactory::UCounterWidgetFactory()
+{
+	if (bAreCounterWidgetsEnabled)
+	{
+		IMainFrameModule::Get().OnMainFrameCreationFinished().AddStatic(&UCounterWidgetFactory::SetupMainWindowIntegrations);
+	}
+}
+
+void UCounterWidgetFactory::RegisterQueries(ITypedElementDataStorageInterface& DataStorage)
+{
+	using namespace TypedElementQueryBuilder;
+	using DSI = ITypedElementDataStorageInterface;
+
+	DataStorage.RegisterQuery(Select(TEXT("Sync counter widgets"), 
+		FProcessor(DSI::EQueryTickPhase::FrameEnd, DataStorage.GetQueryTickGroupName(DSI::EQueryTickGroups::SyncWidgets))
+			.ForceToGameThread(true),
+		[](
+			DSI::IQueryContext& Context,
+			FTypedElementSlateWidgetReferenceColumn& Widget,
+			FTypedElementU32IntValueCacheColumn& Comparison, 
+			const FCounterWidgetColumn& Counter
+		)
+		{
+			DSI::FQueryResult Result = Context.RunQuery(Counter.Query);
+			if (Result.Completed == DSI::FQueryResult::ECompletion::Fully && Result.Count != Comparison.Value)
+			{
+				TSharedPtr<SWidget> WidgetPointer = Widget.Widget.Pin();
+				checkf(WidgetPointer, TEXT("Referenced widget is not valid. A constructed widget may not have been cleaned up. This can "
+					"also happen if this processor is running in the same phase as the processors responsible for cleaning up old "
+					"references."));
+				checkf(WidgetPointer->GetType() == STextBlock::StaticWidgetClass().GetWidgetType(),
+					TEXT("Stored widget with FTypedElementCounterWidgetFragment doesn't match type %s, but was a %s."),
+					*(STextBlock::StaticWidgetClass().GetWidgetType().ToString()),
+					*(WidgetPointer->GetTypeAsString()));
+
+				STextBlock* WidgetInstance = static_cast<STextBlock*>(WidgetPointer.Get());
+				WidgetInstance->SetText(FText::Format(Counter.LabelTextFormatter, Result.Count));
+				Comparison.Value = Result.Count;
+			}
+		}
+	).Compile());
+}
+
+void UCounterWidgetFactory::RegisterWidgetPurposes(ITypedElementDataStorageUiInterface& DataStorageUi) const
+{
+	DataStorageUi.RegisterWidgetPurpose(WigetPurpose, ITypedElementDataStorageUiInterface::EPurposeType::Generic,
+		LOCTEXT("ToolBarPurposeDescription", "Widgets added to the status bar at the bottom editor of the main editor window."));
+}
+
+void UCounterWidgetFactory::RegisterWidgetConstructors(ITypedElementDataStorageInterface& DataStorage,
+	ITypedElementDataStorageUiInterface& DataStorageUi) const
+{
+	using namespace TypedElementQueryBuilder;
+
+	TUniquePtr<FCounterWidgetConstructor> WidgetCounter = MakeUnique<FCounterWidgetConstructor>();
+	WidgetCounter->LabelText = LOCTEXT("WidgetCounterStatusBarLabel", "{0} {0}|plural(one=Widget, other=Widgets)");
+	WidgetCounter->ToolTipText = LOCTEXT(
+		"WidgetCounterStatusBarToolTip",
+		"The total number of widgets in the editor hosted through the Typed Element's Data Storage.");
+	WidgetCounter->Query = DataStorage.RegisterQuery(
+		Count().
+		Where().
+			All<FTypedElementSlateWidgetReferenceColumn>().
+		Compile());
+	DataStorageUi.RegisterWidgetFactory(WigetPurpose, MoveTemp(WidgetCounter));
+}
+
+void UCounterWidgetFactory::EnableCounterWidgets()
+{
+	bAreCounterWidgetsEnabled = true;
+	UCounterWidgetFactory::SetupMainWindowIntegrations(nullptr, false);
+}
+
+void UCounterWidgetFactory::SetupMainWindowIntegrations(TSharedPtr<SWindow> ParentWindow, bool bIsRunningStartupDialog)
+{
+	if (!bHasBeenSetup)
+	{
+		UTypedElementRegistry* Registry = UTypedElementRegistry::GetInstance();
+		checkf(Registry, TEXT(
+			"FTypedElementsDataStorageUiModule didn't find the UTypedElementRegistry during main window integration when it should be available."));
+
+		ITypedElementDataStorageUiInterface* UiInterface = Registry->GetMutableDataStorageUi();
+		checkf(UiInterface, TEXT(
+			"FTypedElementsDataStorageUiModule tried to integrate with the main window before the "
+			"Typed Elements Data Storage UI interface is available."));
+
+		UToolMenu* Menu = UToolMenus::Get()->ExtendMenu(WigetPurpose);
+
+		TArray<TSharedRef<SWidget>> Widgets;
+		UiInterface->ConstructWidgets(WigetPurpose, {},
+			[&Widgets](const TSharedRef<SWidget>& NewWidget, TypedElementRowHandle Row)
+			{
+				Widgets.Add(NewWidget);
+			});
+
+		if (!Widgets.IsEmpty())
+		{
+			FToolMenuSection& Section = Menu->AddSection("DataStorageSection");
+			int32 WidgetCount = Widgets.Num();
+
+			Section.AddEntry(FToolMenuEntry::InitWidget("DataStorageStatusBarWidget_0", MoveTemp(Widgets[0]), FText::GetEmpty()));
+			for (int32 I = 1; I < WidgetCount; ++I)
+			{
+				Section.AddSeparator(FName(*FString::Format(TEXT("DataStorageStatusBarWidgetDivider_{0}"), { FString::FromInt(I) })));
+				Section.AddEntry(FToolMenuEntry::InitWidget(
+					FName(*FString::Format(TEXT("DataStorageStatusBarWidget_{0}"), { FString::FromInt(I) })),
+					MoveTemp(Widgets[I]), FText::GetEmpty()));
+			}
+		}
+		bHasBeenSetup = true;
+	}
+}
+
+
+
+//
+// FCounterWidgetConstructor
+//
+
+FCounterWidgetConstructor::FCounterWidgetConstructor()
+	: Super(FCounterWidgetConstructor::StaticStruct())
+{
+}
+
+TConstArrayView<const UScriptStruct*> FCounterWidgetConstructor::GetAdditionalColumnsList() const
+{
+	static TTypedElementColumnTypeList<FCounterWidgetColumn, FTypedElementU32IntValueCacheColumn> Columns;
+	return Columns;
+}
+
+TSharedPtr<SWidget> FCounterWidgetConstructor::CreateWidget(const TypedElementDataStorage::FMetaDataView& Arguments)
+{
+	return SNew(STextBlock)
+		.Text(FText::Format(LabelText, 0))
+		.Margin(FMargin(4.0f, 0.0f))
+		.ToolTipText(ToolTipText)
+		.Justification(ETextJustify::Center);
+}
+
+bool FCounterWidgetConstructor::SetColumns(ITypedElementDataStorageInterface* DataStorage, TypedElementRowHandle Row)
+{
+	FCounterWidgetColumn* CounterColumn = DataStorage->GetColumn<FCounterWidgetColumn>(Row);
+	checkf(CounterColumn, TEXT("Added a new FCounterWidgetColumn to the Typed Elements Data Storage, but didn't get a valid pointer back."));
+	CounterColumn->LabelTextFormatter = LabelText;
+	CounterColumn->Query = Query;
+
+	FTypedElementU32IntValueCacheColumn* CacheColumn = DataStorage->GetColumn<FTypedElementU32IntValueCacheColumn>(Row);
+	checkf(CacheColumn, TEXT("Added a new FTypedElementUnsigned32BitIntValueCache to the Typed Elements Data Storage, but didn't get a valid pointer back."));
+	CacheColumn->Value = 0;
+
+	return true;
+}
+
+#undef LOCTEXT_NAMESPACE
\ No newline at end of file
diff --git a/Engine/Plugins/Experimental/EditorDataStorage/Source/TedsUI/Private/Widgets/ExportedTextWidget.cpp b/Engine/Plugins/Experimental/EditorDataStorage/Source/TedsUI/Private/Widgets/ExportedTextWidget.cpp
new file mode 100644
--- /dev/null
+++ b/Engine/Plugins/Experimental/EditorDataStorage/Source/TedsUI/Private/Widgets/ExportedTextWidget.cpp
@@ -0,0 +1,174 @@
+// Copyright Epic Games, Inc. All Rights Reserved.
+
+#include "Widgets/ExportedTextWidget.h"
+
+#include "Elements/Columns/TypedElementMiscColumns.h"
+#include "Elements/Columns/TypedElementSlateWidgetColumns.h"
+#include "Elements/Columns/TypedElementTypeInfoColumns.h"
+#include "Elements/Framework/TypedElementQueryBuilder.h"
+#include "Widgets/Text/STextBlock.h"
+
+#define LOCTEXT_NAMESPACE "TedsUI_ExportedTextWidget"
+
+//
+// UExportedTextWidgetFactory
+//
+
+static void UpdateExportedTextWidget(const void* Data, FTypedElementSlateWidgetReferenceColumn& Widget, 
+	const UScriptStruct* StructType)
+{
+	TSharedPtr<SWidget> WidgetPointer = Widget.Widget.Pin();
+	checkf(WidgetPointer, TEXT("Referenced widget is not valid. A constructed widget may not have been cleaned up. This can "
+		"also happen if this processor is running in the same phase as the processors responsible for cleaning up old "
+		"references."));
+	checkf(WidgetPointer->GetType() == STextBlock::StaticWidgetClass().GetWidgetType(),
+		TEXT("Stored widget with FExportedTextWidgetTag doesn't match type %s, but was a %s."),
+		*(STextBlock::StaticWidgetClass().GetWidgetType().ToString()),
+		*(WidgetPointer->GetTypeAsString()));
+
+	FString Label;
+	StructType->ExportText(Label, Data, Data, nullptr, PPF_None, nullptr);
+	STextBlock* TextWidget = static_cast<STextBlock*>(WidgetPointer.Get());
+	FText Text = FText::FromString(MoveTemp(Label));
+	TextWidget->SetToolTipText(Text);
+	TextWidget->SetText(MoveTemp(Text));
+}
+
+static void UpdateExportedTextWidget(ITypedElementDataStorageInterface& DataStorage, FTypedElementSlateWidgetReferenceColumn& Widget,
+	const FTypedElementScriptStructTypeInfoColumn& TypeInfo, const FTypedElementRowReferenceColumn& ReferencedRow)
+{
+	const UScriptStruct* StructType = TypeInfo.TypeInfo.Get();
+	if (void* Data = DataStorage.GetColumnData(ReferencedRow.Row, StructType))
+	{
+		UpdateExportedTextWidget(Data, Widget, StructType);
+	}
+}
+
+static TypedElementQueryHandle RegisterUpdateCallback(ITypedElementDataStorageInterface& DataStorage, const UScriptStruct* Target)
+{
+	using namespace TypedElementQueryBuilder;
+	using DSI = ITypedElementDataStorageInterface;
+	namespace DS = TypedElementDataStorage;
+	
+	TypedElementQueryHandle TypeDataQuery = DataStorage.RegisterQuery(
+		Select()
+			.ReadOnly(Target)
+		.Where()
+			.Any<FTypedElementSyncFromWorldTag, FTypedElementSyncBackToWorldTag>()
+		.Compile());
+
+	FString Name = TEXT("Sync exported text widgets (");
+	Target->AppendName(Name);
+	Name += ')';
+
+	FName ProcessorName(Name);
+
+	return DataStorage.RegisterQuery(
+		Select(ProcessorName,
+			FProcessor(DSI::EQueryTickPhase::FrameEnd, DataStorage.GetQueryTickGroupName(DSI::EQueryTickGroups::SyncWidgets))
+				.ForceToGameThread(true),
+			[Target, ProcessorName](
+				DS::IQueryContext& Context, 
+				FTypedElementSlateWidgetReferenceColumn& Widget,
+				const FTypedElementScriptStructTypeInfoColumn& TypeInfo,
+				const FTypedElementRowReferenceColumn& ReferencedRow)
+			{
+				const UScriptStruct* StructType = TypeInfo.TypeInfo.Get();
+				if (!ensureMsgf(StructType != nullptr, TEXT("WeakPtr TypeInfo is null in query '%s'"), *ProcessorName.ToString()))
+				{
+					return;
+				}
+
+				/* This query will grab all ExportedText widgets that were created for rows with the "Target" column, but we want to make sure we are
+				 * only updating the widgets that are actually displaying the "Target" column.
+				 *
+				 * E.g a row could have ColumnA and ColumnB which are both using the exported text widget to display - but there is no way for the
+				 * two widgets to be differentiated from a TEDS query. So if the widget for ColumnB wants to update: Both the widgets for ColumnA
+				 * and ColumnB would match the query condition (WidgetRow has FExportedTextWidgetTag && TargetRow has ColumnB), but this
+				 * query (+ subquery) only have access to ColumnB and only want to update the widget that's displaying ColumnB.
+				 *
+				 * To work around this we check to make sure the widget this query is trying to update is for the column this query is targeting.
+				 */
+				if(StructType == Target)
+				{
+					Context.RunSubquery(0, ReferencedRow.Row, 
+					[&Widget, StructType](const DS::FQueryDescription&, DS::ISubqueryContext& SubqueryContext)
+					{
+						const void* ColumnData = SubqueryContext.GetColumn(StructType);
+						UpdateExportedTextWidget(ColumnData, Widget, StructType);
+					});
+				}
+			})
+		.Where()
+			.All<FExportedTextWidgetTag>()
+		.DependsOn()
+			.SubQuery(TypeDataQuery)
+		.Compile());
+}
+
+void UExportedTextWidgetFactory::RegisterWidgetConstructors(ITypedElementDataStorageInterface& DataStorage,
+	ITypedElementDataStorageUiInterface& DataStorageUi) const
+{
+	DataStorageUi.RegisterWidgetFactory(FName(TEXT("General.Cell.Default")), FExportedTextWidgetConstructor::StaticStruct());
+}
+
+//
+// FExportedTextWidgetConstructor
+//
+
+FExportedTextWidgetConstructor::FExportedTextWidgetConstructor()
+	: Super(FExportedTextWidgetConstructor::StaticStruct())
+{
+}
+
+TConstArrayView<const UScriptStruct*> FExportedTextWidgetConstructor::GetAdditionalColumnsList() const
+{
+	static TTypedElementColumnTypeList<
+		FTypedElementRowReferenceColumn,
+		FTypedElementScriptStructTypeInfoColumn,
+		FExportedTextWidgetTag> Columns;
+	return Columns;
+}
+
+const TypedElementDataStorage::FQueryConditions* FExportedTextWidgetConstructor::GetQueryConditions() const
+{
+	// For the exported text widget, the query condition we are matched against is the column we are exporting text for
+	return &MatchedColumn;
+}
+
+TSharedPtr<SWidget> FExportedTextWidgetConstructor::CreateWidget(const TypedElementDataStorage::FMetaDataView& Arguments)
+{
+	return SNew(STextBlock);
+}
+
+bool FExportedTextWidgetConstructor::FinalizeWidget(
+	ITypedElementDataStorageInterface* DataStorage,
+	ITypedElementDataStorageUiInterface* DataStorageUi,
+	TypedElementRowHandle Row,
+	const TSharedPtr<SWidget>& Widget)
+{
+	FTypedElementScriptStructTypeInfoColumn& TypeInfoColumn = *DataStorage->GetColumn<FTypedElementScriptStructTypeInfoColumn>(Row);
+
+	// NOTE: We are currently assuming that an instance of FExportedTextWidgetConstructor will only be used to show the same type info
+	// which isn't ideal but it's better than nothing since we need some sort of matched conditions for column based virtualization to work.
+	// TEDS UI TODO: We should work around it by refactoring this into an STedsWidget in the future so it can store the column conditions per instance
+	MatchedColumn = TypedElementDataStorage::FQueryConditions(TypedElementDataStorage::FColumn(TypeInfoColumn.TypeInfo));
+
+	UpdateExportedTextWidget(
+		*DataStorage,
+		*DataStorage->GetColumn<FTypedElementSlateWidgetReferenceColumn>(Row),
+		TypeInfoColumn,
+		*DataStorage->GetColumn<FTypedElementRowReferenceColumn>(Row));
+	
+	UExportedTextWidgetFactory* Factory = 
+		UExportedTextWidgetFactory::StaticClass()->GetDefaultObject<UExportedTextWidgetFactory>();
+	if (Factory && !Factory->RegisteredTypes.Contains(TypeInfoColumn.TypeInfo))
+	{
+		RegisterUpdateCallback(*DataStorage, TypeInfoColumn.TypeInfo.Get());
+		Factory->RegisteredTypes.Add(TypeInfoColumn.TypeInfo);
+	}
+	
+	return true;
+}
+
+#undef LOCTEXT_NAMESPACE
diff --git a/Engine/Plugins/Experimental/EditorDataStorage/Source/TedsUI/Private/Widgets/GeneralWidgetRegistrationFactory.cpp b/Engine/Plugins/Experimental/EditorDataStorage/Source/TedsUI/Private/Widgets/GeneralWidgetRegistrationFactory.cpp
new file mode 100644
--- /dev/null
+++ b/Engine/Plugins/Experimental/EditorDataStorage/Source/TedsUI/Private/Widgets/GeneralWidgetRegistrationFactory.cpp
@@ -0,0 +1,29 @@
+// Copyright Epic Games, Inc. All Rights Reserved.
+
+#include "Widgets/GeneralWidgetRegistrationFactory.h"
+
+#include "Elements/Interfaces/TypedElementDataStorageUiInterface.h"
+
+#define LOCTEXT_NAMESPACE "TypedElementsUI_GeneralRegistration"
+
+const FName UGeneralWidgetRegistrationFactory::CellPurpose(TEXT("General.Cell"));
+const FName UGeneralWidgetRegistrationFactory::HeaderPurpose(TEXT("General.Header"));
+const FName UGeneralWidgetRegistrationFactory::CellDefaultPurpose(TEXT("General.Cell.Default"));
+const FName UGeneralWidgetRegistrationFactory::HeaderDefaultPurpose(TEXT("General.Header.Default"));
+
+void UGeneralWidgetRegistrationFactory::RegisterWidgetPurposes(ITypedElementDataStorageUiInterface& DataStorageUi) const
+{
+	DataStorageUi.RegisterWidgetPurpose(CellPurpose, ITypedElementDataStorageUiInterface::EPurposeType::UniqueByNameAndColumn,
+		LOCTEXT("GeneralCellPurpose", "General purpose widgets that can be used as cells for specific columns or column combinations."));
+
+	DataStorageUi.RegisterWidgetPurpose(HeaderPurpose, ITypedElementDataStorageUiInterface::EPurposeType::UniqueByNameAndColumn,
+		LOCTEXT("GeneralHeaderPurpose", "General purpose widget that can be used as a header."));
+
+	DataStorageUi.RegisterWidgetPurpose(HeaderDefaultPurpose, ITypedElementDataStorageUiInterface::EPurposeType::UniqueByName,
+		LOCTEXT("GeneralHeaderDefaultPurpose", "The default widget to use in headers if no other specialization is provided."));
+	
+	DataStorageUi.RegisterWidgetPurpose(CellDefaultPurpose, ITypedElementDataStorageUiInterface::EPurposeType::UniqueByName,
+		LOCTEXT("GeneralCellDefaultPurpose", "The default widget to use in cells if no other specialization is provided."));
+}
+
+#undef LOCTEXT_NAMESPACE
diff --git a/Engine/Plugins/Experimental/EditorDataStorage/Source/TedsUI/Private/Widgets/LabelWidget.cpp b/Engine/Plugins/Experimental/EditorDataStorage/Source/TedsUI/Private/Widgets/LabelWidget.cpp
new file mode 100644
--- /dev/null
+++ b/Engine/Plugins/Experimental/EditorDataStorage/Source/TedsUI/Private/Widgets/LabelWidget.cpp
@@ -0,0 +1,247 @@
+// Copyright Epic Games, Inc. All Rights Reserved.
+
+#include "Widgets/LabelWidget.h"
+
+#include "ActorEditorUtils.h"
+#include "Elements/Columns/TypedElementLabelColumns.h"
+#include "Elements/Columns/TypedElementMiscColumns.h"
+#include "Elements/Columns/TypedElementSlateWidgetColumns.h"
+#include "Elements/Columns/TypedElementValueCacheColumns.h"
+#include "Elements/Framework/TypedElementQueryBuilder.h"
+#include "Elements/Interfaces/TypedElementDataStorageInterface.h"
+#include "Elements/Interfaces/Capabilities/TypedElementUiEditableCapability.h"
+#include "Elements/Interfaces/Capabilities/TypedElementUiTextCapability.h"
+#include "Elements/Interfaces/Capabilities/TypedElementUiTooltipCapability.h"
+#include "Elements/Interfaces/Capabilities/TypedElementUiStyleOverrideCapability.h"
+#include "Widgets/Text/STextBlock.h"
+#include "Widgets/Text/SInlineEditableTextBlock.h"
+
+#define LOCTEXT_NAMESPACE "TedsUI_LabelWidget"
+
+//
+// ULabelWidgetFactory
+//
+
+static void UpdateTextWidget(const TWeakPtr<SWidget>& Widget, const FTypedElementLabelColumn& Label, const uint64* HashValue)
+{
+	TSharedPtr<SWidget> WidgetPointer = Widget.Pin();
+	if (ensureMsgf(WidgetPointer, TEXT("Referenced widget is not valid. A constructed widget may not have been cleaned up. This can "
+		"also happen if this processor is running in the same phase as the processors responsible for cleaning up old "
+		"references.")))
+	{
+		if (TSharedPtr<ITypedElementUiTextCapability> Text = WidgetPointer->GetMetaData<ITypedElementUiTextCapability>())
+		{
+			Text->SetText(FText::FromString(Label.Label));
+		}
+
+		if (TSharedPtr<ITypedElementUiTooltipCapability> ToolTip = WidgetPointer->GetMetaData<ITypedElementUiTooltipCapability>())
+		{
+			if (!HashValue)
+			{
+				ToolTip->SetToolTipText(FText::FromString(Label.Label));
+			}
+			else
+			{
+				ToolTip->SetToolTipText(FText::FromString(FString::Format(TEXT("{0}\nHash: {1}"), { Label.Label, *HashValue })));
+			}
+		}
+	}
+}
+
+static void SyncColumnsToWidget(
+	ITypedElementDataStorageInterface* DataStorage, 
+	TypedElementRowHandle TargetRow, 
+	FTypedElementU64IntValueCacheColumn& TextHash,
+	const TWeakPtr<SWidget>& Widget,
+	bool bShowHashInToolTip)
+{
+	if (const FTypedElementLabelColumn* LabelColumn = DataStorage->GetColumn<FTypedElementLabelColumn>(TargetRow))
+	{
+		if (const FTypedElementLabelHashColumn* LabelHashColumn = DataStorage->GetColumn<FTypedElementLabelHashColumn>(TargetRow))
+		{
+			if (LabelHashColumn->LabelHash != TextHash.Value)
+			{
+				UpdateTextWidget(Widget, *LabelColumn, bShowHashInToolTip ? &LabelHashColumn->LabelHash : nullptr);
+				TextHash.Value = LabelHashColumn->LabelHash;
+			}
+		}
+		else
+		{
+			UpdateTextWidget(Widget, *LabelColumn, nullptr);
+		}
+	}
+}
+
+void ULabelWidgetFactory::RegisterQueries(ITypedElementDataStorageInterface& DataStorage)
+{
+	using namespace TypedElementQueryBuilder;
+	using DSI = ITypedElementDataStorageInterface;
+
+	TypedElementQueryHandle UpdateLabelWidget = DataStorage.RegisterQuery(
+		Select()
+			.ReadOnly<FTypedElementLabelColumn>()
+		.Where()
+			.Any<FTypedElementSyncFromWorldTag, FTypedElementSyncBackToWorldTag>()
+			.None<FTypedElementLabelHashColumn>()
+		.Compile());
+
+	TypedElementQueryHandle UpdateLabelAndHashWidget = DataStorage.RegisterQuery(
+		Select()
+			.ReadOnly<FTypedElementLabelColumn, FTypedElementLabelHashColumn>()
+		.Where()
+			.Any<FTypedElementSyncFromWorldTag, FTypedElementSyncBackToWorldTag>()
+		.Compile());
+
+	DataStorage.RegisterQuery(
+		Select(
+			TEXT("Sync label to widget"),
+			FProcessor(DSI::EQueryTickPhase::FrameEnd, DataStorage.GetQueryTickGroupName(DSI::EQueryTickGroups::SyncWidgets))
+				.ForceToGameThread(true),
+			[](
+				DSI::IQueryContext& Context, 
+				FTypedElementSlateWidgetReferenceColumn& Widget,
+				FTypedElementU64IntValueCacheColumn& TextHash,
+				const FLabelWidgetColumn& Config,
+				const FTypedElementRowReferenceColumn& Target)
+			{
+				Context.RunSubquery(0, Target.Row, CreateSubqueryCallbackBinding(
+					[&Widget](const FTypedElementLabelColumn& Label)
+					{
+						UpdateTextWidget(Widget.Widget, Label, nullptr);
+					}));
+				Context.RunSubquery(1, Target.Row, CreateSubqueryCallbackBinding(
+					[&Widget, &TextHash, &Config](const FTypedElementLabelColumn& Label, const FTypedElementLabelHashColumn& Hash)
+					{
+						if (Hash.LabelHash != TextHash.Value)
+						{
+							UpdateTextWidget(Widget.Widget, Label, Config.bShowHashInTooltip ? &Hash.LabelHash : nullptr);
+							TextHash.Value = Hash.LabelHash;
+						}
+					}));
+			})
+		.DependsOn()
+			.SubQuery({ UpdateLabelWidget, UpdateLabelAndHashWidget })
+		.Compile()
+	);
+}
+
+void ULabelWidgetFactory::RegisterWidgetConstructors(ITypedElementDataStorageInterface& DataStorage,
+	ITypedElementDataStorageUiInterface& DataStorageUi) const
+{
+	using namespace TypedElementDataStorage;
+
+	DataStorageUi.RegisterWidgetFactory<FLabelWidgetConstructor>(FName(TEXT("General.Cell")), 
+		FColumn<FTypedElementLabelColumn>() || (FColumn<FTypedElementLabelColumn>() && FColumn<FTypedElementLabelHashColumn>()));
+}
+
+
+
+//
+// FLabelWidgetConstructor
+//
+
+FLabelWidgetConstructor::FLabelWidgetConstructor()
+	: Super(FLabelWidgetConstructor::StaticStruct())
+{
+}
+
+FLabelWidgetConstructor::FLabelWidgetConstructor(const UScriptStruct* InTypeInfo)
+	: Super(InTypeInfo)
+{
+}
+
+TConstArrayView<const UScriptStruct*> FLabelWidgetConstructor::GetAdditionalColumnsList() const
+{
+	static const TTypedElementColumnTypeList<
+		FTypedElementRowReferenceColumn,
+		FTypedElementU64IntValueCacheColumn,
+		FLabelWidgetColumn> Columns;
+	return Columns;
+}
+
+TSharedPtr<SWidget> FLabelWidgetConstructor::Construct(
+	TypedElementRowHandle Row,
+	ITypedElementDataStorageInterface* DataStorage,
+	ITypedElementDataStorageUiInterface* DataStorageUi,
+	const TypedElementDataStorage::FMetaDataView& Arguments)
+{
+	TSharedPtr<SWidget> Result;
+	const bool* IsEditable = Arguments.FindForColumn<FTypedElementLabelColumn>(TypedElementDataStorage::IsEditableName).TryGetExact<bool>();
+	if (IsEditable && *IsEditable)
+	{
+		if (FTypedElementRowReferenceColumn* TargetRowColumn = DataStorage->GetColumn<FTypedElementRowReferenceColumn>(Row))
+		{
+			TSharedPtr<SInlineEditableTextBlock> TextBlock = SNew(SInlineEditableTextBlock)
+				.OnTextCommitted_Lambda(
+					[DataStorage, TargetRow = TargetRowColumn->Row](const FText& NewText, ETextCommit::Type CommitInfo)
+					{
+						// This callback happens on the game thread so it's safe to directly call into the data storage.
+						FString NewLabelText = NewText.ToString();
+						if (FTypedElementLabelHashColumn* LabelHashColumn = DataStorage->GetColumn<FTypedElementLabelHashColumn>(TargetRow))
+						{
+							LabelHashColumn->LabelHash = CityHash64(reinterpret_cast<const char*>(*NewLabelText), NewLabelText.Len() * sizeof(**NewLabelText));
+						}
+						if (FTypedElementLabelColumn* LabelColumn = DataStorage->GetColumn<FTypedElementLabelColumn>(TargetRow))
+						{
+							LabelColumn->Label = MoveTemp(NewLabelText);
+						}
+						DataStorage->AddColumn<FTypedElementSyncBackToWorldTag>(TargetRow);
+					})
+				.OnVerifyTextChanged_Lambda([](const FText& Label, FText& ErrorMessage)
+					{
+						// Note: The use of actor specific functionality should be minimized, but this function acts generic enough that the 
+						// use of actor is just in names.
+						return FActorEditorUtils::ValidateActorName(Label, ErrorMessage);
+					});
+			TextBlock->AddMetadata(MakeShared<TTypedElementUiEditableCapability<SInlineEditableTextBlock>>(*TextBlock));
+			TextBlock->AddMetadata(MakeShared<TTypedElementUiTextCapability<SInlineEditableTextBlock>>(*TextBlock));
+			TextBlock->AddMetadata(MakeShared<TTypedElementUiTooltipCapability<SInlineEditableTextBlock>>(*TextBlock));
+			TextBlock->AddMetadata(MakeShared<TTypedElementUiStyleOverrideCapability<SInlineEditableTextBlock>>(*TextBlock));
+			Result = TextBlock;
+		}
+	}
+	else
+	{
+		TSharedPtr<STextBlock> TextBlock = SNew(STextBlock)
+			.IsEnabled(false);
+		TextBlock->AddMetadata(MakeShared<TTypedElementUiTextCapability<STextBlock>>(*TextBlock));
+		TextBlock->AddMetadata(MakeShared<TTypedElementUiTooltipCapability<STextBlock>>(*TextBlock));
+		Result = TextBlock;
+	}
+	
+	if (Result)
+	{
+		DataStorage->GetColumn<FTypedElementSlateWidgetReferenceColumn>(Row)->Widget = Result;
+		if (SetColumns(DataStorage, Row))
+		{
+			if (FinalizeWidget(DataStorage, DataStorageUi, Row, Result))
+			{
+				SetupDebugColumns(Row, DataStorage, DataStorageUi, Result);
+				return Result;
+			}
+		}
+	}
+
+	return nullptr;
+}
+
+bool FLabelWidgetConstructor::SetColumns(ITypedElementDataStorageInterface* DataStorage, TypedElementRowHandle Row)
+{
+	DataStorage->AddColumn(Row, FLabelWidgetColumn{ .bShowHashInTooltip = (MatchedColumnTypes.Num() == 2) });
+	return true;
+}
+
+bool FLabelWidgetConstructor::FinalizeWidget(
+	ITypedElementDataStorageInterface* DataStorage,
+	ITypedElementDataStorageUiInterface* DataStorageUi,
+	TypedElementRowHandle Row,
+	const TSharedPtr<SWidget>& Widget)
+{
+	SyncColumnsToWidget(
+		DataStorage, 
+		DataStorage->GetColumn<FTypedElementRowReferenceColumn>(Row)->Row,
+		*DataStorage->GetColumn<FTypedElementU64IntValueCacheColumn>(Row),
+		Widget, MatchedColumnTypes.Num() == 2);
+	return true;
+}
+#undef LOCTEXT_NAMESPACE
diff --git a/Engine/Plugins/Experimental/EditorDataStorage/Source/TedsUI/Private/Widgets/PackagePathWidget.cpp b/Engine/Plugins/Experimental/EditorDataStorage/Source/TedsUI/Private/Widgets/PackagePathWidget.cpp
new file mode 100644
--- /dev/null
+++ b/Engine/Plugins/Experimental/EditorDataStorage/Source/TedsUI/Private/Widgets/PackagePathWidget.cpp
@@ -0,0 +1,99 @@
+// Copyright Epic Games, Inc. All Rights Reserved.
+
+#include "Widgets/PackagePathWidget.h"
+
+#include "Elements/Columns/TypedElementPackageColumns.h"
+#include "Elements/Columns/TypedElementMiscColumns.h"
+#include "Elements/Interfaces/TypedElementDataStorageInterface.h"
+#include "Widgets/Text/STextBlock.h"
+
+//
+// UPackagePathWidgetFactory
+//
+
+void UPackagePathWidgetFactory::RegisterWidgetConstructors(ITypedElementDataStorageInterface& DataStorage,
+	ITypedElementDataStorageUiInterface& DataStorageUi) const
+{
+	using namespace TypedElementDataStorage;
+
+	DataStorageUi.RegisterWidgetFactory<FPackagePathWidgetConstructor>(FName(TEXT("General.Cell")),
+		FColumn<FTypedElementPackagePathColumn>());
+	DataStorageUi.RegisterWidgetFactory<FLoadedPackagePathWidgetConstructor>(FName(TEXT("General.Cell")),
+		FColumn<FTypedElementPackageLoadedPathColumn>());
+}
+
+
+
+//
+// FPackagePathWidgetConstructor
+//
+
+FPackagePathWidgetConstructor::FPackagePathWidgetConstructor()
+	: Super(FPackagePathWidgetConstructor::StaticStruct())
+{
+}
+
+FPackagePathWidgetConstructor::FPackagePathWidgetConstructor(const UScriptStruct* InTypeInfo)
+	: Super(FPackagePathWidgetConstructor::StaticStruct())
+{
+}
+
+TSharedPtr<SWidget> FPackagePathWidgetConstructor::CreateWidget(const TypedElementDataStorage::FMetaDataView& Arguments)
+{
+	return SNew(STextBlock)
+		.OverflowPolicy(ETextOverflowPolicy::Ellipsis)
+		.Justification(ETextJustify::Right);
+}
+
+bool FPackagePathWidgetConstructor::FinalizeWidget(
+	ITypedElementDataStorageInterface* DataStorage,
+	ITypedElementDataStorageUiInterface* DataStorageUi,
+	TypedElementRowHandle Row,
+	const TSharedPtr<SWidget>& Widget)
+{
+	TypedElementRowHandle TargetRow = DataStorage->GetColumn<FTypedElementRowReferenceColumn>(Row)->Row;
+	if (const FTypedElementPackagePathColumn* Path = DataStorage->GetColumn<FTypedElementPackagePathColumn>(TargetRow))
+	{
+		STextBlock* TextWidget = static_cast<STextBlock*>(Widget.Get());
+		FText Text = FText::FromString(Path->Path);
+		TextWidget->SetToolTipText(Text);
+		TextWidget->SetText(MoveTemp(Text));
+		return true;
+	}
+	else
+	{
+		return false;
+	}
+}
+
+
+
+//
+// FLoadedPackagePathWidgetConstructor
+//
+
+FLoadedPackagePathWidgetConstructor::FLoadedPackagePathWidgetConstructor()
+	: Super(FLoadedPackagePathWidgetConstructor::StaticStruct())
+{
+}
+
+bool FLoadedPackagePathWidgetConstructor::FinalizeWidget(
+	ITypedElementDataStorageInterface* DataStorage,
+	ITypedElementDataStorageUiInterface* DataStorageUi,
+	TypedElementRowHandle Row,
+	const TSharedPtr<SWidget>& Widget)
+{
+	TypedElementRowHandle TargetRow = DataStorage->GetColumn<FTypedElementRowReferenceColumn>(Row)->Row;
+	if (const FTypedElementPackageLoadedPathColumn* Path = DataStorage->GetColumn<FTypedElementPackageLoadedPathColumn>(TargetRow))
+	{
+		STextBlock* TextWidget = static_cast<STextBlock*>(Widget.Get());
+		FText Text = FText::FromString(Path->LoadedPath.GetLocalFullPath());
+		TextWidget->SetToolTipText(Text);
+		TextWidget->SetText(MoveTemp(Text));
+		return true;
+	}
+	else
+	{
+		return false;
+	}
+}
diff --git a/Engine/Plugins/Experimental/EditorDataStorage/Source/TedsUI/Public/TypedElementCounterWidget.h b/Engine/Plugins/Experimental/EditorDataStorage/Source/TedsUI/Public/TypedElementCounterWidget.h
deleted file mode 100644
diff --git a/Engine/Plugins/Experimental/EditorDataStorage/Source/TedsUI/Public/TypedElementExportedTextWidget.h b/Engine/Plugins/Experimental/EditorDataStorage/Source/TedsUI/Public/TypedElementExportedTextWidget.h
deleted file mode 100644
diff --git a/Engine/Plugins/Experimental/EditorDataStorage/Source/TedsUI/Public/TypedElementGeneralRegistration.h b/Engine/Plugins/Experimental/EditorDataStorage/Source/TedsUI/Public/TypedElementGeneralRegistration.h
deleted file mode 100644
diff --git a/Engine/Plugins/Experimental/EditorDataStorage/Source/TedsUI/Public/TypedElementLabelWidget.h b/Engine/Plugins/Experimental/EditorDataStorage/Source/TedsUI/Public/TypedElementLabelWidget.h
deleted file mode 100644
diff --git a/Engine/Plugins/Experimental/EditorDataStorage/Source/TedsUI/Public/TypedElementPackagePathWidget.h b/Engine/Plugins/Experimental/EditorDataStorage/Source/TedsUI/Public/TypedElementPackagePathWidget.h
deleted file mode 100644
diff --git a/Engine/Plugins/Experimental/EditorDataStorage/Source/TedsUI/Public/TypedElementsDataStorageUI.h b/Engine/Plugins/Experimental/EditorDataStorage/Source/TedsUI/Public/TypedElementsDataStorageUI.h
deleted file mode 100644
diff --git a/Engine/Plugins/Experimental/EditorDataStorage/Source/TedsUI/Public/Widgets/CounterWidget.h b/Engine/Plugins/Experimental/EditorDataStorage/Source/TedsUI/Public/Widgets/CounterWidget.h
new file mode 100644
--- /dev/null
+++ b/Engine/Plugins/Experimental/EditorDataStorage/Source/TedsUI/Public/Widgets/CounterWidget.h
@@ -0,0 +1,74 @@
+// Copyright Epic Games, Inc. All Rights Reserved.
+
+#pragma once
+
+#include "Elements/Interfaces/TypedElementDataStorageFactory.h"
+#include "Elements/Interfaces/TypedElementDataStorageUiInterface.h"
+#include "Framework/Text/TextLayout.h"
+#include "Internationalization/Text.h"
+
+#include "CounterWidget.generated.h"
+
+class ITypedElementDataStorageInterface;
+class SWindow;
+class UScriptStruct;
+
+UCLASS()
+class TEDSUI_API UCounterWidgetFactory : public UTypedElementDataStorageFactory
+{
+	GENERATED_BODY()
+
+public:
+	static FName WigetPurpose;
+
+	UCounterWidgetFactory();
+	~UCounterWidgetFactory() override = default;
+
+	void RegisterQueries(ITypedElementDataStorageInterface& DataStorage) override;
+	void RegisterWidgetPurposes(ITypedElementDataStorageUiInterface& DataStorageUi) const override;
+	void RegisterWidgetConstructors(ITypedElementDataStorageInterface& DataStorage,
+		ITypedElementDataStorageUiInterface& DataStorageUi) const override;
+
+	static void EnableCounterWidgets();
+
+private:
+	static void SetupMainWindowIntegrations(TSharedPtr<SWindow> ParentWindow, bool bIsRunningStartupDialog);
+
+	static bool bAreCounterWidgetsEnabled;
+	static bool bHasBeenSetup;
+};
+
+/**
+ * Constructor for the counter widget. The counter widget accepts a "count"-query. The query will be periodically
+ * run and the result is written to a textbox widget after it's been formatted using LabelText. An example for 
+ * LabelText is "{0} {0}|plural(one=MyCounter, other=MyCounters)" which will use "MyCounter" if there's exactly one
+ * entry found and otherwise "MyCounters".
+ */
+USTRUCT()
+struct TEDSUI_API FCounterWidgetConstructor : public FTypedElementWidgetConstructor
+{
+	GENERATED_BODY()
+
+public:
+	FCounterWidgetConstructor();
+	~FCounterWidgetConstructor() override = default;
+
+	TConstArrayView<const UScriptStruct*> GetAdditionalColumnsList() const override;
+
+	FText ToolTipText{ NSLOCTEXT("TypedElementUI_CounterWidget", "Tooltip", "Shows the total number found in the editor.") };
+	FText LabelText{ NSLOCTEXT("TypedElementUI_CounterWidget", "Label", "Counted") };
+	TypedElementQueryHandle Query;
+
+protected:
+	TSharedPtr<SWidget> CreateWidget(const TypedElementDataStorage::FMetaDataView& Arguments) override;
+	bool SetColumns(ITypedElementDataStorageInterface* DataStorage, TypedElementRowHandle Row) override;
+};
+
+USTRUCT(meta = (DisplayName = "Counter widget"))
+struct TEDSUI_API FCounterWidgetColumn : public FTypedElementDataStorageColumn
+{
+	GENERATED_BODY()
+
+	FTextFormat LabelTextFormatter;
+	TypedElementQueryHandle Query;
+};
\ No newline at end of file
diff --git a/Engine/Plugins/Experimental/EditorDataStorage/Source/TedsUI/Public/Widgets/ExportedTextWidget.h b/Engine/Plugins/Experimental/EditorDataStorage/Source/TedsUI/Public/Widgets/ExportedTextWidget.h
new file mode 100644
--- /dev/null
+++ b/Engine/Plugins/Experimental/EditorDataStorage/Source/TedsUI/Public/Widgets/ExportedTextWidget.h
@@ -0,0 +1,57 @@
+// Copyright Epic Games, Inc. All Rights Reserved.
+
+#pragma once
+
+#include "Elements/Interfaces/TypedElementDataStorageFactory.h"
+#include "Elements/Interfaces/TypedElementDataStorageUiInterface.h"
+#include "Elements/Common/TypedElementQueryConditions.h"
+
+#include "ExportedTextWidget.generated.h"
+
+class ITypedElementDataStorageInterface;
+class UScriptStruct;
+
+UCLASS()
+class TEDSUI_API UExportedTextWidgetFactory : public UTypedElementDataStorageFactory
+{
+	GENERATED_BODY()
+
+public:
+	virtual ~UExportedTextWidgetFactory() override = default;
+
+	virtual void RegisterWidgetConstructors(ITypedElementDataStorageInterface& DataStorage,
+		ITypedElementDataStorageUiInterface& DataStorageUi) const override;
+
+	TSet<TWeakObjectPtr<const UScriptStruct>> RegisteredTypes;
+};
+
+USTRUCT()
+struct TEDSUI_API FExportedTextWidgetConstructor : public FTypedElementWidgetConstructor
+{
+	GENERATED_BODY()
+
+public:
+	FExportedTextWidgetConstructor();
+	virtual ~FExportedTextWidgetConstructor() override = default;
+
+	virtual TConstArrayView<const UScriptStruct*> GetAdditionalColumnsList() const override;
+	virtual const TypedElementDataStorage::FQueryConditions* GetQueryConditions() const override;
+
+protected:
+	virtual TSharedPtr<SWidget> CreateWidget(const TypedElementDataStorage::FMetaDataView& Arguments) override;
+	virtual bool FinalizeWidget(
+		ITypedElementDataStorageInterface* DataStorage,
+		ITypedElementDataStorageUiInterface* DataStorageUi,
+		TypedElementRowHandle Row,
+		const TSharedPtr<SWidget>& Widget) override;
+
+protected:
+	// The column this exported text widget is operating on
+	TypedElementDataStorage::FQueryConditions MatchedColumn;
+};
+
+USTRUCT(meta = (DisplayName = "Exported text widget"))
+struct TEDSUI_API FExportedTextWidgetTag : public FTypedElementDataStorageTag
+{
+	GENERATED_BODY()
+};
\ No newline at end of file
diff --git a/Engine/Plugins/Experimental/EditorDataStorage/Source/TedsUI/Public/Widgets/GeneralWidgetRegistrationFactory.h b/Engine/Plugins/Experimental/EditorDataStorage/Source/TedsUI/Public/Widgets/GeneralWidgetRegistrationFactory.h
new file mode 100644
--- /dev/null
+++ b/Engine/Plugins/Experimental/EditorDataStorage/Source/TedsUI/Public/Widgets/GeneralWidgetRegistrationFactory.h
@@ -0,0 +1,24 @@
+// Copyright Epic Games, Inc. All Rights Reserved.
+
+#pragma once
+
+#include "Elements/Interfaces/TypedElementDataStorageFactory.h"
+#include "UObject/ObjectMacros.h"
+
+#include "GeneralWidgetRegistrationFactory.generated.h"
+
+UCLASS()
+class TEDSUI_API UGeneralWidgetRegistrationFactory : public UTypedElementDataStorageFactory
+{
+	GENERATED_BODY()
+
+public:
+	static const FName CellPurpose;
+	static const FName HeaderPurpose;
+	static const FName CellDefaultPurpose;
+	static const FName HeaderDefaultPurpose;
+
+	~UGeneralWidgetRegistrationFactory() override = default;
+
+	void RegisterWidgetPurposes(ITypedElementDataStorageUiInterface& DataStorageUi) const override;
+};
diff --git a/Engine/Plugins/Experimental/EditorDataStorage/Source/TedsUI/Public/Widgets/LabelWidget.h b/Engine/Plugins/Experimental/EditorDataStorage/Source/TedsUI/Public/Widgets/LabelWidget.h
new file mode 100644
--- /dev/null
+++ b/Engine/Plugins/Experimental/EditorDataStorage/Source/TedsUI/Public/Widgets/LabelWidget.h
@@ -0,0 +1,58 @@
+// Copyright Epic Games, Inc. All Rights Reserved.
+
+#pragma once
+
+#include "Elements/Interfaces/TypedElementDataStorageFactory.h"
+#include "Elements/Interfaces/TypedElementDataStorageUiInterface.h"
+#include "Internationalization/Text.h"
+#include "UObject/ObjectMacros.h"
+
+#include "LabelWidget.generated.h"
+
+class ITypedElementDataStorageInterface;
+class UScriptStruct;
+
+UCLASS()
+class ULabelWidgetFactory : public UTypedElementDataStorageFactory
+{
+	GENERATED_BODY()
+
+public:
+	~ULabelWidgetFactory() override = default;
+
+	TEDSUI_API void RegisterQueries(ITypedElementDataStorageInterface& DataStorage) override;
+	TEDSUI_API void RegisterWidgetConstructors(ITypedElementDataStorageInterface& DataStorage,
+		ITypedElementDataStorageUiInterface& DataStorageUi) const override;
+};
+
+USTRUCT()
+struct FLabelWidgetConstructor : public FTypedElementWidgetConstructor
+{
+	GENERATED_BODY()
+
+public:
+	TEDSUI_API FLabelWidgetConstructor();
+	~FLabelWidgetConstructor() override = default;
+
+	TEDSUI_API TConstArrayView<const UScriptStruct*> GetAdditionalColumnsList() const override;
+
+	TEDSUI_API TSharedPtr<SWidget> Construct(
+		TypedElementRowHandle Row, 
+		ITypedElementDataStorageInterface* DataStorage,
+		ITypedElementDataStorageUiInterface* DataStorageUi,
+		const TypedElementDataStorage::FMetaDataView& Arguments) override;
+
+protected:
+	explicit FLabelWidgetConstructor(const UScriptStruct* InTypeInfo);
+	bool SetColumns(ITypedElementDataStorageInterface* DataStorage, TypedElementRowHandle Row) override;
+	bool FinalizeWidget(ITypedElementDataStorageInterface* DataStorage, ITypedElementDataStorageUiInterface* DataStorageUi,
+		TypedElementRowHandle Row, const TSharedPtr<SWidget>& Widget) override;
+};
+
+USTRUCT(meta = (DisplayName = "Label widget"))
+struct FLabelWidgetColumn : public FTypedElementDataStorageColumn
+{
+	GENERATED_BODY()
+
+	bool bShowHashInTooltip { false };
+};
\ No newline at end of file
diff --git a/Engine/Plugins/Experimental/EditorDataStorage/Source/TedsUI/Public/Widgets/PackagePathWidget.h b/Engine/Plugins/Experimental/EditorDataStorage/Source/TedsUI/Public/Widgets/PackagePathWidget.h
new file mode 100644
--- /dev/null
+++ b/Engine/Plugins/Experimental/EditorDataStorage/Source/TedsUI/Public/Widgets/PackagePathWidget.h
@@ -0,0 +1,55 @@
+// Copyright Epic Games, Inc. All Rights Reserved.
+
+#pragma once
+
+#include "Elements/Interfaces/TypedElementDataStorageFactory.h"
+#include "Elements/Interfaces/TypedElementDataStorageUiInterface.h"
+
+#include "PackagePathWidget.generated.h"
+
+class ITypedElementDataStorageInterface;
+class SWidget;
+class UScriptStruct;
+
+UCLASS()
+class TEDSUI_API UPackagePathWidgetFactory : public UTypedElementDataStorageFactory
+{
+	GENERATED_BODY()
+
+public:
+	~UPackagePathWidgetFactory() override = default;
+
+	void RegisterWidgetConstructors(ITypedElementDataStorageInterface& DataStorage,
+		ITypedElementDataStorageUiInterface& DataStorageUi) const override;
+};
+
+USTRUCT()
+struct TEDSUI_API FPackagePathWidgetConstructor : public FTypedElementWidgetConstructor
+{
+	GENERATED_BODY()
+
+public:
+	FPackagePathWidgetConstructor();
+	~FPackagePathWidgetConstructor() override = default;
+
+protected:
+	explicit FPackagePathWidgetConstructor(const UScriptStruct* InTypeInfo);
+
+	TSharedPtr<SWidget> CreateWidget(const TypedElementDataStorage::FMetaDataView& Arguments) override;
+	bool FinalizeWidget(ITypedElementDataStorageInterface* DataStorage, ITypedElementDataStorageUiInterface* DataStorageUi,
+		TypedElementRowHandle Row, const TSharedPtr<SWidget>& Widget) override;
+};
+
+USTRUCT()
+struct TEDSUI_API FLoadedPackagePathWidgetConstructor : public FPackagePathWidgetConstructor
+{
+	GENERATED_BODY()
+
+public:
+	FLoadedPackagePathWidgetConstructor();
+	~FLoadedPackagePathWidgetConstructor() override = default;
+
+protected:
+	bool FinalizeWidget(ITypedElementDataStorageInterface* DataStorage, ITypedElementDataStorageUiInterface* DataStorageUi,
+		TypedElementRowHandle Row, const TSharedPtr<SWidget>& Widget) override;
+};
diff --git a/Engine/Plugins/Experimental/EditorDataStorage/Source/TedsUI/TedsUI.Build.cs b/Engine/Plugins/Experimental/EditorDataStorage/Source/TedsUI/TedsUI.Build.cs
--- a/Engine/Plugins/Experimental/EditorDataStorage/Source/TedsUI/TedsUI.Build.cs
+++ b/Engine/Plugins/Experimental/EditorDataStorage/Source/TedsUI/TedsUI.Build.cs
@@ -19,15 +19,11 @@ public class TedsUI : ModuleRules
 					"Core",
 					"CoreUObject",
 					"Engine",
-					"EditorFramework",
-					"SceneOutliner",
 					"Slate",
 					"SlateCore",
 					"ToolMenus",
 					"TypedElementFramework",
 					"TedsCore",
-					"EditorSubsystem",
-					"UnrealEd",
 				});
 
 			DynamicallyLoadedModuleNames.AddRange(new string[] {});
