[TEDS] DynamicTags Feature
* DynamicTags are mechanism to add attributes to rows that are cheaply filterable without incurring the per-row memory and processing cost to perform the filtering.
* Current capabilities allow filtering on a single DynamicTag value or all values.
* DynamicTags are mutually exclusive on a row.  For example, a row can be assigned a `DynamicTag("Color") { "Red" }` and `DynamicTag("AssetType") { "StaticMesh" }` but cannot be subsequently assigned a DynamicTag("Color") { "Blue" }

* Under the hood, they are implemented using Shared Fragments in Mass.  TEDS generates a Shared Fragment type dynamically for each unique FDynamicTag that is given to it.  Each type contains an FName value which can be filtered.  The Shared Fragments are attributes of the chunk that the entities belong to rather than the entity data itself.  A special ChunkFilter function is crafted to skip whole chunks that do not conform to filters that TEDS has set up.

Debug commands have been added to test and explore the feature.
`Teds.Debug.CreateRow` - Creates a row and prints the row ID to the console
`Teds.Debug.DynamicTag.AddColumn Row Tag Value` - Adds a DynamicTag of type <Tag> to the given <Row>. Sets the tag's value to <Value>
`Teds.Debug.DynamicTag.RemoveColumn Row Tag` - Removes a DynamicTag of type <Tag> column from the given <Row>
`Teds.Debug.DynamicTag.RunQuery Tag [Value]` - Runs a query that matches all rows with <Tag>.  Optionally, will match only Tags with the given <Value>
`Teds.Debug.DynamicTag.AddWithEnum Row Value` - Adds a DynamicTag using the ETedsDebugEnum enum to the given <Row>. Sets the tag's value to <Value> - it must be one of the enums
`Teds.Debug.DynamicTag.RemoveWithEnum Row` - Removes a DynamicTag using the ETedsDebugEnum from the given <Row>
`Teds.Debug.DynamicTag.RunQueryEnum [Value]` - Runs a query that matches all rows with a dynamic tag generated from ETedsDebugEnum.  Optionally, will match only Tags with the given <Value>

diff --git a/Engine/Plugins/Experimental/EditorDataStorage/Source/TedsCore/Private/Debug/TypedElementDatabaseDebugCommands.cpp b/Engine/Plugins/Experimental/EditorDataStorage/Source/TedsCore/Private/Debug/TypedElementDatabaseDebugCommands.cpp
--- a/Engine/Plugins/Experimental/EditorDataStorage/Source/TedsCore/Private/Debug/TypedElementDatabaseDebugCommands.cpp
+++ b/Engine/Plugins/Experimental/EditorDataStorage/Source/TedsCore/Private/Debug/TypedElementDatabaseDebugCommands.cpp
@@ -1,10 +1,13 @@
 // Copyright Epic Games, Inc. All Rights Reserved.
 
+#include "TypedElementDatabaseDebugTypes.h"
 #include "Elements/Columns/TypedElementCompatibilityColumns.h"
 #include "Elements/Columns/TypedElementLabelColumns.h"
 #include "Elements/Columns/TypedElementTypeInfoColumns.h"
+#include "Elements/Common/TypedElementDataStorageLog.h"
 #include "Elements/Framework/TypedElementQueryBuilder.h"
 #include "Elements/Framework/TypedElementRegistry.h"
+#include "Elements/Framework/TypedElementTestColumns.h"
 #include "Elements/Interfaces/TypedElementDataStorageInterface.h"
 #include "Elements/Interfaces/TypedElementDataStorageCompatibilityInterface.h"
 #include "Elements/Interfaces/TypedElementDataStorageUiInterface.h"
@@ -206,3 +209,275 @@ FAutoConsoleCommandWithOutputDevice ListExtensionsConsoleCommand(
 		}
 	));
 
+
+
+static FAutoConsoleCommand CVarCreateRow(
+	TEXT("TEDS.Debug.CreateRow"),
+	TEXT("Argument: \n"),
+	FConsoleCommandWithArgsDelegate::CreateLambda([](const TArray<FString>& Args)
+	{
+		ITypedElementDataStorageInterface* DataStorage = UTypedElementRegistry::GetInstance()->GetMutableDataStorage();
+		static TypedElementDataStorage::TableHandle Table = DataStorage->RegisterTable<FTestColumnA>(FName(TEXT("Debug.CreateRow Table")));
+
+		const TypedElementDataStorage::RowHandle RowHandle = DataStorage->AddRow(Table);
+
+		UE_LOG(LogTypedElementDataStorage, Warning, TEXT("Added Row %llu"), static_cast<uint64>(RowHandle));
+	}));
+
+static FAutoConsoleCommand CVarAddDynamicTag(
+	TEXT("TEDS.Debug.DynamicTag.AddColumn"),
+	TEXT("Argument: Row, Tag, Value\n"),
+	FConsoleCommandWithArgsDelegate::CreateLambda([](const TArray<FString>& Args)
+	{
+		using namespace UE::EditorDataStorage;
+		ITypedElementDataStorageInterface* DataStorage = UTypedElementRegistry::GetInstance()->GetMutableDataStorage();
+
+		if (Args.Num() != 3)
+		{
+			return;
+		}
+		
+		const uint64 RowAsU64 = FCString::Strtoui64(*Args[0], nullptr, 10);
+		const TypedElementDataStorage::RowHandle Row = RowAsU64;
+
+		
+		const FName Value(*Args[2]);
+
+		constexpr bool bUseTemplateSugar = true;
+		if constexpr (bUseTemplateSugar)
+		{
+			const FName Tag = *Args[1];
+			DataStorage->AddColumn<FDynamicTag>(Row, Tag, Value); 
+		}
+		else
+		{
+			const FDynamicTag Tag(*Args[1]);
+			DataStorage->AddColumn(Row, Tag, Value);
+		}
+		
+		
+	}),
+	ECVF_Default);
+
+static FAutoConsoleCommand CVarRemoveDynamicTag(
+	TEXT("TEDS.Debug.DynamicTag.RemoveColumn"),
+	TEXT("Argument: Row, Group\n"),
+	FConsoleCommandWithArgsDelegate::CreateLambda([](const TArray<FString>& Args)
+	{
+		ITypedElementDataStorageInterface* DataStorage = UTypedElementRegistry::GetInstance()->GetMutableDataStorage();
+
+		if (Args.Num() != 2)
+		{
+			return;
+		}
+		
+		const uint64 RowAsU64 = FCString::Strtoui64(*Args[0], nullptr, 10);
+		const TypedElementDataStorage::RowHandle Row = RowAsU64;
+		
+		constexpr bool bUseTemplateSugar = true;
+		if constexpr (bUseTemplateSugar)
+		{
+			using namespace UE::EditorDataStorage;
+			const FName Tag = *Args[1];
+			DataStorage->RemoveColumn<FDynamicTag>(Row, Tag);
+		}
+		else
+		{
+			const UE::EditorDataStorage::FDynamicTag Tag(*Args[1]);
+			DataStorage->RemoveColumn(Row, Tag);
+		}		
+	}),
+	ECVF_Default);
+	
+static FAutoConsoleCommand CVarMatchDynamicTag(
+	TEXT("TEDS.Debug.DynamicTag.RunQuery"),
+	TEXT("Argument: Tag, [optional] Value\n"),
+	FConsoleCommandWithArgsDelegate::CreateLambda([](const TArray<FString>& Args)
+	{
+		using namespace TypedElementQueryBuilder;
+		using DSI = ITypedElementDataStorageInterface;
+		using namespace UE::EditorDataStorage;
+		
+		ITypedElementDataStorageInterface* DataStorage = UTypedElementRegistry::GetInstance()->GetMutableDataStorage();
+
+		if (Args.Num() < 1 || Args.Num() > 2)
+		{
+			return;
+		}			
+
+		const TypedElementQueryHandle Query = [&Args, DataStorage]() -> TypedElementQueryHandle
+		{
+			const FName Tag(*Args[0]);
+			if (Args.Num() == 1)
+			{
+				// Matches all rows with the 
+				return DataStorage->RegisterQuery(
+					Select().
+						Where().
+							// Match all rows with a dynamic tag of type Tag (ie. all rows with a dynamic tag of "Color")
+							All<FDynamicTag>(Tag).
+							All<FTestColumnA>().
+						Compile());
+			}
+			else
+			{
+				const FName MatchValue(*Args[1]);
+				return DataStorage->RegisterQuery(
+					Select().
+					Where().
+						// Match all rows with a dynamic tag of type Tag that has a MatchValue (ie. all rows with dynamic tag "Color" with value "Red")
+						All<FDynamicTag>(Tag, MatchValue).
+						All<FTestColumnA>().
+					Compile());
+			}
+		}();
+
+		uint64 Count = 0;
+		
+		const TypedElementDataStorage::FQueryResult Result = DataStorage->RunQuery(Query, CreateDirectQueryCallbackBinding(
+			[&Count](const DSI::IDirectQueryContext& Context, const TypedElementRowHandle*)
+			{
+				Count += Context.GetRowCount();
+			}));
+		DataStorage->UnregisterQuery(Query);
+
+		UE_LOG(LogTypedElementDataStorage, Warning, TEXT("Processed %llu rows"), static_cast<uint64>(Count));
+	}),
+	ECVF_Default);
+
+static FAutoConsoleCommand CVarAddDynamicTagFromEnum(
+	TEXT("TEDS.Debug.DynamicTag.AddWithEnum"),
+	TEXT("Argument: Row, EnumValue\n"),
+	FConsoleCommandWithArgsDelegate::CreateLambda([](const TArray<FString>& Args)
+	{
+		ITypedElementDataStorageInterface* DataStorage = UTypedElementRegistry::GetInstance()->GetMutableDataStorage();
+
+		if (Args.Num() < 1 || Args.Num() > 2)
+		{
+			return;
+		}
+		
+		const uint64 RowAsU64 = FCString::Strtoui64(*Args[0], nullptr, 10);
+		const TypedElementDataStorage::RowHandle Row = RowAsU64;
+
+		if (Args.Num() == 1)
+		{
+			// Use of a enum value directly as a template parameter.  Only useful if enum value known at compile time
+			DataStorage->AddColumn<ETedsDebugEnum::Red>(Row);
+		}
+		else
+		{
+			// Use an enum value from a runtime source
+			// In this case, the argument is parsed and converted to an enum type
+			UEnum* Enum = StaticEnum<ETedsDebugEnum>();
+			int64 EnumValueAsI64 = Enum->GetValueByNameString(*Args[1]);
+			if (EnumValueAsI64 == INDEX_NONE)
+			{
+				return;
+			}
+			const ETedsDebugEnum EnumValue = static_cast<ETedsDebugEnum>(EnumValueAsI64);
+			
+			DataStorage->AddColumn(Row, EnumValue);
+		}
+	}),
+	ECVF_Default);
+
+static FAutoConsoleCommand CVarRemoveDynamicTagFromEnum(
+	TEXT("TEDS.Debug.DynamicTag.RemoveWithEnum"),
+	TEXT("Argument: Row\n"),
+	FConsoleCommandWithArgsDelegate::CreateLambda([](const TArray<FString>& Args)
+	{
+		ITypedElementDataStorageInterface* DataStorage = UTypedElementRegistry::GetInstance()->GetMutableDataStorage();
+
+		if (Args.Num() != 1)
+		{
+			return;
+		}
+		
+		const uint64 RowAsU64 = FCString::Strtoui64(*Args[0], nullptr, 10);
+		const TypedElementDataStorage::RowHandle Row = RowAsU64;
+		
+		DataStorage->RemoveColumn<ETedsDebugEnum>(Row);
+	}),
+	ECVF_Default);
+
+
+static FAutoConsoleCommand CVarMatchDynamicTagFromEnum(
+	TEXT("TEDS.Debug.DynamicTag.RunQueryEnum"),
+	TEXT("Argument: [optional] EnumValue\n"),
+	FConsoleCommandWithArgsDelegate::CreateLambda([](const TArray<FString>& Args)
+	{
+		using namespace TypedElementQueryBuilder;
+		using DSI = ITypedElementDataStorageInterface;
+		using namespace UE::EditorDataStorage;
+		
+		ITypedElementDataStorageInterface* DataStorage = UTypedElementRegistry::GetInstance()->GetMutableDataStorage();
+
+		if (Args.Num() > 1)
+		{
+			UE_LOG(LogTypedElementDataStorage, Error, TEXT("Invalid number of arguments"));
+			return;
+		}
+
+		if (Args.Num() == 1)
+		{
+			// Make sure that the given enum value is actually a value
+			UEnum* Enum = StaticEnum<ETedsDebugEnum>();
+			int64 EnumValue = Enum->GetValueByNameString(*Args[0]);
+			if (EnumValue == INDEX_NONE)
+			{
+				return;
+			}
+		}
+
+		const TypedElementQueryHandle Query = [&Args, DataStorage]() -> TypedElementQueryHandle
+		{
+			if (Args.Num() == 0)
+			{
+				// Matches all rows with the 
+				return DataStorage->RegisterQuery(
+					Select().
+						Where().
+							// Match all rows with an enum dynamic tag of the hardcoded enum type
+							All<ETedsDebugEnum>().
+						Compile());
+			}
+			else if (Args.Num() == 1)
+			{
+				UEnum* Enum = StaticEnum<ETedsDebugEnum>();
+				int64 EnumValue = Enum->GetValueByNameString(*Args[0]);
+				return DataStorage->RegisterQuery(
+					Select().
+					Where().
+						// Match all rows with a dynamic tag of the hardcoded enum type that has the given value
+						// Note, usually this would be written something like:
+						//   All(ETedsDebugEnum::Red).
+						// However it isn't possible to do that when getting the enum value from a string.  API is still exercised
+						// using the static_cast
+						All(static_cast<ETedsDebugEnum>(EnumValue)).
+					Compile());
+			}
+			else
+			{
+				return TypedElementDataStorage::InvalidQueryHandle;
+			}
+		}();
+		if (Query == TypedElementDataStorage::InvalidQueryHandle)
+		{
+			UE_LOG(LogTypedElementDataStorage, Error, TEXT("Invalid number of arguments"));
+			return;
+		}
+		
+		uint64 Count = 0;
+		
+		const TypedElementDataStorage::FQueryResult Result = DataStorage->RunQuery(Query, CreateDirectQueryCallbackBinding(
+			[&Count](const DSI::IDirectQueryContext& Context, const TypedElementRowHandle*)
+			{
+				Count += Context.GetRowCount();
+			}));
+		DataStorage->UnregisterQuery(Query);
+
+		UE_LOG(LogTypedElementDataStorage, Warning, TEXT("Processed %llu rows"), static_cast<uint64>(Count));
+		
+	}),
+	ECVF_Default);
\ No newline at end of file
diff --git a/Engine/Plugins/Experimental/EditorDataStorage/Source/TedsCore/Private/Debug/TypedElementDatabaseDebugTypes.h b/Engine/Plugins/Experimental/EditorDataStorage/Source/TedsCore/Private/Debug/TypedElementDatabaseDebugTypes.h
new file mode 100644
--- /dev/null
+++ b/Engine/Plugins/Experimental/EditorDataStorage/Source/TedsCore/Private/Debug/TypedElementDatabaseDebugTypes.h
@@ -0,0 +1,20 @@
+// Copyright Epic Games, Inc. All Rights Reserved.
+
+#pragma once
+
+#include "UObject/ObjectMacros.h"
+
+#include "TypedElementDatabaseDebugTypes.generated.h"
+
+UENUM()
+enum class ETedsDebugEnum : int8
+{
+	Red,
+	Blue,
+	Green,
+	Yellow,
+	Black,
+	Pink,
+	Orange,
+	Purple
+};
diff --git a/Engine/Plugins/Experimental/EditorDataStorage/Source/TedsCore/Private/DynamicColumnGenerator.cpp b/Engine/Plugins/Experimental/EditorDataStorage/Source/TedsCore/Private/DynamicColumnGenerator.cpp
new file mode 100644
--- /dev/null
+++ b/Engine/Plugins/Experimental/EditorDataStorage/Source/TedsCore/Private/DynamicColumnGenerator.cpp
@@ -0,0 +1,119 @@
+﻿// Copyright Epic Games, Inc. All Rights Reserved.
+
+#include "DynamicColumnGenerator.h"
+
+#include "AssetDefinitionDefault.h"
+#include "TypedElementDataStorageSharedColumn.h"
+#include "Elements/Common/TypedElementCommonTypes.h"
+
+UE::EditorDataStorage::FDynamicColumnGeneratorInfo UE::EditorDataStorage::FDynamicColumnGenerator::GenerateColumn(const FName& ColumnName, const UScriptStruct& Template)
+{
+	const FGeneratedColumnKey Key
+	{
+		.Name = ColumnName,
+		.Template = Template
+	};
+
+	FDynamicColumnGeneratorInfo GeneratedColumnInfo;
+	UE_MT_SCOPED_WRITE_ACCESS(AccessDetector);
+	
+	{
+		const int32* GeneratedColumnIndex = GenerationParamsLookup.Find(Key);
+		if (GeneratedColumnIndex != nullptr)
+		{
+			const FGeneratedColumnRecord& GeneratedColumnRecord = GeneratedColumnData[*GeneratedColumnIndex];
+			GeneratedColumnInfo.Type = GeneratedColumnRecord.Type;
+			GeneratedColumnInfo.bNewlyGenerated = false;
+			return GeneratedColumnInfo;
+		}		
+	}
+	
+	checkf(
+		Template.IsChildOf(FTypedElementDataStorageColumn::StaticStruct()) ||
+		Template.IsChildOf(FTypedElementDataStorageTag::StaticStruct()) ||
+		Template.IsChildOf(FTedsSharedColumn::StaticStruct()),
+		TEXT("Template struct must derive from Column, Tag or SharedColumn"));
+
+	{
+		checkf(Template.GetCppStructOps() != nullptr && Template.IsNative(), TEXT("Can only create column from native struct"));
+		
+		UScriptStruct* NewScriptStruct = NewObject<UScriptStruct>(GetTransientPackage(), ColumnName);
+		NewScriptStruct->AddToRoot();
+
+		NewScriptStruct->SetSuperStruct(&const_cast<UScriptStruct&>(Template));
+		
+		NewScriptStruct->DeferCppStructOps(FTopLevelAssetPath(GetTransientPackage()->GetFName(), ColumnName), Template.GetCppStructOps());
+		NewScriptStruct->Bind();
+		NewScriptStruct->PrepareCppStructOps();
+		NewScriptStruct->StaticLink(true);
+
+		
+		const int32 Index = GeneratedColumnData.Emplace(FGeneratedColumnRecord
+		{
+			.Name = ColumnName,
+			.Template = &Template,
+			.Type = NewScriptStruct
+		});
+		
+		GenerationParamsLookup.Add(Key, Index);
+		NameLookup.Add(ColumnName, Index);
+		
+		GeneratedColumnInfo.Type = NewScriptStruct;
+		GeneratedColumnInfo.bNewlyGenerated = true;
+		return GeneratedColumnInfo;
+	}
+}
+
+const UScriptStruct* UE::EditorDataStorage::FDynamicColumnGenerator::LookupColumn(const FName& ColumnName) const
+{
+	UE_MT_SCOPED_READ_ACCESS(AccessDetector);
+	
+	if (const int32* Index = NameLookup.Find(ColumnName))
+	{
+		return GeneratedColumnData[*Index].Type;
+	}
+	return nullptr;
+}
+
+UE::EditorDataStorage::FDynamicTagManager::FDynamicTagManager(FDynamicColumnGenerator& InColumnGenerator)
+	: ColumnGenerator(InColumnGenerator)
+{
+}
+
+FConstSharedStruct UE::EditorDataStorage::FDynamicTagManager::GenerateDynamicTag(const FDynamicTag& InTag, const FName& InValue)
+{
+	TPair<FDynamicTag, FName> Pair(InTag, InValue);
+
+	UE_MT_SCOPED_WRITE_ACCESS(AccessDetector);
+
+	// Common path
+	{
+		if (FConstSharedStruct* TagStruct = DynamicTagLookup.Find(Pair))
+		{
+			return *TagStruct;
+		}
+	}
+
+	// 
+	{
+		const UScriptStruct* ColumnType = GenerateColumnType(InTag);
+
+		const FTypedElementDataStorageDynamicTagTemplate Overlay
+		{
+			.Value = InValue
+		};
+
+		FConstSharedStruct SharedStruct = FConstSharedStruct::Make(ColumnType, reinterpret_cast<const uint8*>(&Overlay));
+		
+		DynamicTagLookup.Emplace(Pair, SharedStruct);
+
+		return SharedStruct;
+	}
+}
+
+const UScriptStruct* UE::EditorDataStorage::FDynamicTagManager::GenerateColumnType(const FDynamicTag& Tag)
+{
+	const FDynamicColumnGeneratorInfo GeneratedColumnType = ColumnGenerator.GenerateColumn(Tag.GetName(), *FTypedElementDataStorageDynamicTagTemplate::StaticStruct());
+	
+	return GeneratedColumnType.Type;
+}
diff --git a/Engine/Plugins/Experimental/EditorDataStorage/Source/TedsCore/Private/DynamicColumnGenerator.h b/Engine/Plugins/Experimental/EditorDataStorage/Source/TedsCore/Private/DynamicColumnGenerator.h
new file mode 100644
--- /dev/null
+++ b/Engine/Plugins/Experimental/EditorDataStorage/Source/TedsCore/Private/DynamicColumnGenerator.h
@@ -0,0 +1,111 @@
+﻿// Copyright Epic Games, Inc. All Rights Reserved.
+
+#pragma once
+
+#include "TypedElementDataStorageSharedColumn.h"
+#include "Containers/Map.h"
+#include "Elements/Interfaces/TypedElementDataStorageInterface.h"
+#include "Misc/MTAccessDetector.h"
+#include "StructUtils/PropertyBag.h"
+#include "UObject/NameTypes.h"
+#include "UObject/Class.h"
+
+#include "DynamicColumnGenerator.generated.h"
+
+class FName;
+class UScriptStruct;
+
+// The template struct that is used to generate the DynamicTag column
+// It is safe to reinterpret a DynamicTag column to this template to access the Value
+USTRUCT()
+struct FTypedElementDataStorageDynamicTagTemplate : public FTedsSharedColumn
+{
+	GENERATED_BODY()
+	
+	UPROPERTY()
+	FName Value;
+};
+
+namespace UE::EditorDataStorage
+{
+	struct FDynamicColumnInfo
+	{
+		const UScriptStruct* Type;
+	};
+	
+	struct FDynamicColumnGeneratorInfo
+	{
+		const UScriptStruct* Type;
+		bool bNewlyGenerated;
+	};
+
+	/**
+	 * Utility class that TEDS can use to dynamically generate column types on the fly
+	 */
+	class FDynamicColumnGenerator
+	{
+	public:
+		/**
+		 * Generates a dynamic TEDS column type based on a Template type (if it hasn't been generated before)
+		 */
+		FDynamicColumnGeneratorInfo GenerateColumn(const FName& ColumnName, const UScriptStruct& Template);
+
+		/**
+		 * Looks up a column based on the name given 
+		 */
+		const UScriptStruct* LookupColumn(const FName& ColumnName) const;
+	private:
+
+		struct FGeneratedColumnRecord
+		{
+			FName Name;
+			const UScriptStruct* Template;
+			const UScriptStruct* Type;
+		};
+
+		UE_MT_DECLARE_RW_ACCESS_DETECTOR(AccessDetector);
+
+		TArray<FGeneratedColumnRecord> GeneratedColumnData;
+
+		struct FGeneratedColumnKey
+		{
+			FName Name;
+			const UScriptStruct& Template;
+
+			friend bool operator==(const FGeneratedColumnKey& Lhs, const FGeneratedColumnKey& Rhs)
+			{
+				return Lhs.Name == Rhs.Name && &Lhs.Template == &Rhs.Template;
+			}
+
+			friend uint32 GetTypeHash(const FGeneratedColumnKey& Key)
+			{
+				return HashCombineFast(GetTypeHash(Key.Name), PointerHash(&Key.Template));
+			}
+		};
+		// Looks up generated column index by the parameters used to generate it
+		// Used to deduplicate
+		TMap<FGeneratedColumnKey, int32> GenerationParamsLookup;
+		// Looks up generated column index by name
+		TMap<FName, int32> NameLookup;
+	};
+
+	class FDynamicTagManager
+	{
+	public:
+		struct FDynamicTagStructLayout
+		{
+			FName Tag;
+		};
+		
+		explicit FDynamicTagManager(FDynamicColumnGenerator& InColumnGenerator);
+		FConstSharedStruct GenerateDynamicTag(const FDynamicTag& InTag, const FName& InValue);
+		const UScriptStruct* GenerateColumnType(const FDynamicTag& InTag);
+	private:
+
+		UE_MT_DECLARE_RW_ACCESS_DETECTOR(AccessDetector);
+
+		TMap<TPair<FDynamicTag, FName>, FConstSharedStruct> DynamicTagLookup;
+
+		FDynamicColumnGenerator& ColumnGenerator;
+	};
+}
diff --git a/Engine/Plugins/Experimental/EditorDataStorage/Source/TedsCore/Private/Queries/TypedElementExtendedQueryStore.cpp b/Engine/Plugins/Experimental/EditorDataStorage/Source/TedsCore/Private/Queries/TypedElementExtendedQueryStore.cpp
--- a/Engine/Plugins/Experimental/EditorDataStorage/Source/TedsCore/Private/Queries/TypedElementExtendedQueryStore.cpp
+++ b/Engine/Plugins/Experimental/EditorDataStorage/Source/TedsCore/Private/Queries/TypedElementExtendedQueryStore.cpp
@@ -4,6 +4,8 @@
 
 #include "MassProcessingPhaseManager.h"
 #include "MassProcessor.h"
+#include "TypedElementDatabaseEnvironment.h"
+#include "Commands/EditorDataStorageCommandBuffer.h"
 #include "Misc/OutputDevice.h"
 #include "Processors/TypedElementProcessorAdaptors.h"
 
@@ -21,6 +23,7 @@ FTypedElementExtendedQueryStore::Handle FTypedElementExtendedQueryStore::Registe
 
 	FMassEntityQuery& NativeQuery =		SetupNativeQuery(StoredQuery.Description, StoredQuery);
 	bool bContinueSetup =				SetupSelectedColumns(StoredQuery.Description, NativeQuery);
+	bContinueSetup = bContinueSetup &&	SetupChunkFilters(Result, StoredQuery.Description, Environment, NativeQuery);
 	bContinueSetup = bContinueSetup &&	SetupConditions(StoredQuery.Description, NativeQuery);
 	bContinueSetup = bContinueSetup &&	SetupDependencies(StoredQuery.Description, NativeQuery);
 	bContinueSetup = bContinueSetup &&	SetupTickGroupDefaults(StoredQuery.Description);
@@ -580,6 +583,90 @@ bool FTypedElementExtendedQueryStore::SetupConditions(
 	return false;
 }
 
+static const FName DynamicTagDataParamName = TEXT("TagName");
+
+bool FTypedElementExtendedQueryStore::SetupChunkFilters(
+	Handle QueryHandle,
+	ITypedElementDataStorageInterface::FQueryDescription& Query,
+	FTypedElementDatabaseEnvironment& Environment,
+	FMassEntityQuery& NativeQuery)
+{
+	if (Query.DynamicTags.IsEmpty())
+	{
+		return true;
+	}
+
+	Algo::SortBy(Query.DynamicTags, [](const TypedElementDataStorage::FQueryDescription::FDynamicTagData& DynamicTagData)
+	{
+		return DynamicTagData.Tag.GetName();
+	}, FNameFastLess());
+	// Check if there are any duplicate groups. Not yet supported until we can match multiple MatchTags
+	UE::EditorDataStorage::FDynamicTag PreviousTag = Query.DynamicTags[0].Tag;
+	for (int32 Index = 1, End = Query.DynamicTags.Num(); Index < End; ++Index)
+	{
+		if (Query.DynamicTags[Index].Tag == PreviousTag)
+		{
+			return false;
+		}
+		PreviousTag = Query.DynamicTags[Index].Tag;
+	}
+	
+	struct FGroupTagPair
+	{
+		const UScriptStruct* ColumnType;
+		FName Value;
+	};
+
+	TArray<FGroupTagPair> GroupTagPairsTemp;
+	GroupTagPairsTemp.Reserve(Query.DynamicTags.Num());
+	for (int32 Index = 0, End = Query.DynamicTags.Num(); Index < End; ++Index)
+	{
+		const UScriptStruct* ColumnType = Environment.GenerateColumnType(Query.DynamicTags[Index].Tag);
+		GroupTagPairsTemp.Emplace(FGroupTagPair
+			{
+				.ColumnType = ColumnType,
+				.Value = Query.DynamicTags[Index].MatchValue
+			});
+	}
+
+	check(!GroupTagPairsTemp.IsEmpty());
+
+	for (const FGroupTagPair& Element : GroupTagPairsTemp)
+	{
+		NativeQuery.AddConstSharedRequirement(Element.ColumnType);
+	}
+	
+	auto ChunkFilterFunction = [GroupTagPairs = MoveTemp(GroupTagPairsTemp)](const FMassExecutionContext& MassContext) -> bool
+	{
+		for (const FGroupTagPair& GroupTagPair : GroupTagPairs)
+		{
+			const void* SharedFragmentData = MassContext.GetConstSharedFragmentPtr(*GroupTagPair.ColumnType);
+
+			if (SharedFragmentData)
+			{
+				// TODO: Use reflection / knowledge of key offset to read the tag and compare it
+				// For now... we know it is at offset 0
+				const FTypedElementDataStorageDynamicTagTemplate* TagOverlay = static_cast<const FTypedElementDataStorageDynamicTagTemplate*>(SharedFragmentData);
+				// NAME_None will match any presence of the shared fragment
+				// otherwise, match the specific tag only
+				if (GroupTagPair.Value != NAME_None && TagOverlay->Value != GroupTagPair.Value)
+				{
+					return false;
+				}
+			}
+			else
+			{
+				return false;
+			}
+		}
+		
+		return true;
+	};
+	
+	NativeQuery.SetChunkFilter(ChunkFilterFunction);
+	return true;
+}
+
 bool FTypedElementExtendedQueryStore::SetupDependencies(
 	ITypedElementDataStorageInterface::FQueryDescription& Query, FMassEntityQuery& NativeQuery)
 {
diff --git a/Engine/Plugins/Experimental/EditorDataStorage/Source/TedsCore/Private/Queries/TypedElementExtendedQueryStore.h b/Engine/Plugins/Experimental/EditorDataStorage/Source/TedsCore/Private/Queries/TypedElementExtendedQueryStore.h
--- a/Engine/Plugins/Experimental/EditorDataStorage/Source/TedsCore/Private/Queries/TypedElementExtendedQueryStore.h
+++ b/Engine/Plugins/Experimental/EditorDataStorage/Source/TedsCore/Private/Queries/TypedElementExtendedQueryStore.h
@@ -162,6 +162,7 @@ private:
 	FMassEntityQuery& SetupNativeQuery(ITypedElementDataStorageInterface::FQueryDescription& Query, FTypedElementExtendedQuery& StoredQuery);
 	bool SetupSelectedColumns(ITypedElementDataStorageInterface::FQueryDescription& Query, FMassEntityQuery& NativeQuery);
 	bool SetupConditions(ITypedElementDataStorageInterface::FQueryDescription& Query, FMassEntityQuery& NativeQuery);
+	bool SetupChunkFilters(Handle QueryHandle, ITypedElementDataStorageInterface::FQueryDescription& Query, FTypedElementDatabaseEnvironment& Environment, FMassEntityQuery& NativeQuery);
 	bool SetupDependencies(ITypedElementDataStorageInterface::FQueryDescription& Query, FMassEntityQuery& NativeQuery);
 	bool SetupTickGroupDefaults(ITypedElementDataStorageInterface::FQueryDescription& Query);
 	bool SetupProcessors(Handle QueryHandle, FTypedElementExtendedQuery& StoredQuery, FTypedElementDatabaseEnvironment& Environment,
diff --git a/Engine/Plugins/Experimental/EditorDataStorage/Source/TedsCore/Private/TypedElementDataStorageSharedColumn.h b/Engine/Plugins/Experimental/EditorDataStorage/Source/TedsCore/Private/TypedElementDataStorageSharedColumn.h
new file mode 100644
--- /dev/null
+++ b/Engine/Plugins/Experimental/EditorDataStorage/Source/TedsCore/Private/TypedElementDataStorageSharedColumn.h
@@ -0,0 +1,16 @@
+// Copyright Epic Games, Inc. All Rights Reserved.
+
+#pragma once
+
+#include "CoreMinimal.h"
+
+#include "TypedElementDataStorageSharedColumn.generated.h"
+
+/*
+ * Shared Columns are virtual columns shared by many rows
+ */
+USTRUCT()
+struct FTedsSharedColumn
+{
+	GENERATED_BODY()
+};
\ No newline at end of file
diff --git a/Engine/Plugins/Experimental/EditorDataStorage/Source/TedsCore/Private/TypedElementDatabase.cpp b/Engine/Plugins/Experimental/EditorDataStorage/Source/TedsCore/Private/TypedElementDatabase.cpp
--- a/Engine/Plugins/Experimental/EditorDataStorage/Source/TedsCore/Private/TypedElementDatabase.cpp
+++ b/Engine/Plugins/Experimental/EditorDataStorage/Source/TedsCore/Private/TypedElementDatabase.cpp
@@ -602,6 +602,33 @@ void UTypedElementDatabase::AddColumns(TypedElementRowHandle Row, TConstArrayVie
 	}
 }
 
+void UTypedElementDatabase::AddColumn(TypedElementRowHandle Row, const UE::EditorDataStorage::FDynamicTag& Tag, const FName& InValue)
+{
+	if (ActiveEditorEntityManager)
+	{
+		const FConstSharedStruct SharedStruct = Environment->GenerateDynamicTag(Tag, InValue);
+
+		FMassEntityHandle Entity = FMassEntityHandle::FromNumber(Row);
+		if (ActiveEditorEntityManager->IsEntityActive(Entity))
+		{
+			FTypedElementDatabaseCommandBuffer::Execute_AddSharedColumnCommand(*ActiveEditorEntityManager, Row, SharedStruct);
+		}
+	}
+}
+
+void UTypedElementDatabase::RemoveColumn(TypedElementRowHandle Row, const UE::EditorDataStorage::FDynamicTag& Tag)
+{
+	if (ActiveEditorEntityManager)
+	{
+		const UScriptStruct* DynamicTagType = Environment->GenerateColumnType(Tag);
+		FMassEntityHandle Entity = FMassEntityHandle::FromNumber(Row);
+		if (ActiveEditorEntityManager->IsEntityActive(Entity))
+		{
+			FTypedElementDatabaseCommandBuffer::Execute_RemoveSharedColumnCommand(*ActiveEditorEntityManager, Row, *DynamicTagType);
+		}
+	}
+}
+
 void UTypedElementDatabase::RemoveColumns(TypedElementRowHandle Row, TConstArrayView<const UScriptStruct*> Columns)
 {
 	FMassEntityHandle Entity = FMassEntityHandle::FromNumber(Row);
diff --git a/Engine/Plugins/Experimental/EditorDataStorage/Source/TedsCore/Private/TypedElementDatabase.h b/Engine/Plugins/Experimental/EditorDataStorage/Source/TedsCore/Private/TypedElementDatabase.h
--- a/Engine/Plugins/Experimental/EditorDataStorage/Source/TedsCore/Private/TypedElementDatabase.h
+++ b/Engine/Plugins/Experimental/EditorDataStorage/Source/TedsCore/Private/TypedElementDatabase.h
@@ -104,10 +104,12 @@ public:
 	bool IsRowAssigned(TypedElementDataStorage::RowHandle Row) const override;
 
 	void AddColumn(TypedElementRowHandle Row, const UScriptStruct* ColumnType) override;
+	void AddColumn(TypedElementRowHandle Row, const UE::EditorDataStorage::FDynamicTag& Tag, const FName& InValue) override;
 	void AddColumnData(TypedElementRowHandle Row, const UScriptStruct* ColumnType,
 		const TypedElementDataStorage::ColumnCreationCallbackRef& Initializer,
 		TypedElementDataStorage::ColumnCopyOrMoveCallback Relocator) override;
 	void RemoveColumn(TypedElementRowHandle Row, const UScriptStruct* ColumnType) override;
+	void RemoveColumn(TypedElementRowHandle Row, const UE::EditorDataStorage::FDynamicTag& Tag) override;
 	void* GetColumnData(TypedElementRowHandle Row, const UScriptStruct* ColumnType) override;
 	const void* GetColumnData(TypedElementRowHandle Row, const UScriptStruct* ColumnType) const override;
 	void AddColumns(TypedElementRowHandle Row, TConstArrayView<const UScriptStruct*> Columns) override;
diff --git a/Engine/Plugins/Experimental/EditorDataStorage/Source/TedsCore/Private/TypedElementDatabaseCommandBuffer.cpp b/Engine/Plugins/Experimental/EditorDataStorage/Source/TedsCore/Private/TypedElementDatabaseCommandBuffer.cpp
--- a/Engine/Plugins/Experimental/EditorDataStorage/Source/TedsCore/Private/TypedElementDatabaseCommandBuffer.cpp
+++ b/Engine/Plugins/Experimental/EditorDataStorage/Source/TedsCore/Private/TypedElementDatabaseCommandBuffer.cpp
@@ -5,6 +5,7 @@
 #include "HAL/UnrealMemory.h"
 #include "MassEntityManager.h"
 #include "TypedElementDatabaseEnvironment.h"
+#include "TypedElementDataStorageSharedColumn.h"
 
 // 
 // Commands section
@@ -237,6 +238,24 @@ void FTypedElementDatabaseCommandBuffer::Execute_AddDataColumnCommand(
 	}
 }
 
+void FTypedElementDatabaseCommandBuffer::Execute_AddSharedColumnCommand(FMassEntityManager& MassEntityManager, TypedElementDataStorage::RowHandle Row, const FConstSharedStruct& SharedColumn)
+{
+	if (SharedColumn.IsValid())
+	{
+		FMassEntityHandle Entity = FMassEntityHandle::FromNumber(Row);
+		MassEntityManager.AddConstSharedFragmentToEntity(Entity, SharedColumn);
+	}
+}
+
+void FTypedElementDatabaseCommandBuffer::Execute_RemoveSharedColumnCommand(FMassEntityManager& MassEntityManager, TypedElementDataStorage::RowHandle Row, const UScriptStruct& ColumnType)
+{
+	if (ColumnType.IsChildOf(FTedsSharedColumn::StaticStruct()))
+	{
+		FMassEntityHandle Entity = FMassEntityHandle::FromNumber(Row);
+		MassEntityManager.RemoveConstSharedFragmentFromEntity(Entity, ColumnType);
+	}
+}
+
 void FTypedElementDatabaseCommandBuffer::Execute_AddColumnsCommand(
 	FMassEntityManager& MassEntityManager,
 	TypedElementDataStorage::RowHandle Row,
diff --git a/Engine/Plugins/Experimental/EditorDataStorage/Source/TedsCore/Private/TypedElementDatabaseCommandBuffer.h b/Engine/Plugins/Experimental/EditorDataStorage/Source/TedsCore/Private/TypedElementDatabaseCommandBuffer.h
--- a/Engine/Plugins/Experimental/EditorDataStorage/Source/TedsCore/Private/TypedElementDatabaseCommandBuffer.h
+++ b/Engine/Plugins/Experimental/EditorDataStorage/Source/TedsCore/Private/TypedElementDatabaseCommandBuffer.h
@@ -50,6 +50,10 @@ public:
 	static void Execute_AddDataColumnCommand(
 		FMassEntityManager& MassEntityManager, TypedElementDataStorage::RowHandle Row, const UScriptStruct* ColumnType,
 		void* Data, TypedElementDataStorage::ColumnCopyOrMoveCallback Relocator);
+	static void Execute_AddSharedColumnCommand(
+		FMassEntityManager& MassEntityManager, TypedElementDataStorage::RowHandle Row, const FConstSharedStruct& SharedColumn);
+	static void Execute_RemoveSharedColumnCommand(
+		FMassEntityManager& MassEntityManager, TypedElementDataStorage::RowHandle Row, const UScriptStruct& ColumnType);
 	static void Execute_AddColumnsCommand(
 		FMassEntityManager& MassEntityManager, TypedElementDataStorage::RowHandle Row,
 		FMassFragmentBitSet FragmentsToAdd, FMassTagBitSet TagsToAdd);
diff --git a/Engine/Plugins/Experimental/EditorDataStorage/Source/TedsCore/Private/TypedElementDatabaseEnvironment.cpp b/Engine/Plugins/Experimental/EditorDataStorage/Source/TedsCore/Private/TypedElementDatabaseEnvironment.cpp
--- a/Engine/Plugins/Experimental/EditorDataStorage/Source/TedsCore/Private/TypedElementDatabaseEnvironment.cpp
+++ b/Engine/Plugins/Experimental/EditorDataStorage/Source/TedsCore/Private/TypedElementDatabaseEnvironment.cpp
@@ -8,6 +8,7 @@ FTypedElementDatabaseEnvironment::FTypedElementDatabaseEnvironment(UTypedElement
 	: DataStorage(InDataStorage)
 	, DirectDeferredCommands(*this)
 	, MementoSystem(InDataStorage)
+	, DynamicTagManager(DynamicColumnGenerator)
 	, MassEntityManager(InMassEntityManager)
 	, MassPhaseManager(InMassPhaseManager)
 {
@@ -87,6 +88,17 @@ const FMassProcessingPhaseManager& FTypedElementDatabaseEnvironment::GetMassPhas
 {
 	return MassPhaseManager;
 }
+FConstSharedStruct GenerateDynamicTag(const UE::EditorDataStorage::FDynamicTag& Tag, const FName& Value);
+const UScriptStruct* GenerateColumnType(const UE::EditorDataStorage::FDynamicTag& Tag);
+FConstSharedStruct FTypedElementDatabaseEnvironment::GenerateDynamicTag(const UE::EditorDataStorage::FDynamicTag& Tag, const FName& Value)
+{
+	return DynamicTagManager.GenerateDynamicTag(Tag, Value);
+}
+
+const UScriptStruct* FTypedElementDatabaseEnvironment::GenerateColumnType(const UE::EditorDataStorage::FDynamicTag& Tag)
+{
+	return DynamicTagManager.GenerateColumnType(Tag);
+}
 
 void FTypedElementDatabaseEnvironment::NextUpdateCycle()
 {
diff --git a/Engine/Plugins/Experimental/EditorDataStorage/Source/TedsCore/Private/TypedElementDatabaseEnvironment.h b/Engine/Plugins/Experimental/EditorDataStorage/Source/TedsCore/Private/TypedElementDatabaseEnvironment.h
--- a/Engine/Plugins/Experimental/EditorDataStorage/Source/TedsCore/Private/TypedElementDatabaseEnvironment.h
+++ b/Engine/Plugins/Experimental/EditorDataStorage/Source/TedsCore/Private/TypedElementDatabaseEnvironment.h
@@ -2,6 +2,7 @@
 
 #pragma once
 
+#include "DynamicColumnGenerator.h"
 #include "MassEntityManager.h"
 #include "MassProcessingPhaseManager.h"
 #include "TypedElementDatabaseCommandBuffer.h"
@@ -41,6 +42,9 @@ public:
 	FMassProcessingPhaseManager& GetMassPhaseManager();
 	const FMassProcessingPhaseManager& GetMassPhaseManager() const;
 
+	FConstSharedStruct GenerateDynamicTag(const UE::EditorDataStorage::FDynamicTag& Tag, const FName& Value);
+	const UScriptStruct* GenerateColumnType(const UE::EditorDataStorage::FDynamicTag& Tag);
+
 	void NextUpdateCycle();
 	uint64 GetUpdateCycleId() const;
 
@@ -51,6 +55,8 @@ private:
 	FTypedElementDatabaseScratchBuffer ScratchBuffer;
 	FTypedElementExtendedQueryStore Queries;
 	UTypedElementMementoSystem MementoSystem;
+	UE::EditorDataStorage::FDynamicColumnGenerator DynamicColumnGenerator;
+	UE::EditorDataStorage::FDynamicTagManager DynamicTagManager;
 
 	FMassEntityManager& MassEntityManager;
 	FMassProcessingPhaseManager& MassPhaseManager;
diff --git a/Engine/Plugins/Experimental/EditorDataStorage/Source/TedsCore/Private/TypedElementsDataStorage.cpp b/Engine/Plugins/Experimental/EditorDataStorage/Source/TedsCore/Private/TypedElementsDataStorage.cpp
--- a/Engine/Plugins/Experimental/EditorDataStorage/Source/TedsCore/Private/TypedElementsDataStorage.cpp
+++ b/Engine/Plugins/Experimental/EditorDataStorage/Source/TedsCore/Private/TypedElementsDataStorage.cpp
@@ -13,6 +13,7 @@
 #include "TypedElementDatabase.h"
 #include "TypedElementDatabaseCompatibility.h"
 #include "TypedElementDatabaseUI.h"
+#include "TypedElementDataStorageSharedColumn.h"
 #include "UObject/UObjectGlobals.h"
 
 #define LOCTEXT_NAMESPACE "FTypedElementsDataStorageModule"
@@ -47,6 +48,8 @@ void ImpersonateMassTagsAndFragments()
 	static_assert(!TIsPolymorphic<FTypedElementDataStorageTag>::Value,
 		"In order to be able to use FTypedElementDataStorageTag to impersonate FMassTag it can't have any virtual functions.");
 	FTypedElementDataStorageTag::StaticStruct()->SetSuperStruct(FMassTag::StaticStruct());
+
+	FTedsSharedColumn::StaticStruct()->SetSuperStruct(FMassConstSharedFragment::StaticStruct());
 }
 
 void FTypedElementsDataStorageModule::StartupModule()
diff --git a/Engine/Source/Runtime/MassEntity/Public/MassExecutionContext.h b/Engine/Source/Runtime/MassEntity/Public/MassExecutionContext.h
--- a/Engine/Source/Runtime/MassEntity/Public/MassExecutionContext.h
+++ b/Engine/Source/Runtime/MassEntity/Public/MassExecutionContext.h
@@ -195,6 +195,12 @@ public:
 	}
 
 	/** Shared fragment related operations */
+	const void* GetConstSharedFragmentPtr(const UScriptStruct& SharedFragmentType) const
+	{
+		const FConstSharedFragmentView* FoundSharedFragmentData = ConstSharedFragmentViews.FindByPredicate([&SharedFragmentType](const FConstSharedFragmentView& Element) { return Element.Requirement.StructType == &SharedFragmentType; });
+		return FoundSharedFragmentData ? FoundSharedFragmentData->FragmentView.GetMemory() : nullptr;
+	}
+	
 	template<typename T>
 	const T* GetConstSharedFragmentPtr() const
 	{
diff --git a/Engine/Source/Runtime/TypedElementFramework/Private/Elements/Framework/TypedElementQueryBuilder.cpp b/Engine/Source/Runtime/TypedElementFramework/Private/Elements/Framework/TypedElementQueryBuilder.cpp
--- a/Engine/Source/Runtime/TypedElementFramework/Private/Elements/Framework/TypedElementQueryBuilder.cpp
+++ b/Engine/Source/Runtime/TypedElementFramework/Private/Elements/Framework/TypedElementQueryBuilder.cpp
@@ -3,6 +3,7 @@
 #include "Elements/Framework/TypedElementQueryBuilder.h"
 
 #include "Algo/BinarySearch.h"
+#include "Elements/Common/TypedElementDataStorageLog.h"
 #include "Elements/Framework/TypedElementMetaData.h"
 #include "GenericPlatform/GenericPlatformMath.h"
 
@@ -138,6 +139,52 @@ namespace TypedElementQueryBuilder
 		return *this;
 	}
 
+	FSimpleQuery& FSimpleQuery::All(const UE::EditorDataStorage::FDynamicTag& Tag, const FName& Value)
+	{
+		Query->DynamicTags.Emplace(
+			TypedElementDataStorage::FQueryDescription::FDynamicTagData
+			{
+				.Tag = Tag,
+				.MatchValue = Value
+			});
+		return *this;
+	}
+
+	FSimpleQuery& FSimpleQuery::All(const UEnum& Enum)
+	{
+		using namespace UE::EditorDataStorage;
+		return All(FDynamicTag(Enum.GetFName()));
+	}
+
+	FSimpleQuery& FSimpleQuery::All(const UEnum& Enum, int64 Value)
+	{
+		const FName ValueName = Enum.GetNameByValue(Value);
+		if (ValueName == NAME_None)
+		{
+			UE_LOG(LogTypedElementDataStorage, Warning, TEXT("Invalid value '%lld' for enum '%s'"), Value, *Enum.GetName());
+			return *this;
+		}
+		using namespace UE::EditorDataStorage;
+		return All(FDynamicTag(Enum.GetFName()), ValueName);
+	}
+
+	template <>
+	FSimpleQuery& FSimpleQuery::All<UE::EditorDataStorage::FDynamicTag>(const FName& Tag)
+	{
+		return All(UE::EditorDataStorage::FDynamicTag(Tag));
+	}
+
+	template <>
+	FSimpleQuery& FSimpleQuery::All<UE::EditorDataStorage::FDynamicTag>(const FName& Tag, const FName& Value)
+	{
+		return All(UE::EditorDataStorage::FDynamicTag(Tag), Value);
+	}
+
+	FSimpleQuery& FSimpleQuery::All(const UE::EditorDataStorage::FDynamicTag& Tag)
+	{
+		return All(Tag, NAME_None);
+	}
+
 	FSimpleQuery& FSimpleQuery::Any(const UScriptStruct* Target)
 	{
 		if (Target)
diff --git a/Engine/Source/Runtime/TypedElementFramework/Public/Elements/Common/TypedElementCommonTypes.h b/Engine/Source/Runtime/TypedElementFramework/Public/Elements/Common/TypedElementCommonTypes.h
--- a/Engine/Source/Runtime/TypedElementFramework/Public/Elements/Common/TypedElementCommonTypes.h
+++ b/Engine/Source/Runtime/TypedElementFramework/Public/Elements/Common/TypedElementCommonTypes.h
@@ -43,6 +43,51 @@ namespace UE
 	concept derived_from = std::is_base_of_v<Base, Derived> && std::is_convertible_v<const volatile Derived*, const volatile Base*>;
 }
 
+namespace UE
+{
+	namespace EditorDataStorage
+	{
+		/**
+		 * Defines a dynamic type for a dynamic tag
+		 * Example:
+		 *   FDynamicTag ColorTagType(TEXT("Color"));
+		 *   FDynamicTag DirectionTagType(TEXT("Direction"));
+		 * A dynamic tag can take on different values for each type.  This is set up when a tag is added to a row.
+		 */
+		class FDynamicTag
+		{
+		public:
+			explicit FDynamicTag(const FName& InName);
+			
+			const FName& GetName() const;
+			bool operator==(const FDynamicTag& Other) const;
+		private:
+			friend uint32 GetTypeHash(const FDynamicTag& InName)
+			{
+				return GetTypeHash(InName.Name);
+			}
+			FName Name;
+		};
+
+		inline FDynamicTag::FDynamicTag(const FName& InTypeName)
+			: Name(InTypeName)
+		{}
+
+		inline const FName& FDynamicTag::GetName() const
+		{
+			return Name;
+		}
+
+		inline bool FDynamicTag::operator==(const FDynamicTag& Other) const
+		{
+			return Other.Name == Name;
+		}
+
+	}
+}
+
+
+
 namespace TypedElementDataStorage
 {
 	// Standard callbacks.
@@ -64,4 +109,7 @@ namespace TypedElementDataStorage
 	template<typename T>
 	concept TColumnType = TDataColumnType<T> || TTagColumnType<T>;
 
+	template<typename T>
+	concept TEnumType = std::is_enum_v<T>;
+
 } // namespace TypedElementDataStorage
\ No newline at end of file
diff --git a/Engine/Source/Runtime/TypedElementFramework/Public/Elements/Common/TypedElementQueryDescription.h b/Engine/Source/Runtime/TypedElementFramework/Public/Elements/Common/TypedElementQueryDescription.h
--- a/Engine/Source/Runtime/TypedElementFramework/Public/Elements/Common/TypedElementQueryDescription.h
+++ b/Engine/Source/Runtime/TypedElementFramework/Public/Elements/Common/TypedElementQueryDescription.h
@@ -5,6 +5,7 @@
 #include "Containers/Array.h"
 #include "Containers/ContainerAllocationPolicies.h"
 #include "Elements/Common/TypedElementHandles.h"
+#include "Elements/Common/TypedElementCommonTypes.h"
 #include "Elements/Common/TypedElementQueryTypes.h"
 #include "Elements/Interfaces/TypedElementQueryStorageInterfaces.h"
 #include "Elements/Framework/TypedElementMetaData.h"
@@ -55,6 +56,16 @@ namespace TypedElementDataStorage
 			TWeakObjectPtr<const UScriptStruct> Type;
 		};
 
+		struct FDynamicTagData
+		{
+			// The Tag maps to a Mass ConstSharedFragment object
+			UE::EditorDataStorage::FDynamicTag Tag;
+
+			// The MatchValue specifies the value that the fragment must have to be matched
+			// If MatchValue is NAME_None, then TEDS will match all values
+			FName MatchValue;
+		};
+
 		struct FCallbackData
 		{
 			TArray<FName, TInlineAllocator<NumInlineGroups>> BeforeGroups;
@@ -86,6 +97,8 @@ namespace TypedElementDataStorage
 		TArray<EOperatorType, TInlineAllocator<NumInlineConditions>> ConditionTypes;
 		TArray<FOperator, TInlineAllocator<NumInlineConditions>> ConditionOperators;
 
+		TArray<FDynamicTagData> DynamicTags;
+
 		TArray<TWeakObjectPtr<const UClass>, TInlineAllocator<NumInlineDependencies>> DependencyTypes;
 		TArray<EQueryDependencyFlags, TInlineAllocator<NumInlineDependencies>> DependencyFlags;
 		/** Cached instances of the dependencies. This will always match the count of the other Dependency*Types, but may contain null pointers. */
diff --git a/Engine/Source/Runtime/TypedElementFramework/Public/Elements/Framework/TypedElementQueryBuilder.h b/Engine/Source/Runtime/TypedElementFramework/Public/Elements/Framework/TypedElementQueryBuilder.h
--- a/Engine/Source/Runtime/TypedElementFramework/Public/Elements/Framework/TypedElementQueryBuilder.h
+++ b/Engine/Source/Runtime/TypedElementFramework/Public/Elements/Framework/TypedElementQueryBuilder.h
@@ -156,6 +156,7 @@ namespace TypedElementQueryBuilder
 		FSimpleQuery& All();
 		TYPEDELEMENTFRAMEWORK_API FSimpleQuery& All(const UScriptStruct* Target);
 		TYPEDELEMENTFRAMEWORK_API FSimpleQuery& All(TConstArrayView<const UScriptStruct*> Targets);
+		
 		template<TypedElementDataStorage::TColumnType... TargetTypes>
 		FSimpleQuery& Any();
 		TYPEDELEMENTFRAMEWORK_API FSimpleQuery& Any(const UScriptStruct* Target);
@@ -164,6 +165,41 @@ namespace TypedElementQueryBuilder
 		FSimpleQuery& None();
 		TYPEDELEMENTFRAMEWORK_API FSimpleQuery& None(const UScriptStruct* Target);
 		TYPEDELEMENTFRAMEWORK_API FSimpleQuery& None(TConstArrayView<const UScriptStruct*> Targets);
+		
+		// Dynamic Tags
+		// ============
+		// Adds a filter to the query which must match a given DynamicTag.  The value of the DynamicTag will not be checked.
+		TYPEDELEMENTFRAMEWORK_API FSimpleQuery& All(const UE::EditorDataStorage::FDynamicTag& Tag);
+		// Adds a filter to the query which must match a given DynamicTag.  The value of the DynamicTag must also match.
+		// Note: The query can only match a single value.  Multiple value queries are not supported at this time.
+		TYPEDELEMENTFRAMEWORK_API FSimpleQuery& All(const UE::EditorDataStorage::FDynamicTag& Tag, const FName& Value);
+		
+		TYPEDELEMENTFRAMEWORK_API FSimpleQuery& All(const UEnum& Enum);
+		TYPEDELEMENTFRAMEWORK_API FSimpleQuery& All(const UEnum& Enum, int64 Value);
+
+		// Disabled generic template for DynamicTags
+		template<typename... T>
+		FSimpleQuery& All(const FName&) = delete;
+		// Disabled generic template for DynamicTags
+		template<typename... T>
+		FSimpleQuery& All(const FName&, const FName&) = delete;
+
+		// Adds a filter to the query which must match a given DynamicTag.  The value of the DynamicTag will not be checked.
+		template<>
+		TYPEDELEMENTFRAMEWORK_API FSimpleQuery& All<UE::EditorDataStorage::FDynamicTag>(const FName& Tag);
+		// Adds a filter to the query which must match a given DynamicTag.  The value of the DynamicTag must also match.
+		// Note: The query can only match a single value.  Multiple value queries are not supported at this time.
+		template<>
+		TYPEDELEMENTFRAMEWORK_API FSimpleQuery& All<UE::EditorDataStorage::FDynamicTag>(const FName& Tag, const FName& Value);
+		
+		template<TypedElementDataStorage::TEnumType EnumT>
+		FSimpleQuery& All();
+		
+		template<TypedElementDataStorage::TEnumType EnumT>
+		FSimpleQuery& All(EnumT EnumValue);
+
+		template<auto Value, TypedElementDataStorage::TEnumType EnumT = decltype(Value)>
+		FSimpleQuery& All();
 
 	private:
 		TYPEDELEMENTFRAMEWORK_API explicit FSimpleQuery(ITypedElementDataStorageInterface::FQueryDescription* Query);
diff --git a/Engine/Source/Runtime/TypedElementFramework/Public/Elements/Framework/TypedElementQueryBuilder.inl b/Engine/Source/Runtime/TypedElementFramework/Public/Elements/Framework/TypedElementQueryBuilder.inl
--- a/Engine/Source/Runtime/TypedElementFramework/Public/Elements/Framework/TypedElementQueryBuilder.inl
+++ b/Engine/Source/Runtime/TypedElementFramework/Public/Elements/Framework/TypedElementQueryBuilder.inl
@@ -1163,8 +1163,24 @@ e.g. void(FCachedQueryContext<Subsystem1, const Subsystem2>& Context, TypedEleme
 		None({ TargetTypes::StaticStruct()... });
 		return *this;
 	}
-
-
+	
+	template <TypedElementDataStorage::TEnumType EnumT>
+    FSimpleQuery& FSimpleQuery::All()
+    {
+    	return All(static_cast<const UEnum&>(*StaticEnum<EnumT>()));
+    }
+
+    template <TypedElementDataStorage::TEnumType EnumT>
+    FSimpleQuery& FSimpleQuery::All(EnumT EnumValue)
+    {
+    	return All(static_cast<const UEnum&>(*StaticEnum<EnumT>()), static_cast<int64>(EnumValue));
+    }
+	
+	template<auto Value, TypedElementDataStorage::TEnumType EnumT = decltype(Value)>
+	FSimpleQuery& All()
+	{
+		return All<EnumT>(Value);
+	}
 
 	//
 	// External query bindings
diff --git a/Engine/Source/Runtime/TypedElementFramework/Public/Elements/Interfaces/TypedElementDataStorageInterface.h b/Engine/Source/Runtime/TypedElementFramework/Public/Elements/Interfaces/TypedElementDataStorageInterface.h
--- a/Engine/Source/Runtime/TypedElementFramework/Public/Elements/Interfaces/TypedElementDataStorageInterface.h
+++ b/Engine/Source/Runtime/TypedElementFramework/Public/Elements/Interfaces/TypedElementDataStorageInterface.h
@@ -184,6 +184,31 @@ public:
 	template<TypedElementDataStorage::TDataColumnType ColumnType>
 	void AddColumn(TypedElementRowHandle Row, ColumnType&& Column);
 
+	/**
+	 * Adds a DynamicTag with the given value to a row
+	 * A row can have multiple DynamicTags, but only one of each tag type.
+	 * Example:
+	 *   AddColumn(Row, FDynamicTag(TEXT("Color"), TEXT("Red));     // Valid
+	 *   AddColumn(Row, FDynamicTag(TEXT("Direction"), TEXT("Up")); // Valid
+	 *   AddColumn(Row, FDynamicTag(TEXT("Color"), TEXT("Blue"));   // Will do nothing since there already exists a Color dynamic tag
+	 * Note: Current support for changing a dynamic tag from one value to another requires that the tag is removed before a new one
+	 *       is added.  This will likely change in the future to transparently replace the tag to have consistent behaviour with other usages
+	 *       of AddColumn
+	 */
+	virtual void AddColumn(TypedElementDataStorage::RowHandle Row, const UE::EditorDataStorage::FDynamicTag& Tag, const FName& Value) = 0;
+
+	template<typename T>
+	void AddColumn(TypedElementDataStorage::RowHandle Row, const FName& Tag, const FName& Value) = delete;
+	
+	template<>
+	void AddColumn<UE::EditorDataStorage::FDynamicTag>(TypedElementDataStorage::RowHandle Row, const FName& Tag, const FName& Value);
+
+	template<TypedElementDataStorage::TEnumType EnumT>
+	void AddColumn(TypedElementDataStorage::RowHandle Row, EnumT Value);
+	
+	template<auto Value, TypedElementDataStorage::TEnumType EnumT = decltype(Value)>
+	void AddColumn(TypedElementDataStorage::RowHandle Row);
+
 	/**
 	 * Adds multiple columns from a row. This is typically more efficient than adding columns one 
 	 * at a time.
@@ -196,6 +221,22 @@ public:
 	virtual void RemoveColumn(TypedElementRowHandle Row, const UScriptStruct* ColumnType) = 0;
 	template<TypedElementDataStorage::TColumnType Column>
 	void RemoveColumn(TypedElementRowHandle Row);
+
+	template<TypedElementDataStorage::TEnumType EnumT>
+	void RemoveColumn(TypedElementDataStorage::RowHandle Row);
+
+	/**
+	 * Removes a dynamic tag from the given row
+	 * If tag does not exist on row, operation will do nothing.
+	 */
+	virtual void RemoveColumn(TypedElementDataStorage::RowHandle Row, const UE::EditorDataStorage::FDynamicTag& Tag) = 0;
+
+	template<typename T>
+	void RemoveColumn(TypedElementDataStorage::RowHandle Row, const FName& Tag) = delete;
+	
+	template<>
+	void RemoveColumn<UE::EditorDataStorage::FDynamicTag>(TypedElementDataStorage::RowHandle Row, const FName& Tag);
+
 	/**
 	 * Removes multiple columns from a row. This is typically more efficient than adding columns one
 	 * at a time.
@@ -236,7 +277,6 @@ public:
 	/** Determines if the columns in the row match the query conditions. */
 	virtual bool MatchesColumns(TypedElementDataStorage::RowHandle Row, const TypedElementDataStorage::FQueryConditions& Conditions) const = 0;
 	
-
 	/**
 	 * @section Query
 	 * @description
@@ -355,8 +395,6 @@ public:
 	virtual void ListExtensions(TFunctionRef<void(FName)> Callback) const = 0;
 };
 
-
-
 // Implementations
 
 template <typename FactoryT>
@@ -396,6 +434,43 @@ void ITypedElementDataStorageInterface::AddColumns(TypedElementRowHandle Row)
 	AddColumns(Row, { Columns::StaticStruct()...});
 }
 
+template <>
+inline void ITypedElementDataStorageInterface::AddColumn<UE::EditorDataStorage::FDynamicTag>(TypedElementDataStorage::RowHandle Row, const FName& Tag, const FName& Value)
+{
+	AddColumn(Row, UE::EditorDataStorage::FDynamicTag(Tag), Value);
+}
+
+template <>
+inline void ITypedElementDataStorageInterface::RemoveColumn<UE::EditorDataStorage::FDynamicTag>(TypedElementDataStorage::RowHandle Row, const FName& Tag)
+{
+	using namespace UE::EditorDataStorage;
+	RemoveColumn(Row, FDynamicTag(Tag));
+}
+
+template<TypedElementDataStorage::TEnumType EnumT>
+void ITypedElementDataStorageInterface::AddColumn(TypedElementDataStorage::RowHandle Row, EnumT Value)
+{
+	const UEnum* Enum = StaticEnum<EnumT>();
+	const FName ValueAsFName = *Enum->GetNameStringByValue(static_cast<int64>(Value));
+	if (ValueAsFName != NAME_None)
+	{
+		AddColumn(Row, UE::EditorDataStorage::FDynamicTag(Enum->GetFName()), ValueAsFName);
+	}
+}
+
+template<auto Value, TypedElementDataStorage::TEnumType EnumT>
+void ITypedElementDataStorageInterface::AddColumn(TypedElementDataStorage::RowHandle Row)
+{
+	AddColumn<EnumT>(Row, Value);
+}
+
+template<TypedElementDataStorage::TEnumType EnumT>
+void ITypedElementDataStorageInterface::RemoveColumn(TypedElementDataStorage::RowHandle Row)
+{
+	const UEnum* Enum = StaticEnum<EnumT>();
+	RemoveColumn(Row, UE::EditorDataStorage::FDynamicTag(Enum->GetFName()));
+}
+
 template<TypedElementDataStorage::TColumnType... Columns>
 void ITypedElementDataStorageInterface::RemoveColumns(TypedElementRowHandle Row)
 {
