Added the ability to TEDS Core to retrieve the column types, and optional the data, on a row.

diff --git a/Engine/Plugins/Experimental/EditorDataStorage/Source/TedsCore/Private/TypedElementDatabase.cpp b/Engine/Plugins/Experimental/EditorDataStorage/Source/TedsCore/Private/TypedElementDatabase.cpp
--- a/Engine/Plugins/Experimental/EditorDataStorage/Source/TedsCore/Private/TypedElementDatabase.cpp
+++ b/Engine/Plugins/Experimental/EditorDataStorage/Source/TedsCore/Private/TypedElementDatabase.cpp
@@ -852,6 +852,64 @@ bool UTypedElementDatabase::HasColumns(TypedElementRowHandle Row, TConstArrayVie
 	return false;
 }
 
+void UTypedElementDatabase::ListColumns(TypedElementDataStorage::RowHandle Row, TypedElementDataStorage::ColumnListCallbackRef Callback) const
+{
+	if (ActiveEditorEntityManager)
+	{
+		FMassEntityHandle Entity = FMassEntityHandle::FromNumber(Row);
+		if (ActiveEditorEntityManager->IsEntityActive(Entity))
+		{
+			FMassArchetypeHandle Archetype = ActiveEditorEntityManager->GetArchetypeForEntity(Entity);
+			const FMassArchetypeCompositionDescriptor& Composition = ActiveEditorEntityManager->GetArchetypeComposition(Archetype);
+			
+			auto CallbackWrapper = [&Callback](const UScriptStruct* ColumnType)
+				{
+					if (ColumnType)
+					{
+						Callback(*ColumnType);
+					}
+					return true;
+				};
+			Composition.Fragments.ExportTypes(CallbackWrapper);
+			Composition.Tags.ExportTypes(CallbackWrapper);
+
+		}
+	}
+}
+
+void UTypedElementDatabase::ListColumns(TypedElementDataStorage::RowHandle Row, TypedElementDataStorage::ColumnListWithDataCallbackRef Callback)
+{
+	if (ActiveEditorEntityManager)
+	{
+		FMassEntityHandle Entity = FMassEntityHandle::FromNumber(Row);
+		if (ActiveEditorEntityManager->IsEntityActive(Entity))
+		{
+			FMassArchetypeHandle Archetype = ActiveEditorEntityManager->GetArchetypeForEntity(Entity);
+			const FMassArchetypeCompositionDescriptor& Composition = ActiveEditorEntityManager->GetArchetypeComposition(Archetype);
+
+			Composition.Fragments.ExportTypes(
+				[this, &Callback, Entity](const UScriptStruct* ColumnType)
+				{
+					if (ColumnType)
+					{
+						Callback(ActiveEditorEntityManager->GetFragmentDataStruct(Entity, ColumnType).GetMemory(), *ColumnType);
+					}
+					return true;
+				});
+			Composition.Tags.ExportTypes(
+				[&Callback](const UScriptStruct* ColumnType)
+				{
+					if (ColumnType)
+					{
+						Callback(nullptr, *ColumnType);
+					}
+					return true;
+				});
+
+		}
+	}
+}
+
 bool UTypedElementDatabase::MatchesColumns(TypedElementDataStorage::RowHandle Row, const TypedElementDataStorage::FQueryConditions& Conditions) const
 {
 	if (ActiveEditorEntityManager)
diff --git a/Engine/Plugins/Experimental/EditorDataStorage/Source/TedsCore/Private/TypedElementDatabase.h b/Engine/Plugins/Experimental/EditorDataStorage/Source/TedsCore/Private/TypedElementDatabase.h
--- a/Engine/Plugins/Experimental/EditorDataStorage/Source/TedsCore/Private/TypedElementDatabase.h
+++ b/Engine/Plugins/Experimental/EditorDataStorage/Source/TedsCore/Private/TypedElementDatabase.h
@@ -120,6 +120,8 @@ public:
 		TConstArrayView<const UScriptStruct*> ColumnsToRemove) override;
 	bool HasColumns(TypedElementRowHandle Row, TConstArrayView<const UScriptStruct*> ColumnTypes) const override;
 	bool HasColumns(TypedElementRowHandle Row, TConstArrayView<TWeakObjectPtr<const UScriptStruct>> ColumnTypes) const override;
+	void ListColumns(TypedElementDataStorage::RowHandle Row, TypedElementDataStorage::ColumnListCallbackRef Callback) const;
+	void ListColumns(TypedElementDataStorage::RowHandle Row, TypedElementDataStorage::ColumnListWithDataCallbackRef Callback);
 	bool MatchesColumns(TypedElementDataStorage::RowHandle Row, const TypedElementDataStorage::FQueryConditions& Conditions) const override;
 
 	void RegisterTickGroup(FName GroupName, EQueryTickPhase Phase, FName BeforeGroup, FName AfterGroup, bool bRequiresMainThread);
diff --git a/Engine/Source/Runtime/TypedElementFramework/Public/Elements/Common/TypedElementCommonTypes.h b/Engine/Source/Runtime/TypedElementFramework/Public/Elements/Common/TypedElementCommonTypes.h
--- a/Engine/Source/Runtime/TypedElementFramework/Public/Elements/Common/TypedElementCommonTypes.h
+++ b/Engine/Source/Runtime/TypedElementFramework/Public/Elements/Common/TypedElementCommonTypes.h
@@ -94,6 +94,8 @@ namespace TypedElementDataStorage
 
 	using RowCreationCallbackRef = TFunctionRef<void(TypedElementDataStorage::RowHandle Row)>;
 	using ColumnCreationCallbackRef = TFunctionRef<void(void* Column, const UScriptStruct& ColumnType)>;
+	using ColumnListCallbackRef = TFunctionRef<void(const UScriptStruct& ColumnType)>;
+	using ColumnListWithDataCallbackRef = TFunctionRef<void(void* Column, const UScriptStruct& ColumnType)>;
 	using ColumnCopyOrMoveCallback = void (*)(const UScriptStruct& ColumnType, void* Destination, void* Source);
 
 
diff --git a/Engine/Source/Runtime/TypedElementFramework/Public/Elements/Interfaces/TypedElementDataStorageInterface.h b/Engine/Source/Runtime/TypedElementFramework/Public/Elements/Interfaces/TypedElementDataStorageInterface.h
--- a/Engine/Source/Runtime/TypedElementFramework/Public/Elements/Interfaces/TypedElementDataStorageInterface.h
+++ b/Engine/Source/Runtime/TypedElementFramework/Public/Elements/Interfaces/TypedElementDataStorageInterface.h
@@ -274,6 +274,15 @@ public:
 	template<TypedElementDataStorage::TColumnType... ColumnTypes>
 	bool HasColumns(TypedElementRowHandle Row) const;
 
+	/** Lists the columns on a row. This includes data and tag columns. */
+	virtual void ListColumns(TypedElementDataStorage::RowHandle Row, TypedElementDataStorage::ColumnListCallbackRef Callback) const = 0;
+
+	/** 
+	 * Lists the column type and data on a row. This includes data and tag columns. Not all columns may have data so the data pointer in 
+	 * the callback can be null.
+	 */
+	virtual void ListColumns(TypedElementDataStorage::RowHandle Row, TypedElementDataStorage::ColumnListWithDataCallbackRef Callback) = 0;
+
 	/** Determines if the columns in the row match the query conditions. */
 	virtual bool MatchesColumns(TypedElementDataStorage::RowHandle Row, const TypedElementDataStorage::FQueryConditions& Conditions) const = 0;
 	
