Allow TEDS UI widget constructors to provide a name for the column(s) they represent.

diff --git a/Engine/Plugins/Experimental/EditorDataStorage/Source/TedsUI/Private/Widgets/ExportedTextWidget.cpp b/Engine/Plugins/Experimental/EditorDataStorage/Source/TedsUI/Private/Widgets/ExportedTextWidget.cpp
--- a/Engine/Plugins/Experimental/EditorDataStorage/Source/TedsUI/Private/Widgets/ExportedTextWidget.cpp
+++ b/Engine/Plugins/Experimental/EditorDataStorage/Source/TedsUI/Private/Widgets/ExportedTextWidget.cpp
@@ -14,8 +14,7 @@
 // UExportedTextWidgetFactory
 //
 
-static void UpdateExportedTextWidget(const void* Data, FTypedElementSlateWidgetReferenceColumn& Widget, 
-	const UScriptStruct* StructType)
+static void UpdateExportedTextWidget(FText Text, FTypedElementSlateWidgetReferenceColumn& Widget)
 {
 	TSharedPtr<SWidget> WidgetPointer = Widget.Widget.Pin();
 	checkf(WidgetPointer, TEXT("Referenced widget is not valid. A constructed widget may not have been cleaned up. This can "
@@ -26,14 +25,19 @@ static void UpdateExportedTextWidget(const void* Data, FTypedElementSlateWidgetR
 		*(STextBlock::StaticWidgetClass().GetWidgetType().ToString()),
 		*(WidgetPointer->GetTypeAsString()));
 
-	FString Label;
-	StructType->ExportText(Label, Data, Data, nullptr, PPF_None, nullptr);
 	STextBlock* TextWidget = static_cast<STextBlock*>(WidgetPointer.Get());
-	FText Text = FText::FromString(MoveTemp(Label));
 	TextWidget->SetToolTipText(Text);
 	TextWidget->SetText(MoveTemp(Text));
 }
 
+static void UpdateExportedTextWidget(const void* Data, FTypedElementSlateWidgetReferenceColumn& Widget, 
+	const UScriptStruct* StructType)
+{
+	FString Label;
+	StructType->ExportText(Label, Data, Data, nullptr, PPF_None, nullptr);
+	UpdateExportedTextWidget(FText::FromString(MoveTemp(Label)), Widget);
+}
+
 static void UpdateExportedTextWidget(ITypedElementDataStorageInterface& DataStorage, FTypedElementSlateWidgetReferenceColumn& Widget,
 	const FTypedElementScriptStructTypeInfoColumn& TypeInfo, const FTypedElementRowReferenceColumn& ReferencedRow)
 {
@@ -136,6 +140,19 @@ const TypedElementDataStorage::FQueryConditions* FExportedTextWidgetConstructor:
 	return &MatchedColumn;
 }
 
+FString FExportedTextWidgetConstructor::CreateWidgetDisplayName(ITypedElementDataStorageInterface* DataStorage,
+	TypedElementDataStorage::RowHandle Row) const
+{
+	if (FTypedElementScriptStructTypeInfoColumn* TypeInfoColumn = DataStorage->GetColumn<FTypedElementScriptStructTypeInfoColumn>(Row))
+	{
+		return DescribeColumnType(TypeInfoColumn->TypeInfo.Get());
+	}
+	else
+	{
+		return DescribeColumnType(nullptr);
+	}
+}
+
 TSharedPtr<SWidget> FExportedTextWidgetConstructor::CreateWidget(const TypedElementDataStorage::FMetaDataView& Arguments)
 {
 	return SNew(STextBlock);
@@ -154,18 +171,27 @@ bool FExportedTextWidgetConstructor::FinalizeWidget(
 	// TEDS UI TODO: We should work around it by refactoring this into an STedsWidget in the future so it can store the column conditions per instance
 	MatchedColumn = TypedElementDataStorage::FQueryConditions(TypedElementDataStorage::FColumn(TypeInfoColumn.TypeInfo));
 
-	UpdateExportedTextWidget(
-		*DataStorage,
-		*DataStorage->GetColumn<FTypedElementSlateWidgetReferenceColumn>(Row),
-		TypeInfoColumn,
-		*DataStorage->GetColumn<FTypedElementRowReferenceColumn>(Row));
-	
-	UExportedTextWidgetFactory* Factory = 
-		UExportedTextWidgetFactory::StaticClass()->GetDefaultObject<UExportedTextWidgetFactory>();
-	if (Factory && !Factory->RegisteredTypes.Contains(TypeInfoColumn.TypeInfo))
+	if (TypeInfoColumn.TypeInfo->IsChildOf(FTypedElementDataStorageTag::StaticStruct()))
+	{
+		UpdateExportedTextWidget(
+			LOCTEXT("ExportedTextWidgetTag", "<Tag>"),
+			*DataStorage->GetColumn<FTypedElementSlateWidgetReferenceColumn>(Row));
+	}
+	else
 	{
-		RegisterUpdateCallback(*DataStorage, TypeInfoColumn.TypeInfo.Get());
-		Factory->RegisteredTypes.Add(TypeInfoColumn.TypeInfo);
+		UpdateExportedTextWidget(
+			*DataStorage,
+			*DataStorage->GetColumn<FTypedElementSlateWidgetReferenceColumn>(Row),
+			TypeInfoColumn,
+			*DataStorage->GetColumn<FTypedElementRowReferenceColumn>(Row));
+
+		UExportedTextWidgetFactory* Factory =
+			UExportedTextWidgetFactory::StaticClass()->GetDefaultObject<UExportedTextWidgetFactory>();
+		if (Factory && !Factory->RegisteredTypes.Contains(TypeInfoColumn.TypeInfo))
+		{
+			RegisterUpdateCallback(*DataStorage, TypeInfoColumn.TypeInfo.Get());
+			Factory->RegisteredTypes.Add(TypeInfoColumn.TypeInfo);
+		}
 	}
 	
 	return true;
diff --git a/Engine/Plugins/Experimental/EditorDataStorage/Source/TedsUI/Public/Widgets/ExportedTextWidget.h b/Engine/Plugins/Experimental/EditorDataStorage/Source/TedsUI/Public/Widgets/ExportedTextWidget.h
--- a/Engine/Plugins/Experimental/EditorDataStorage/Source/TedsUI/Public/Widgets/ExportedTextWidget.h
+++ b/Engine/Plugins/Experimental/EditorDataStorage/Source/TedsUI/Public/Widgets/ExportedTextWidget.h
@@ -36,6 +36,8 @@ public:
 
 	virtual TConstArrayView<const UScriptStruct*> GetAdditionalColumnsList() const override;
 	virtual const TypedElementDataStorage::FQueryConditions* GetQueryConditions() const override;
+	virtual FString CreateWidgetDisplayName(
+		ITypedElementDataStorageInterface* DataStorage, TypedElementDataStorage::RowHandle Row) const override;
 
 protected:
 	virtual TSharedPtr<SWidget> CreateWidget(const TypedElementDataStorage::FMetaDataView& Arguments) override;
diff --git a/Engine/Source/Runtime/TypedElementFramework/Private/Elements/Common/TypedElementQueryConditions.cpp b/Engine/Source/Runtime/TypedElementFramework/Private/Elements/Common/TypedElementQueryConditions.cpp
--- a/Engine/Source/Runtime/TypedElementFramework/Private/Elements/Common/TypedElementQueryConditions.cpp
+++ b/Engine/Source/Runtime/TypedElementFramework/Private/Elements/Common/TypedElementQueryConditions.cpp
@@ -109,11 +109,11 @@ namespace TypedElementDataStorage
 					}
 					return false;
 				});
-				if (Result)
-				{
-					ConvertColumnBitToArray(MatchedColumns, Matches);
-				}
-				return Result;
+		if (Result)
+		{
+			ConvertColumnBitToArray(MatchedColumns, Matches);
+		}
+		return Result;
 	}
 
 	bool FQueryConditions::Verify(TConstArrayView<TWeakObjectPtr<const UScriptStruct>> AvailableColumns, bool AvailableColumnsAreSorted) const
@@ -138,7 +138,7 @@ namespace TypedElementDataStorage
 	bool FQueryConditions::Verify(TArray<TWeakObjectPtr<const UScriptStruct>>& MatchedColumns,
 		TConstArrayView<TWeakObjectPtr<const UScriptStruct>> AvailableColumns, bool AvailableColumnsAreSorted) const
 	{
-		static_assert(MaxColumnCount < 64, "Query conditions use a bit mask to locate matches. As a result MaxColumnCount can be larger than 64.");
+		static_assert(MaxColumnCount < 64, "Query conditions use a bit mask to locate matches. As a result MaxColumnCount cannot be larger than 64.");
 		uint64 Matches = 0;
 		bool Result = AvailableColumnsAreSorted
 			? VerifyBootstrap(
@@ -160,11 +160,11 @@ namespace TypedElementDataStorage
 				{
 					return AvailableColumns.Find(Column) != INDEX_NONE;
 				});
-				if (Result)
-				{
-					ConvertColumnBitToArray(MatchedColumns, Matches);
-				}
-				return Result;
+		if (Result)
+		{
+			ConvertColumnBitToArray(MatchedColumns, Matches);
+		}
+		return Result;
 	}
 
 	bool FQueryConditions::Verify(TSet<TWeakObjectPtr<const UScriptStruct>> AvailableColumns) const
diff --git a/Engine/Source/Runtime/TypedElementFramework/Private/Elements/Interfaces/TypedElementDataStorageUiInterface.cpp b/Engine/Source/Runtime/TypedElementFramework/Private/Elements/Interfaces/TypedElementDataStorageUiInterface.cpp
--- a/Engine/Source/Runtime/TypedElementFramework/Private/Elements/Interfaces/TypedElementDataStorageUiInterface.cpp
+++ b/Engine/Source/Runtime/TypedElementFramework/Private/Elements/Interfaces/TypedElementDataStorageUiInterface.cpp
@@ -41,6 +41,58 @@ TConstArrayView<const UScriptStruct*> FTypedElementWidgetConstructor::GetAdditio
 	return {};
 }
 
+FString FTypedElementWidgetConstructor::CreateWidgetDisplayName(
+	ITypedElementDataStorageInterface* DataStorage, TypedElementDataStorage::RowHandle Row) const
+{
+	switch (MatchedColumnTypes.Num())
+	{
+	case 0:
+		return FString(TEXT("TEDS Column"));
+	case 1:
+		return DescribeColumnType(MatchedColumnTypes[0].Get());
+	default:
+	{
+		FString LongestMatchString = DescribeColumnType(MatchedColumnTypes[0].Get());
+		FStringView LongestMatch = LongestMatchString;
+		const TWeakObjectPtr<const UScriptStruct>* It = MatchedColumnTypes.GetData();
+		const TWeakObjectPtr<const UScriptStruct>* ItEnd = It + MatchedColumnTypes.Num();
+		++It; // Skip the first entry as that's already set.
+		for (; It != ItEnd; ++It)
+		{
+			FString NextMatchText = DescribeColumnType(It->Get());
+			FStringView NextMatch = NextMatchText;
+
+			int32 MatchSize = 0;
+			auto ItLeft = LongestMatch.begin();
+			auto ItLeftEnd = LongestMatch.end();
+			auto ItRight = NextMatch.begin();
+			auto ItRightEnd = NextMatch.end();
+			while (
+				ItLeft != ItLeftEnd &&
+				ItRight != ItRightEnd &&
+				*ItLeft == *ItRight)
+			{
+				++MatchSize;
+				++ItLeft;
+				++ItRight;
+			}
+
+			// At least 3 letters have to match to avoid single or double letter names which typically mean nothing.
+			if (MatchSize > 2)
+			{
+				LongestMatch.LeftInline(MatchSize);
+			}
+			else
+			{
+				// There are not enough characters in the string that match. Just return the name of the first column
+				return LongestMatchString;
+			}
+		}
+		return FString(LongestMatch);
+	}
+	};
+}
+
 TSharedPtr<SWidget> FTypedElementWidgetConstructor::ConstructFinalWidget(
 	TypedElementRowHandle Row,
 	ITypedElementDataStorageInterface* DataStorage,
@@ -122,6 +174,23 @@ FString FTypedElementWidgetConstructor::GetWidgetLabel(const TSharedPtr<SWidget>
 	return Widget->GetType().ToString();
 }
 
+FString FTypedElementWidgetConstructor::DescribeColumnType(const UScriptStruct* ColumnType) const
+{
+	static const FName DisplayNameName(TEXT("DisplayName"));
+
+#if WITH_EDITOR
+	if (ColumnType)
+	{
+		const FString* Name = ColumnType->FindMetaData(DisplayNameName);
+		return Name ? *Name : ColumnType->GetDisplayNameText().ToString();
+	}
+	else
+#endif
+	{
+		return FString(TEXT("<Invalid>"));
+	}
+}
+
 bool FTypedElementWidgetConstructor::FinalizeWidget(
 	ITypedElementDataStorageInterface* DataStorage,
 	ITypedElementDataStorageUiInterface* DataStorageUi,
diff --git a/Engine/Source/Runtime/TypedElementFramework/Public/Elements/Interfaces/TypedElementDataStorageUiInterface.h b/Engine/Source/Runtime/TypedElementFramework/Public/Elements/Interfaces/TypedElementDataStorageUiInterface.h
--- a/Engine/Source/Runtime/TypedElementFramework/Public/Elements/Interfaces/TypedElementDataStorageUiInterface.h
+++ b/Engine/Source/Runtime/TypedElementFramework/Public/Elements/Interfaces/TypedElementDataStorageUiInterface.h
@@ -44,6 +44,15 @@ public:
 
 	/** Returns a list of additional columns the widget requires to be added to its rows. */
 	TYPEDELEMENTFRAMEWORK_API virtual TConstArrayView<const UScriptStruct*> GetAdditionalColumnsList() const;
+
+	/** 
+	 * Returns a friendly name for the data the created widget represents.
+	 * By default the associated column is used. If there are multiple columns associated with the constructor
+	 * the default implementation will attempt to find the longest common starting string for all the columns.
+	 * Individual widget constructors can override this function with a name specific to them.
+	 */
+	TYPEDELEMENTFRAMEWORK_API virtual FString CreateWidgetDisplayName(
+		ITypedElementDataStorageInterface* DataStorage, TypedElementDataStorage::RowHandle Row) const;
 	
 	/**
 	 *	Calls Construct() to create the internal widget, and then stores it in a container before returning.
@@ -81,6 +90,9 @@ protected:
 	/** Get a user facing label for the widget used in table views. */
 	TYPEDELEMENTFRAMEWORK_API virtual FString GetWidgetLabel(const TSharedPtr<SWidget>& Widget);
 
+	/** Creates a (friendly) name for the provided column type. */
+	TYPEDELEMENTFRAMEWORK_API virtual FString DescribeColumnType(const UScriptStruct* ColumnType) const;
+
 	/** 
 	 * Last opportunity to configure anything in the widget or the row. This step can be needed to initialize widgets with data stored
 	 * in columns.
