[Teds] Rename native classes and move into UE::EditorDataStorage namespaces

Various other code style consistency changes

#rnx

diff --git a/Engine/Plugins/Experimental/EditorDataStorage/Source/TedsCore/Private/Commands/EditorDataStorageCommandBuffer.h b/Engine/Plugins/Experimental/EditorDataStorage/Source/TedsCore/Private/Commands/EditorDataStorageCommandBuffer.h
--- a/Engine/Plugins/Experimental/EditorDataStorage/Source/TedsCore/Private/Commands/EditorDataStorageCommandBuffer.h
+++ b/Engine/Plugins/Experimental/EditorDataStorage/Source/TedsCore/Private/Commands/EditorDataStorageCommandBuffer.h
@@ -7,7 +7,7 @@
 #include "Templates/Function.h"
 #include "TypedElementDatabaseScratchBuffer.h"
 
-namespace UE::EditorDataStorage
+namespace UE::Editor::DataStorage
 {
 	/** Empty command that does nothing. Can be used as placeholder or to disable an existing command. */
 	struct FNopCommand final {};
@@ -96,7 +96,7 @@ namespace UE::EditorDataStorage
 		private:
 			TArray<TCommandVariant*> CommandReferences;
 			uint32 CommandCount[sizeof...(TCommand) + 1];
-			FTypedElementDatabaseScratchBuffer* ScratchBuffer = nullptr;
+			FScratchBuffer* ScratchBuffer = nullptr;
 		};
 
 		/** Utility class to help optimize a collection of commands. */
@@ -200,7 +200,7 @@ namespace UE::EditorDataStorage
 			FCommandBuffer::FCollection& Commands;
 		};
 
-		void Initialize(FTypedElementDatabaseScratchBuffer& InScratchBuffer);
+		void Initialize(FScratchBuffer& InScratchBuffer);
 
 		/** Adds a command the command buffer in a thread-safe manner. T has to be one of the provided commands or TNopCommand. */
 		template<typename T>
@@ -231,8 +231,8 @@ namespace UE::EditorDataStorage
 		FCommandInstance* CommandFront = nullptr;
 		std::atomic<FCommandInstance*> CommandBack = nullptr;
 
-		FTypedElementDatabaseScratchBuffer* ScratchBuffer = nullptr;
+		FScratchBuffer* ScratchBuffer = nullptr;
 	};
-} // namespace UE::EditorDataStorage
+} // namespace UE::Editor::DataStorage
 
 #include "Commands/EditorDataStorageCommandBuffer.inl"
\ No newline at end of file
diff --git a/Engine/Plugins/Experimental/EditorDataStorage/Source/TedsCore/Private/Commands/EditorDataStorageCommandBuffer.inl b/Engine/Plugins/Experimental/EditorDataStorage/Source/TedsCore/Private/Commands/EditorDataStorageCommandBuffer.inl
--- a/Engine/Plugins/Experimental/EditorDataStorage/Source/TedsCore/Private/Commands/EditorDataStorageCommandBuffer.inl
+++ b/Engine/Plugins/Experimental/EditorDataStorage/Source/TedsCore/Private/Commands/EditorDataStorageCommandBuffer.inl
@@ -3,7 +3,7 @@
 #include "Algo/BinarySearch.h"
 #include "HAL/UnrealMemory.h"
 
-namespace UE::EditorDataStorage
+namespace UE::Editor::DataStorage
 {
 	//
 	// FCommandBuffer::FCollection
@@ -381,7 +381,7 @@ namespace UE::EditorDataStorage
 	//
 
 	template<typename... TCommand>
-	void FCommandBuffer<TCommand...>::Initialize(FTypedElementDatabaseScratchBuffer& InScratchBuffer)
+	void FCommandBuffer<TCommand...>::Initialize(FScratchBuffer& InScratchBuffer)
 	{
 		ScratchBuffer = &InScratchBuffer;
 		CommandFront = ScratchBuffer->Emplace<FCommandInstance>();
@@ -440,4 +440,4 @@ namespace UE::EditorDataStorage
 
 		return Result;
 	}
-} // namespace UE::EditorDataStorage
\ No newline at end of file
+} // namespace UE::Editor::DataStorage
\ No newline at end of file
diff --git a/Engine/Plugins/Experimental/EditorDataStorage/Source/TedsCore/Private/Commands/EditorDataStorageCompatibilityCommands.cpp b/Engine/Plugins/Experimental/EditorDataStorage/Source/TedsCore/Private/Commands/EditorDataStorageCompatibilityCommands.cpp
--- a/Engine/Plugins/Experimental/EditorDataStorage/Source/TedsCore/Private/Commands/EditorDataStorageCompatibilityCommands.cpp
+++ b/Engine/Plugins/Experimental/EditorDataStorage/Source/TedsCore/Private/Commands/EditorDataStorageCompatibilityCommands.cpp
@@ -14,7 +14,7 @@
 #include "TypedElementDatabaseEnvironment.h"
 #include "TypedElementDatabaseScratchBuffer.h"
 
-namespace UE::EditorDataStorage
+namespace UE::Editor::DataStorage
 {
 	UScriptStruct* FAddSyncFromWorldTag::GetType()
 	{
@@ -154,7 +154,7 @@ namespace UE::EditorDataStorage
 	}
 
 	void FPatchData::RunPatch(CompatibilityCommandBuffer::FCollection& Commands, UTypedElementDatabaseCompatibility& StorageCompat,
-		FTypedElementDatabaseScratchBuffer& ScratchBuffer)
+		FScratchBuffer& ScratchBuffer)
 	{
 		using namespace TypedElementDataStorage;
 
@@ -837,7 +837,7 @@ namespace UE::EditorDataStorage
 	//
 
 	FCommandOptimizer::FCommandOptimizer(
-		CompatibilityCommandBuffer::FOptimizer& InOptimizer, FTypedElementDatabaseScratchBuffer& InScratchBuffer)
+		CompatibilityCommandBuffer::FOptimizer& InOptimizer, FScratchBuffer& InScratchBuffer)
 		: Optimizer(InOptimizer)
 		, ScratchBuffer(InScratchBuffer)
 	{}
@@ -1143,7 +1143,7 @@ namespace UE::EditorDataStorage
 		}
 	}
 
-	void FCommandOptimizer::Run(CompatibilityCommandBuffer::FCollection& Commands, FTypedElementDatabaseScratchBuffer& ScratchBuffer)
+	void FCommandOptimizer::Run(CompatibilityCommandBuffer::FCollection& Commands, FScratchBuffer& ScratchBuffer)
 	{
 		CompatibilityCommandBuffer::FOptimizer Optimizer(Commands);
 		FCommandOptimizer OptimizationSelector(Optimizer, ScratchBuffer);
@@ -1154,4 +1154,4 @@ namespace UE::EditorDataStorage
 			Optimizer.ResetRightNextToLeft();
 		}
 	}
-} // namespace UE::EditorDataStorage
+} // namespace UE::Editor::DataStorage
diff --git a/Engine/Plugins/Experimental/EditorDataStorage/Source/TedsCore/Private/Commands/EditorDataStorageCompatibilityCommands.h b/Engine/Plugins/Experimental/EditorDataStorage/Source/TedsCore/Private/Commands/EditorDataStorageCompatibilityCommands.h
--- a/Engine/Plugins/Experimental/EditorDataStorage/Source/TedsCore/Private/Commands/EditorDataStorageCompatibilityCommands.h
+++ b/Engine/Plugins/Experimental/EditorDataStorage/Source/TedsCore/Private/Commands/EditorDataStorageCompatibilityCommands.h
@@ -10,14 +10,15 @@
 #include "UObject/WeakObjectPtrTemplates.h"
 
 class ITypedElementDataStorageInterface;
-class FTypedElementDatabaseScratchBuffer;
 class UTypedElementMementoSystem;
 class UScriptStruct;
 class UTypedElementDatabaseCompatibility;
 class UObject;
 
-namespace UE::EditorDataStorage
+namespace UE::Editor::DataStorage
 {
+	class FScratchBuffer;
+
 	enum class EObjectType : uint8
 	{
 		Struct,
@@ -221,7 +222,7 @@ namespace UE::EditorDataStorage
 
 		static bool IsPatchingRequired(const CompatibilityCommandBuffer::FCollection& Commands);
 		static void RunPatch(CompatibilityCommandBuffer::FCollection& Commands, UTypedElementDatabaseCompatibility& StorageCompat,
-			FTypedElementDatabaseScratchBuffer& ScratchBuffer);
+			FScratchBuffer& ScratchBuffer);
 	};
 
 	/** Prepares each command for further processing, e.g. resolving the target table. */
@@ -359,7 +360,7 @@ namespace UE::EditorDataStorage
 	{
 		// TODO: Allow [add object + restore memento] to be folded into one and [create memento + remove object].
 		//		Not sure how that will fit into the batch add though.
-		FCommandOptimizer(CompatibilityCommandBuffer::FOptimizer& InOptimizer, FTypedElementDatabaseScratchBuffer& InScratchBuffer);
+		FCommandOptimizer(CompatibilityCommandBuffer::FOptimizer& InOptimizer, FScratchBuffer& InScratchBuffer);
 
 		template<typename T>
 		void operator()(const T&) {}
@@ -372,7 +373,7 @@ namespace UE::EditorDataStorage
 		void operator()(const FAddSyncFromWorldTag& Command);
 
 		CompatibilityCommandBuffer::FOptimizer& Optimizer;
-		FTypedElementDatabaseScratchBuffer& ScratchBuffer;
+		FScratchBuffer& ScratchBuffer;
 
 		template<typename AddCommandType>
 		int32 FoldCommandsForAdd(const AddCommandType& Command);
@@ -384,6 +385,6 @@ namespace UE::EditorDataStorage
 		CompatibilityCommandBuffer::FOptimizer CreateRangeOptimizer(TypedElementDataStorage::RowHandle RowCluster);
 		/** Keeps processing a subset of commands until the right no longer has a command with the provided source row. */
 		void RunRightOnRowCluster(TypedElementDataStorage::RowHandle RowCluster);
-		static void Run(CompatibilityCommandBuffer::FCollection& Commands, FTypedElementDatabaseScratchBuffer& ScratchBuffer);
+		static void Run(CompatibilityCommandBuffer::FCollection& Commands, FScratchBuffer& ScratchBuffer);
 	};
-} // namespace UE::EditorDataStorage
+} // namespace UE::Editor::DataStorage
diff --git a/Engine/Plugins/Experimental/EditorDataStorage/Source/TedsCore/Private/Compatibility/TypedElementObjectReinstancingManager.cpp b/Engine/Plugins/Experimental/EditorDataStorage/Source/TedsCore/Private/Compatibility/TypedElementObjectReinstancingManager.cpp
--- a/Engine/Plugins/Experimental/EditorDataStorage/Source/TedsCore/Private/Compatibility/TypedElementObjectReinstancingManager.cpp
+++ b/Engine/Plugins/Experimental/EditorDataStorage/Source/TedsCore/Private/Compatibility/TypedElementObjectReinstancingManager.cpp
@@ -19,7 +19,7 @@ UTypedElementObjectReinstancingManager::UTypedElementObjectReinstancingManager()
 void UTypedElementObjectReinstancingManager::Initialize(UTypedElementDatabase& InDatabase, UTypedElementDatabaseCompatibility& InDataStorageCompatibility)
 {
 	using namespace TypedElementDataStorage;
-	using namespace UE::EditorDataStorage;
+	using namespace UE::Editor::DataStorage;
 
 	Database = &InDatabase;
 	DataStorageCompatibility = &InDataStorageCompatibility;
@@ -59,7 +59,7 @@ void UTypedElementObjectReinstancingManager::UpdateCompleted()
 
 void UTypedElementObjectReinstancingManager::HandleOnObjectPreRemoved(
 	const void* Object, 
-	const UE::EditorDataStorage::FObjectTypeInfo& TypeInfo, 
+	const UE::Editor::DataStorage::FObjectTypeInfo& TypeInfo, 
 	TypedElementDataStorage::RowHandle ObjectRow)
 {
 	// This is the chance to record the old object to memento
diff --git a/Engine/Plugins/Experimental/EditorDataStorage/Source/TedsCore/Private/Compatibility/TypedElementObjectReinstancingManager.h b/Engine/Plugins/Experimental/EditorDataStorage/Source/TedsCore/Private/Compatibility/TypedElementObjectReinstancingManager.h
--- a/Engine/Plugins/Experimental/EditorDataStorage/Source/TedsCore/Private/Compatibility/TypedElementObjectReinstancingManager.h
+++ b/Engine/Plugins/Experimental/EditorDataStorage/Source/TedsCore/Private/Compatibility/TypedElementObjectReinstancingManager.h
@@ -14,7 +14,7 @@ class ITypedElementDataStorageInterface;
 class UTypedElementDatabaseCompatibility;
 class UTypedElementMementoSystem;
 
-namespace UE::EditorDataStorage
+namespace UE::Editor::DataStorage
 {
 	struct FObjectTypeInfo;
 }
@@ -33,7 +33,7 @@ private:
 	void UpdateCompleted();
 	void HandleOnObjectPreRemoved(
 		const void* Object, 
-		const UE::EditorDataStorage::FObjectTypeInfo& TypeInfo, 
+		const UE::Editor::DataStorage::FObjectTypeInfo& TypeInfo, 
 		TypedElementDataStorage::RowHandle ObjectRow);
 	void HandleOnObjectsReinstanced(const FCoreUObjectDelegates::FReplacementObjectMap& ObjectReplacementMap);
 
diff --git a/Engine/Plugins/Experimental/EditorDataStorage/Source/TedsCore/Private/Debug/TypedElementDatabaseDebugCommands.cpp b/Engine/Plugins/Experimental/EditorDataStorage/Source/TedsCore/Private/Debug/TypedElementDatabaseDebugCommands.cpp
--- a/Engine/Plugins/Experimental/EditorDataStorage/Source/TedsCore/Private/Debug/TypedElementDatabaseDebugCommands.cpp
+++ b/Engine/Plugins/Experimental/EditorDataStorage/Source/TedsCore/Private/Debug/TypedElementDatabaseDebugCommands.cpp
@@ -229,7 +229,7 @@ static FAutoConsoleCommand CVarAddDynamicTag(
 	TEXT("Argument: Row, Tag, Value\n"),
 	FConsoleCommandWithArgsDelegate::CreateLambda([](const TArray<FString>& Args)
 	{
-		using namespace UE::EditorDataStorage;
+		using namespace UE::Editor::DataStorage;
 		ITypedElementDataStorageInterface* DataStorage = UTypedElementRegistry::GetInstance()->GetMutableDataStorage();
 
 		if (Args.Num() != 3)
@@ -277,13 +277,13 @@ static FAutoConsoleCommand CVarRemoveDynamicTag(
 		constexpr bool bUseTemplateSugar = true;
 		if constexpr (bUseTemplateSugar)
 		{
-			using namespace UE::EditorDataStorage;
+			using namespace UE::Editor::DataStorage;
 			const FName Tag = *Args[1];
 			DataStorage->RemoveColumn<FDynamicTag>(Row, Tag);
 		}
 		else
 		{
-			const UE::EditorDataStorage::FDynamicTag Tag(*Args[1]);
+			const UE::Editor::DataStorage::FDynamicTag Tag(*Args[1]);
 			DataStorage->RemoveColumn(Row, Tag);
 		}		
 	}),
@@ -296,7 +296,7 @@ static FAutoConsoleCommand CVarMatchDynamicTag(
 	{
 		using namespace TypedElementQueryBuilder;
 		using DSI = ITypedElementDataStorageInterface;
-		using namespace UE::EditorDataStorage;
+		using namespace UE::Editor::DataStorage;
 		
 		ITypedElementDataStorageInterface* DataStorage = UTypedElementRegistry::GetInstance()->GetMutableDataStorage();
 
@@ -409,7 +409,7 @@ static FAutoConsoleCommand CVarMatchDynamicTagFromEnum(
 	{
 		using namespace TypedElementQueryBuilder;
 		using DSI = ITypedElementDataStorageInterface;
-		using namespace UE::EditorDataStorage;
+		using namespace UE::Editor::DataStorage;
 		
 		ITypedElementDataStorageInterface* DataStorage = UTypedElementRegistry::GetInstance()->GetMutableDataStorage();
 
diff --git a/Engine/Plugins/Experimental/EditorDataStorage/Source/TedsCore/Private/DynamicColumnGenerator.cpp b/Engine/Plugins/Experimental/EditorDataStorage/Source/TedsCore/Private/DynamicColumnGenerator.cpp
--- a/Engine/Plugins/Experimental/EditorDataStorage/Source/TedsCore/Private/DynamicColumnGenerator.cpp
+++ b/Engine/Plugins/Experimental/EditorDataStorage/Source/TedsCore/Private/DynamicColumnGenerator.cpp
@@ -1,4 +1,4 @@
-﻿// Copyright Epic Games, Inc. All Rights Reserved.
+// Copyright Epic Games, Inc. All Rights Reserved.
 
 #include "DynamicColumnGenerator.h"
 
@@ -6,114 +6,117 @@
 #include "TypedElementDataStorageSharedColumn.h"
 #include "Elements/Common/TypedElementCommonTypes.h"
 
-UE::EditorDataStorage::FDynamicColumnGeneratorInfo UE::EditorDataStorage::FDynamicColumnGenerator::GenerateColumn(const FName& ColumnName, const UScriptStruct& Template)
+namespace UE::Editor::DataStorage
 {
-	const FGeneratedColumnKey Key
+	FDynamicColumnGeneratorInfo FDynamicColumnGenerator::GenerateColumn(const FName& ColumnName, const UScriptStruct& Template)
 	{
-		.Name = ColumnName,
-		.Template = Template
-	};
-
-	FDynamicColumnGeneratorInfo GeneratedColumnInfo;
-	UE_MT_SCOPED_WRITE_ACCESS(AccessDetector);
+		const FGeneratedColumnKey Key
+		{
+			.Name = ColumnName,
+			.Template = Template
+		};
 	
-	{
-		const int32* GeneratedColumnIndex = GenerationParamsLookup.Find(Key);
-		if (GeneratedColumnIndex != nullptr)
+		FDynamicColumnGeneratorInfo GeneratedColumnInfo;
+		UE_MT_SCOPED_WRITE_ACCESS(AccessDetector);
+		
 		{
-			const FGeneratedColumnRecord& GeneratedColumnRecord = GeneratedColumnData[*GeneratedColumnIndex];
-			GeneratedColumnInfo.Type = GeneratedColumnRecord.Type;
-			GeneratedColumnInfo.bNewlyGenerated = false;
+			const int32* GeneratedColumnIndex = GenerationParamsLookup.Find(Key);
+			if (GeneratedColumnIndex != nullptr)
+			{
+				const FGeneratedColumnRecord& GeneratedColumnRecord = GeneratedColumnData[*GeneratedColumnIndex];
+				GeneratedColumnInfo.Type = GeneratedColumnRecord.Type;
+				GeneratedColumnInfo.bNewlyGenerated = false;
+				return GeneratedColumnInfo;
+			}		
+		}
+		
+		checkf(
+			Template.IsChildOf(FTypedElementDataStorageColumn::StaticStruct()) ||
+			Template.IsChildOf(FTypedElementDataStorageTag::StaticStruct()) ||
+			Template.IsChildOf(FTedsSharedColumn::StaticStruct()),
+			TEXT("Template struct must derive from Column, Tag or SharedColumn"));
+	
+		{
+			checkf(Template.GetCppStructOps() != nullptr && Template.IsNative(), TEXT("Can only create column from native struct"));
+			
+			UScriptStruct* NewScriptStruct = NewObject<UScriptStruct>(GetTransientPackage(), ColumnName);
+			NewScriptStruct->AddToRoot();
+	
+			NewScriptStruct->SetSuperStruct(&const_cast<UScriptStruct&>(Template));
+			
+			NewScriptStruct->DeferCppStructOps(FTopLevelAssetPath(GetTransientPackage()->GetFName(), ColumnName), Template.GetCppStructOps());
+			NewScriptStruct->Bind();
+			NewScriptStruct->PrepareCppStructOps();
+			NewScriptStruct->StaticLink(true);
+	
+			
+			const int32 Index = GeneratedColumnData.Emplace(FGeneratedColumnRecord
+			{
+				.Name = ColumnName,
+				.Template = &Template,
+				.Type = NewScriptStruct
+			});
+			
+			GenerationParamsLookup.Add(Key, Index);
+			NameLookup.Add(ColumnName, Index);
+			
+			GeneratedColumnInfo.Type = NewScriptStruct;
+			GeneratedColumnInfo.bNewlyGenerated = true;
 			return GeneratedColumnInfo;
-		}		
+		}
 	}
 	
-	checkf(
-		Template.IsChildOf(FTypedElementDataStorageColumn::StaticStruct()) ||
-		Template.IsChildOf(FTypedElementDataStorageTag::StaticStruct()) ||
-		Template.IsChildOf(FTedsSharedColumn::StaticStruct()),
-		TEXT("Template struct must derive from Column, Tag or SharedColumn"));
-
+	const UScriptStruct* FDynamicColumnGenerator::LookupColumn(const FName& ColumnName) const
 	{
-		checkf(Template.GetCppStructOps() != nullptr && Template.IsNative(), TEXT("Can only create column from native struct"));
-		
-		UScriptStruct* NewScriptStruct = NewObject<UScriptStruct>(GetTransientPackage(), ColumnName);
-		NewScriptStruct->AddToRoot();
-
-		NewScriptStruct->SetSuperStruct(&const_cast<UScriptStruct&>(Template));
+		UE_MT_SCOPED_READ_ACCESS(AccessDetector);
 		
-		NewScriptStruct->DeferCppStructOps(FTopLevelAssetPath(GetTransientPackage()->GetFName(), ColumnName), Template.GetCppStructOps());
-		NewScriptStruct->Bind();
-		NewScriptStruct->PrepareCppStructOps();
-		NewScriptStruct->StaticLink(true);
-
-		
-		const int32 Index = GeneratedColumnData.Emplace(FGeneratedColumnRecord
+		if (const int32* Index = NameLookup.Find(ColumnName))
 		{
-			.Name = ColumnName,
-			.Template = &Template,
-			.Type = NewScriptStruct
-		});
-		
-		GenerationParamsLookup.Add(Key, Index);
-		NameLookup.Add(ColumnName, Index);
-		
-		GeneratedColumnInfo.Type = NewScriptStruct;
-		GeneratedColumnInfo.bNewlyGenerated = true;
-		return GeneratedColumnInfo;
+			return GeneratedColumnData[*Index].Type;
+		}
+		return nullptr;
 	}
-}
-
-const UScriptStruct* UE::EditorDataStorage::FDynamicColumnGenerator::LookupColumn(const FName& ColumnName) const
-{
-	UE_MT_SCOPED_READ_ACCESS(AccessDetector);
 	
-	if (const int32* Index = NameLookup.Find(ColumnName))
+	FDynamicTagManager::FDynamicTagManager(FDynamicColumnGenerator& InColumnGenerator)
+		: ColumnGenerator(InColumnGenerator)
 	{
-		return GeneratedColumnData[*Index].Type;
 	}
-	return nullptr;
-}
-
-UE::EditorDataStorage::FDynamicTagManager::FDynamicTagManager(FDynamicColumnGenerator& InColumnGenerator)
-	: ColumnGenerator(InColumnGenerator)
-{
-}
-
-FConstSharedStruct UE::EditorDataStorage::FDynamicTagManager::GenerateDynamicTag(const FDynamicTag& InTag, const FName& InValue)
-{
-	TPair<FDynamicTag, FName> Pair(InTag, InValue);
-
-	UE_MT_SCOPED_WRITE_ACCESS(AccessDetector);
-
-	// Common path
+	
+	FConstSharedStruct FDynamicTagManager::GenerateDynamicTag(const FDynamicTag& InTag, const FName& InValue)
 	{
-		if (FConstSharedStruct* TagStruct = DynamicTagLookup.Find(Pair))
+		TPair<FDynamicTag, FName> Pair(InTag, InValue);
+	
+		UE_MT_SCOPED_WRITE_ACCESS(AccessDetector);
+	
+		// Common path
+		{
+			if (FConstSharedStruct* TagStruct = DynamicTagLookup.Find(Pair))
+			{
+				return *TagStruct;
+			}
+		}
+	
+		// 
 		{
-			return *TagStruct;
+			const UScriptStruct* ColumnType = GenerateColumnType(InTag);
+	
+			const UE::Editor::DataStorage::FDynamicTagColumn Overlay
+			{
+				.Value = InValue
+			};
+	
+			FConstSharedStruct SharedStruct = FConstSharedStruct::Make(ColumnType, reinterpret_cast<const uint8*>(&Overlay));
+			
+			DynamicTagLookup.Emplace(Pair, SharedStruct);
+	
+			return SharedStruct;
 		}
 	}
-
-	// 
+	
+	const UScriptStruct* FDynamicTagManager::GenerateColumnType(const FDynamicTag& Tag)
 	{
-		const UScriptStruct* ColumnType = GenerateColumnType(InTag);
-
-		const FTypedElementDataStorageDynamicTagTemplate Overlay
-		{
-			.Value = InValue
-		};
-
-		FConstSharedStruct SharedStruct = FConstSharedStruct::Make(ColumnType, reinterpret_cast<const uint8*>(&Overlay));
+		const FDynamicColumnGeneratorInfo GeneratedColumnType = ColumnGenerator.GenerateColumn(Tag.GetName(), *FDynamicTagColumn::StaticStruct());
 		
-		DynamicTagLookup.Emplace(Pair, SharedStruct);
-
-		return SharedStruct;
+		return GeneratedColumnType.Type;
 	}
-}
-
-const UScriptStruct* UE::EditorDataStorage::FDynamicTagManager::GenerateColumnType(const FDynamicTag& Tag)
-{
-	const FDynamicColumnGeneratorInfo GeneratedColumnType = ColumnGenerator.GenerateColumn(Tag.GetName(), *FTypedElementDataStorageDynamicTagTemplate::StaticStruct());
-	
-	return GeneratedColumnType.Type;
-}
+} // namespace UE::Editor::DataStorage
diff --git a/Engine/Plugins/Experimental/EditorDataStorage/Source/TedsCore/Private/DynamicColumnGenerator.h b/Engine/Plugins/Experimental/EditorDataStorage/Source/TedsCore/Private/DynamicColumnGenerator.h
--- a/Engine/Plugins/Experimental/EditorDataStorage/Source/TedsCore/Private/DynamicColumnGenerator.h
+++ b/Engine/Plugins/Experimental/EditorDataStorage/Source/TedsCore/Private/DynamicColumnGenerator.h
@@ -1,4 +1,4 @@
-﻿// Copyright Epic Games, Inc. All Rights Reserved.
+// Copyright Epic Games, Inc. All Rights Reserved.
 
 #pragma once
 
@@ -18,7 +18,7 @@ class UScriptStruct;
 // The template struct that is used to generate the DynamicTag column
 // It is safe to reinterpret a DynamicTag column to this template to access the Value
 USTRUCT()
-struct FTypedElementDataStorageDynamicTagTemplate : public FTedsSharedColumn
+struct FTedsDynamicTagColumn : public FTedsSharedColumn
 {
 	GENERATED_BODY()
 	
@@ -26,8 +26,10 @@ struct FTypedElementDataStorageDynamicTagTemplate : public FTedsSharedColumn
 	FName Value;
 };
 
-namespace UE::EditorDataStorage
+namespace UE::Editor::DataStorage
 {
+	using FDynamicTagColumn = FTedsDynamicTagColumn;
+
 	struct FDynamicColumnInfo
 	{
 		const UScriptStruct* Type;
@@ -83,7 +85,7 @@ namespace UE::EditorDataStorage
 			}
 		};
 		// Looks up generated column index by the parameters used to generate it
-		// Used to deduplicate
+		// Used to de-duplicate
 		TMap<FGeneratedColumnKey, int32> GenerationParamsLookup;
 		// Looks up generated column index by name
 		TMap<FName, int32> NameLookup;
@@ -108,4 +110,4 @@ namespace UE::EditorDataStorage
 
 		FDynamicColumnGenerator& ColumnGenerator;
 	};
-}
+} // namespace UE::Editor::DataStorage
diff --git a/Engine/Plugins/Experimental/EditorDataStorage/Source/TedsCore/Private/GlobalLock.cpp b/Engine/Plugins/Experimental/EditorDataStorage/Source/TedsCore/Private/GlobalLock.cpp
--- a/Engine/Plugins/Experimental/EditorDataStorage/Source/TedsCore/Private/GlobalLock.cpp
+++ b/Engine/Plugins/Experimental/EditorDataStorage/Source/TedsCore/Private/GlobalLock.cpp
@@ -2,7 +2,7 @@
 
 #include "GlobalLock.h"
 
-namespace UE::EditorDataStorage
+namespace UE::Editor::DataStorage
 {
 	//
 	// FGlobalLock
@@ -136,4 +136,4 @@ namespace UE::EditorDataStorage
 	{
 		FGlobalLock::Unlock(Scope);
 	}
-} // UE::EditorDataStorage
\ No newline at end of file
+} // UE::Editor::DataStorage
\ No newline at end of file
diff --git a/Engine/Plugins/Experimental/EditorDataStorage/Source/TedsCore/Private/GlobalLock.h b/Engine/Plugins/Experimental/EditorDataStorage/Source/TedsCore/Private/GlobalLock.h
--- a/Engine/Plugins/Experimental/EditorDataStorage/Source/TedsCore/Private/GlobalLock.h
+++ b/Engine/Plugins/Experimental/EditorDataStorage/Source/TedsCore/Private/GlobalLock.h
@@ -8,7 +8,7 @@
 
 class UTypedElementDatabase;
 
-namespace UE::EditorDataStorage
+namespace UE::Editor::DataStorage
 {
 	enum class EGlobalLockStatus : uint8
 	{
@@ -84,4 +84,4 @@ namespace UE::EditorDataStorage
 	private:
 		EGlobalLockScope Scope;
 	};
-} // UE::EditorDataStorage
+} // namespace UE::Editor::DataStorage
diff --git a/Engine/Plugins/Experimental/EditorDataStorage/Source/TedsCore/Private/Memento/TypedElementMementoSystem.cpp b/Engine/Plugins/Experimental/EditorDataStorage/Source/TedsCore/Private/Memento/TypedElementMementoSystem.cpp
--- a/Engine/Plugins/Experimental/EditorDataStorage/Source/TedsCore/Private/Memento/TypedElementMementoSystem.cpp
+++ b/Engine/Plugins/Experimental/EditorDataStorage/Source/TedsCore/Private/Memento/TypedElementMementoSystem.cpp
@@ -12,7 +12,7 @@
 UTypedElementMementoSystem::UTypedElementMementoSystem(ITypedElementDataStorageInterface& InDataStorage)
 	: DataStorage(InDataStorage)
 {
-	using namespace UE::EditorDataStorage;
+	using namespace UE::Editor::DataStorage;
 
 	FScopedExclusiveLock Lock(EGlobalLockScope::Public);
 
@@ -41,7 +41,7 @@ UTypedElementMementoSystem::UTypedElementMementoSystem(ITypedElementDataStorageI
 TypedElementRowHandle UTypedElementMementoSystem::CreateMemento(TypedElementDataStorage::RowHandle SourceRow)
 {
 	using namespace TypedElementDataStorage;
-	using namespace UE::EditorDataStorage;
+	using namespace UE::Editor::DataStorage;
 
 	FScopedSharedLock Lock(EGlobalLockScope::Public);
 
@@ -52,7 +52,7 @@ TypedElementRowHandle UTypedElementMementoSystem::CreateMemento(TypedElementData
 
 void UTypedElementMementoSystem::CreateMemento(TypedElementDataStorage::RowHandle ReservedMementoRow, TypedElementDataStorage::RowHandle SourceRow)
 {
-	using namespace UE::EditorDataStorage;
+	using namespace UE::Editor::DataStorage;
 
 	FScopedSharedLock Lock(EGlobalLockScope::Public);
 	
@@ -63,7 +63,7 @@ void UTypedElementMementoSystem::CreateMemento(TypedElementDataStorage::RowHandl
 void UTypedElementMementoSystem::CreateMementoInternal(TypedElementDataStorage::RowHandle MementoRow, TypedElementDataStorage::RowHandle SourceRow)
 {
 	using namespace TypedElementDataStorage;
-	using namespace UE::EditorDataStorage;
+	using namespace UE::Editor::DataStorage;
 	
 	for (const UTypedElementMementoTranslatorBase* Translator : MementoTranslators)
 	{
@@ -90,7 +90,7 @@ void UTypedElementMementoSystem::CreateMementoInternal(TypedElementDataStorage::
 
 void UTypedElementMementoSystem::RestoreMemento(TypedElementDataStorage::RowHandle MementoRow, TypedElementDataStorage::RowHandle TargetRow)
 {
-	using namespace UE::EditorDataStorage;
+	using namespace UE::Editor::DataStorage;
 
 	FScopedSharedLock Lock(EGlobalLockScope::Public);
 
