[Teds] Rename native classes and move into UE::EditorDataStorage namespaces

diff --git a/Engine/Plugins/Experimental/EditorDataStorage/Source/TedsCore/Private/Processors/TypedElementProcessorAdaptors.cpp b/Engine/Plugins/Experimental/EditorDataStorage/Source/TedsCore/Private/Processors/TypedElementProcessorAdaptors.cpp
--- a/Engine/Plugins/Experimental/EditorDataStorage/Source/TedsCore/Private/Processors/TypedElementProcessorAdaptors.cpp
+++ b/Engine/Plugins/Experimental/EditorDataStorage/Source/TedsCore/Private/Processors/TypedElementProcessorAdaptors.cpp
@@ -10,700 +10,704 @@
 #include "TypedElementDatabase.h"
 #include "TypedElementDatabaseEnvironment.h"
 
-struct FMassContextCommon
+namespace UE::Editor::DataStorage
 {
-	FMassExecutionContext& Context;
+	namespace Processors::Private
+	{
+		struct FMassContextCommon
+		{
+			FMassExecutionContext& Context;
 	
-	explicit FMassContextCommon(FMassExecutionContext& InContext)
-		: Context(InContext)
-	{}
+			explicit FMassContextCommon(FMassExecutionContext& InContext)
+				: Context(InContext)
+			{}
 	
-	uint32 GetRowCount() const
-	{
-		return Context.GetNumEntities();
-	}
+			uint32 GetRowCount() const
+			{
+				return Context.GetNumEntities();
+			}
 
-	TConstArrayView<TypedElementRowHandle> GetRowHandles() const
-	{
-		static_assert(
-			sizeof(TypedElementRowHandle) == sizeof(FMassEntityHandle) && alignof(TypedElementRowHandle) == alignof(FMassEntityHandle),
-			"TypedElementRowHandle and FMassEntityHandle need to by layout compatible to support Typed Elements Data Storage.");
-		TConstArrayView<FMassEntityHandle> Entities = Context.GetEntities();
-		return TConstArrayView<TypedElementRowHandle>(reinterpret_cast<const TypedElementRowHandle*>(Entities.GetData()), Entities.Num());
-	}
+			TConstArrayView<TypedElementRowHandle> GetRowHandles() const
+			{
+				static_assert(
+					sizeof(TypedElementRowHandle) == sizeof(FMassEntityHandle) && alignof(TypedElementRowHandle) == alignof(FMassEntityHandle),
+					"TypedElementRowHandle and FMassEntityHandle need to by layout compatible to support TEDS.");
+				TConstArrayView<FMassEntityHandle> Entities = Context.GetEntities();
+				return TConstArrayView<TypedElementRowHandle>(reinterpret_cast<const TypedElementRowHandle*>(Entities.GetData()), Entities.Num());
+			}
 	
-	const void* GetColumn(const UScriptStruct* ColumnType) const
-	{
-		return Context.GetFragmentView(ColumnType).GetData();
-	}
+			const void* GetColumn(const UScriptStruct* ColumnType) const
+			{
+				return Context.GetFragmentView(ColumnType).GetData();
+			}
 
-	void* GetMutableColumn(const UScriptStruct* ColumnType)
-	{
-		return Context.GetMutableFragmentView(ColumnType).GetData();
-	}
+			void* GetMutableColumn(const UScriptStruct* ColumnType)
+			{
+				return Context.GetMutableFragmentView(ColumnType).GetData();
+			}
 
-	void GetColumns(TArrayView<char*> RetrievedAddresses,
-		TConstArrayView<TWeakObjectPtr<const UScriptStruct>> ColumnTypes,
-		TConstArrayView<TypedElementDataStorage::EQueryAccessType> AccessTypes)
-	{
-		checkf(RetrievedAddresses.Num() == ColumnTypes.Num(), TEXT("Unable to retrieve a batch of columns as the number of addresses "
-			"doesn't match the number of requested column."));
-		checkf(RetrievedAddresses.Num() == AccessTypes.Num(), TEXT("Unable to retrieve a batch of columns as the number of addresses "
-			"doesn't match the number of access types."));
+			void GetColumns(TArrayView<char*> RetrievedAddresses,
+				TConstArrayView<TWeakObjectPtr<const UScriptStruct>> ColumnTypes,
+				TConstArrayView<TypedElementDataStorage::EQueryAccessType> AccessTypes)
+			{
+				checkf(RetrievedAddresses.Num() == ColumnTypes.Num(), TEXT("Unable to retrieve a batch of columns as the number of addresses "
+					"doesn't match the number of requested column."));
+				checkf(RetrievedAddresses.Num() == AccessTypes.Num(), TEXT("Unable to retrieve a batch of columns as the number of addresses "
+					"doesn't match the number of access types."));
 
-		GetColumnsUnguarded(ColumnTypes.Num(), RetrievedAddresses.GetData(), ColumnTypes.GetData(), AccessTypes.GetData());
-	}
+				GetColumnsUnguarded(ColumnTypes.Num(), RetrievedAddresses.GetData(), ColumnTypes.GetData(), AccessTypes.GetData());
+			}
 	
-	void GetColumnsUnguarded(
-		int32 TypeCount,
-		char** RetrievedAddresses,
-		const TWeakObjectPtr<const UScriptStruct>* ColumnTypes,
-		const TypedElementDataStorage::EQueryAccessType* AccessTypes)
-	{
-		for (int32 Index = 0; Index < TypeCount; ++Index)
-		{
-			checkf(ColumnTypes->IsValid(), TEXT("Attempting to retrieve a column that is not available."));
-			*RetrievedAddresses = *AccessTypes == ITypedElementDataStorageInterface::EQueryAccessType::ReadWrite
-				? reinterpret_cast<char*>(Context.GetMutableFragmentView(ColumnTypes->Get()).GetData())
-				: const_cast<char*>(reinterpret_cast<const char*>(Context.GetFragmentView(ColumnTypes->Get()).GetData()));
-
-			++RetrievedAddresses;
-			++ColumnTypes;
-			++AccessTypes;
-		}
-	}
-
-	bool HasColumn(const UScriptStruct* ColumnType) const
-	{
-		if (ColumnType->IsChildOf(FMassTag::StaticStruct()))
-		{
-			return Context.DoesArchetypeHaveTag(*ColumnType);
-		}
-		if (ColumnType->IsChildOf(FMassFragment::StaticStruct()))
-		{
-			return Context.DoesArchetypeHaveFragment(*ColumnType);
-		}
-		const bool bIsTagOrFragment = false;
-		checkf(bIsTagOrFragment, TEXT("Attempting to check for a column type that is not a column or tag."));
-		return false;
-	}
-
-	bool HasColumn(TypedElementDataStorage::RowHandle Row, const UScriptStruct* ColumnType) const
-	{
-		FMassEntityHandle Entity = FMassEntityHandle::FromNumber(Row);
-		FMassEntityManager& Manager = Context.GetEntityManagerChecked();
-		FMassArchetypeHandle Archetype = Manager.GetArchetypeForEntity(Entity);
-		const FMassArchetypeCompositionDescriptor& Composition = Manager.GetArchetypeComposition(Archetype);
-
-		if (ColumnType->IsChildOf(FMassTag::StaticStruct()))
-		{
-			return Composition.Tags.Contains(*ColumnType);
-		}
-		if (ColumnType->IsChildOf(FMassFragment::StaticStruct()))
-		{
-			return Composition.Fragments.Contains(*ColumnType);
-		}
-		const bool bIsTagOrFragment = false;
-		checkf(bIsTagOrFragment, TEXT("Attempting to check for a column type that is not a column or tag."));
-		return false;
-	}
-};
-
-struct FMassWithEnvironmentContextCommon : public FMassContextCommon
-{
-	using Parent = FMassContextCommon;
-protected:
-	void TedsColumnsToMassDescriptorIfActiveTable(
-		FMassArchetypeCompositionDescriptor& Descriptor,
-		TConstArrayView<const UScriptStruct*> ColumnTypes)
-	{
-		for (const UScriptStruct* ColumnType : ColumnTypes)
-		{
-			if (ColumnType->IsChildOf(FMassTag::StaticStruct()))
+			void GetColumnsUnguarded(
+				int32 TypeCount,
+				char** RetrievedAddresses,
+				const TWeakObjectPtr<const UScriptStruct>* ColumnTypes,
+				const TypedElementDataStorage::EQueryAccessType* AccessTypes)
 			{
-				if (this->Context.DoesArchetypeHaveTag(*ColumnType))
+				for (int32 Index = 0; Index < TypeCount; ++Index)
 				{
-					Descriptor.Tags.Add(*ColumnType);
+					checkf(ColumnTypes->IsValid(), TEXT("Attempting to retrieve a column that is not available."));
+					*RetrievedAddresses = *AccessTypes == ITypedElementDataStorageInterface::EQueryAccessType::ReadWrite
+						? reinterpret_cast<char*>(Context.GetMutableFragmentView(ColumnTypes->Get()).GetData())
+						: const_cast<char*>(reinterpret_cast<const char*>(Context.GetFragmentView(ColumnTypes->Get()).GetData()));
+
+					++RetrievedAddresses;
+					++ColumnTypes;
+					++AccessTypes;
 				}
 			}
-			else
+
+			bool HasColumn(const UScriptStruct* ColumnType) const
 			{
-				checkf(ColumnType->IsChildOf(FMassFragment::StaticStruct()),
-					TEXT("Given struct type is not a valid fragment or tag type."));
-				if (this->Context.DoesArchetypeHaveFragment(*ColumnType))
+				if (ColumnType->IsChildOf(FMassTag::StaticStruct()))
+				{
+					return Context.DoesArchetypeHaveTag(*ColumnType);
+				}
+				if (ColumnType->IsChildOf(FMassFragment::StaticStruct()))
 				{
-					Descriptor.Fragments.Add(*ColumnType);
+					return Context.DoesArchetypeHaveFragment(*ColumnType);
 				}
+				const bool bIsTagOrFragment = false;
+				checkf(bIsTagOrFragment, TEXT("Attempting to check for a column type that is not a column or tag."));
+				return false;
 			}
-		}
-	}
 
-	void TedsColumnsToMassDescriptor(
-		FMassArchetypeCompositionDescriptor& Descriptor,
-		TConstArrayView<const UScriptStruct*> ColumnTypes)
-	{
-		for (const UScriptStruct* ColumnType : ColumnTypes)
-		{
-			if (ColumnType->IsChildOf(FMassTag::StaticStruct()))
-			{
-				Descriptor.Tags.Add(*ColumnType);
-			}
-			else
+			bool HasColumn(TypedElementDataStorage::RowHandle Row, const UScriptStruct* ColumnType) const
 			{
-				checkf(ColumnType->IsChildOf(FMassFragment::StaticStruct()),
-					TEXT("Given struct type is not a valid fragment or tag type."));
-				Descriptor.Fragments.Add(*ColumnType);
+				FMassEntityHandle Entity = FMassEntityHandle::FromNumber(Row);
+				FMassEntityManager& Manager = Context.GetEntityManagerChecked();
+				FMassArchetypeHandle Archetype = Manager.GetArchetypeForEntity(Entity);
+				const FMassArchetypeCompositionDescriptor& Composition = Manager.GetArchetypeComposition(Archetype);
 
+				if (ColumnType->IsChildOf(FMassTag::StaticStruct()))
+				{
+					return Composition.Tags.Contains(*ColumnType);
+				}
+				if (ColumnType->IsChildOf(FMassFragment::StaticStruct()))
+				{
+					return Composition.Fragments.Contains(*ColumnType);
+				}
+				const bool bIsTagOrFragment = false;
+				checkf(bIsTagOrFragment, TEXT("Attempting to check for a column type that is not a column or tag."));
+				return false;
 			}
-		}
-	}
-
-public:
-	using ObjectCopyOrMove = void (*)(const UScriptStruct& TypeInfo, void* Destination, void* Source);
-	
-	FMassWithEnvironmentContextCommon(FMassExecutionContext& InContext, FTypedElementDatabaseEnvironment& InEnvironment)
-		: FMassContextCommon(InContext)
-		, Environment(InEnvironment)
-	{}
-
-	uint64 GetUpdateCycleId() const 
-	{
-		return Environment.GetUpdateCycleId();
-	}
-
-	bool IsRowAvailable(TypedElementDataStorage::RowHandle Row) const
-	{
-		return Environment.GetMassEntityManager().IsEntityValid(FMassEntityHandle::FromNumber(Row));
-	}
-
-	bool IsRowAssigned(TypedElementDataStorage::RowHandle Row) const
-	{
-		return Environment.GetMassEntityManager().IsEntityActive(FMassEntityHandle::FromNumber(Row));
-	}
-	
-	void ActivateQueries(FName ActivationName)
-	{
-		this->Context.Defer().template PushCommand<FMassDeferredCommand<EMassCommandOperationType::None>>(
-			[Environment = &this->Environment, ActivationName](FMassEntityManager&)
-			{
-				Environment->GetQueryStore().ActivateQueries(ActivationName);
-			});
-	}
-
-	void* AddColumnUninitialized(TypedElementDataStorage::RowHandle Row, const UScriptStruct* ObjectType) 
-	{
-		return AddColumnUninitialized(Row, ObjectType,
-			[](const UScriptStruct& TypeInfo, void* Destination, void* Source)
-			{
-				TypeInfo.CopyScriptStruct(Destination, Source);
-			});
-	}
+		};
 
-	void* AddColumnUninitialized(TypedElementDataStorage::RowHandle Row, const UScriptStruct* ObjectType, ObjectCopyOrMove Relocator)
-	{
-		checkf(ObjectType->IsChildOf(FMassFragment::StaticStruct()), TEXT("Column [%s] can not be a tag"), *ObjectType->GetName());
-		
-		struct FAddValueColumn
+		struct FMassWithEnvironmentContextCommon : public FMassContextCommon
 		{
-			ObjectCopyOrMove Relocator;
-			const UScriptStruct* FragmentType;
-			FMassEntityHandle Entity;
-			void* Object;
-
-			FAddValueColumn() = default;
-			FAddValueColumn(ObjectCopyOrMove InRelocator, const UScriptStruct* InFragmentType, FMassEntityHandle InEntity, void* InObject)
-				: Relocator(InRelocator)
-				, FragmentType(InFragmentType)
-				, Entity(InEntity)
-				, Object(InObject)
-			{}
-
-			~FAddValueColumn()
+			using Parent = FMassContextCommon;
+		protected:
+			void TedsColumnsToMassDescriptorIfActiveTable(
+				FMassArchetypeCompositionDescriptor& Descriptor,
+				TConstArrayView<const UScriptStruct*> ColumnTypes)
 			{
-				if ((this->FragmentType->StructFlags & (STRUCT_IsPlainOldData | STRUCT_NoDestructor)) == 0)
+				for (const UScriptStruct* ColumnType : ColumnTypes)
 				{
-					this->FragmentType->DestroyStruct(this->Object);
+					if (ColumnType->IsChildOf(FMassTag::StaticStruct()))
+					{
+						if (this->Context.DoesArchetypeHaveTag(*ColumnType))
+						{
+							Descriptor.Tags.Add(*ColumnType);
+						}
+					}
+					else
+					{
+						checkf(ColumnType->IsChildOf(FMassFragment::StaticStruct()),
+							TEXT("Given struct type is not a valid fragment or tag type."));
+						if (this->Context.DoesArchetypeHaveFragment(*ColumnType))
+						{
+							Descriptor.Fragments.Add(*ColumnType);
+						}
+					}
 				}
 			}
-		};
 
-		FTypedElementDatabaseScratchBuffer& ScratchBuffer = Environment.GetScratchBuffer();
-		void* ColumnData = ScratchBuffer.Allocate(ObjectType->GetStructureSize(), ObjectType->GetMinAlignment());
-		FAddValueColumn* AddedColumn =
-			ScratchBuffer.Emplace<FAddValueColumn>(Relocator, ObjectType, FMassEntityHandle::FromNumber(Row), ColumnData);
-		
-		this->Context.Defer().template PushCommand<FMassDeferredAddCommand>(
-			[AddedColumn](FMassEntityManager& System)
+			void TedsColumnsToMassDescriptor(
+				FMassArchetypeCompositionDescriptor& Descriptor,
+				TConstArrayView<const UScriptStruct*> ColumnTypes)
 			{
-				// Check entity before proceeding. It's possible it may have been invalidated before this deferred call fired.
-				if (System.IsEntityActive(AddedColumn->Entity))
+				for (const UScriptStruct* ColumnType : ColumnTypes)
 				{
-					// Check before adding.  Mass's AddFragmentToEntity is not idempotent and will assert if adding
-					// column to a row that already has one
-					FStructView Fragment = System.GetFragmentDataStruct(AddedColumn->Entity, AddedColumn->FragmentType);
-					if (!Fragment.IsValid())
+					if (ColumnType->IsChildOf(FMassTag::StaticStruct()))
 					{
-						System.AddFragmentToEntity(AddedColumn->Entity, AddedColumn->FragmentType, 
-							[AddedColumn](void* Fragment, const UScriptStruct& FragmentType)
-							{
-								AddedColumn->Relocator(FragmentType, Fragment, AddedColumn->Object);
-							});
+						Descriptor.Tags.Add(*ColumnType);
 					}
 					else
 					{
-						AddedColumn->Relocator(*AddedColumn->FragmentType, Fragment.GetMemory(), AddedColumn->Object);
+						checkf(ColumnType->IsChildOf(FMassFragment::StaticStruct()),
+							TEXT("Given struct type is not a valid fragment or tag type."));
+						Descriptor.Fragments.Add(*ColumnType);
+
 					}
 				}
-			});
-		
-		return ColumnData;
-	}
+			}
 
-	void AddColumns(TypedElementDataStorage::RowHandle Row, TConstArrayView<const UScriptStruct*> ColumnTypes) 
-	{
-		struct FAddedColumns
-		{
-			FMassArchetypeCompositionDescriptor AddDescriptor;
-			FMassEntityHandle Entity;
-		};
+		public:
+			using ObjectCopyOrMove = void (*)(const UScriptStruct& TypeInfo, void* Destination, void* Source);
+	
+			FMassWithEnvironmentContextCommon(FMassExecutionContext& InContext, FEnvironment& InEnvironment)
+				: FMassContextCommon(InContext)
+				, Environment(InEnvironment)
+			{}
 
-		FAddedColumns* AddedColumns = Environment.GetScratchBuffer().template Emplace<FAddedColumns>();
-		TedsColumnsToMassDescriptor(AddedColumns->AddDescriptor, ColumnTypes);
-		AddedColumns->Entity = FMassEntityHandle::FromNumber(Row);
+			uint64 GetUpdateCycleId() const 
+			{
+				return Environment.GetUpdateCycleId();
+			}
 
-		this->Context.Defer().template PushCommand<FMassDeferredAddCommand>(
-			[AddedColumns](FMassEntityManager& System)
+			bool IsRowAvailable(TypedElementDataStorage::RowHandle Row) const
 			{
-				if (System.IsEntityValid(AddedColumns->Entity))
-				{
-					System.AddCompositionToEntity_GetDelta(AddedColumns->Entity, AddedColumns->AddDescriptor);
-				}
-			});
-	}
+				return Environment.GetMassEntityManager().IsEntityValid(FMassEntityHandle::FromNumber(Row));
+			}
 
-	void AddColumns(TConstArrayView<TypedElementDataStorage::RowHandle> Rows, TConstArrayView<const UScriptStruct*> ColumnTypes) 
-	{
-		struct FAddedColumns
-		{
-			FMassArchetypeCompositionDescriptor AddDescriptor;
-			FMassEntityHandle* Entities;
-			int32 EntityCount;
-		};
+			bool IsRowAssigned(TypedElementDataStorage::RowHandle Row) const
+			{
+				return Environment.GetMassEntityManager().IsEntityActive(FMassEntityHandle::FromNumber(Row));
+			}
+	
+			void ActivateQueries(FName ActivationName)
+			{
+				this->Context.Defer().template PushCommand<FMassDeferredCommand<EMassCommandOperationType::None>>(
+					[Environment = &this->Environment, ActivationName](FMassEntityManager&)
+					{
+						Environment->GetQueryStore().ActivateQueries(ActivationName);
+					});
+			}
 
-		FTypedElementDatabaseScratchBuffer& ScratchBuffer = Environment.GetScratchBuffer();
-		FAddedColumns* AddedColumns = ScratchBuffer.Emplace<FAddedColumns>();
-		TedsColumnsToMassDescriptor(AddedColumns->AddDescriptor, ColumnTypes);
-		
-		FMassEntityHandle* Entities = ScratchBuffer.EmplaceArray<FMassEntityHandle>(Rows.Num());
-		AddedColumns->Entities = Entities;
-		for (TypedElementRowHandle Row : Rows)
-		{
-			*Entities = FMassEntityHandle::FromNumber(Row);
-			Entities++;
-		}
-		AddedColumns->EntityCount = Rows.Num();
+			void* AddColumnUninitialized(TypedElementDataStorage::RowHandle Row, const UScriptStruct* ObjectType) 
+			{
+				return AddColumnUninitialized(Row, ObjectType,
+					[](const UScriptStruct& TypeInfo, void* Destination, void* Source)
+					{
+						TypeInfo.CopyScriptStruct(Destination, Source);
+					});
+			}
 
-		this->Context.Defer().template PushCommand<FMassDeferredAddCommand>(
-			[AddedColumns](FMassEntityManager& System)
+			void* AddColumnUninitialized(TypedElementDataStorage::RowHandle Row, const UScriptStruct* ObjectType, ObjectCopyOrMove Relocator)
 			{
-				FMassEntityHandle* Entities = AddedColumns->Entities;
-				int32 Count = AddedColumns->EntityCount;
-				for (int32 Counter = 0; Counter < Count; ++Counter)
+				checkf(ObjectType->IsChildOf(FMassFragment::StaticStruct()), TEXT("Column [%s] can not be a tag"), *ObjectType->GetName());
+		
+				struct FAddValueColumn
 				{
-					if (System.IsEntityValid(*Entities))
+					ObjectCopyOrMove Relocator;
+					const UScriptStruct* FragmentType;
+					FMassEntityHandle Entity;
+					void* Object;
+
+					FAddValueColumn() = default;
+					FAddValueColumn(ObjectCopyOrMove InRelocator, const UScriptStruct* InFragmentType, FMassEntityHandle InEntity, void* InObject)
+						: Relocator(InRelocator)
+						, FragmentType(InFragmentType)
+						, Entity(InEntity)
+						, Object(InObject)
+					{}
+
+					~FAddValueColumn()
 					{
-						System.AddCompositionToEntity_GetDelta(*Entities++, AddedColumns->AddDescriptor);
+						if ((this->FragmentType->StructFlags & (STRUCT_IsPlainOldData | STRUCT_NoDestructor)) == 0)
+						{
+							this->FragmentType->DestroyStruct(this->Object);
+						}
 					}
-				}
-			});
-	}
-
-	void RemoveColumns(TypedElementDataStorage::RowHandle Row, TConstArrayView<const UScriptStruct*> ColumnTypes) 
-	{
-		struct FRemovedColumns
-		{
-			FMassArchetypeCompositionDescriptor RemoveDescriptor;
-			FMassEntityHandle Entity;
-		};
-
-		FRemovedColumns* RemovedColumns = Environment.GetScratchBuffer().template Emplace<FRemovedColumns>();
-		TedsColumnsToMassDescriptorIfActiveTable(RemovedColumns->RemoveDescriptor, ColumnTypes);
-		if (!RemovedColumns->RemoveDescriptor.IsEmpty())
-		{
-			RemovedColumns->Entity = FMassEntityHandle::FromNumber(Row);
+				};
 
-			this->Context.Defer().template PushCommand<FMassDeferredAddCommand>(
-				[RemovedColumns](FMassEntityManager& System)
-				{
-					if (System.IsEntityValid(RemovedColumns->Entity))
+				FScratchBuffer& ScratchBuffer = Environment.GetScratchBuffer();
+				void* ColumnData = ScratchBuffer.Allocate(ObjectType->GetStructureSize(), ObjectType->GetMinAlignment());
+				FAddValueColumn* AddedColumn =
+					ScratchBuffer.Emplace<FAddValueColumn>(Relocator, ObjectType, FMassEntityHandle::FromNumber(Row), ColumnData);
+		
+				this->Context.Defer().template PushCommand<FMassDeferredAddCommand>(
+					[AddedColumn](FMassEntityManager& System)
 					{
-						System.RemoveCompositionFromEntity(RemovedColumns->Entity, RemovedColumns->RemoveDescriptor);
-					}
-				});
-		}
-	}
+						// Check entity before proceeding. It's possible it may have been invalidated before this deferred call fired.
+						if (System.IsEntityActive(AddedColumn->Entity))
+						{
+							// Check before adding.  Mass's AddFragmentToEntity is not idempotent and will assert if adding
+							// column to a row that already has one
+							FStructView Fragment = System.GetFragmentDataStruct(AddedColumn->Entity, AddedColumn->FragmentType);
+							if (!Fragment.IsValid())
+							{
+								System.AddFragmentToEntity(AddedColumn->Entity, AddedColumn->FragmentType, 
+									[AddedColumn](void* Fragment, const UScriptStruct& FragmentType)
+									{
+										AddedColumn->Relocator(FragmentType, Fragment, AddedColumn->Object);
+									});
+							}
+							else
+							{
+								AddedColumn->Relocator(*AddedColumn->FragmentType, Fragment.GetMemory(), AddedColumn->Object);
+							}
+						}
+					});
+		
+				return ColumnData;
+			}
 
-	void RemoveColumns(TConstArrayView<TypedElementDataStorage::RowHandle> Rows, TConstArrayView<const UScriptStruct*> ColumnTypes) 
-	{
-		struct FRemovedColumns
-		{
-			FMassArchetypeCompositionDescriptor RemoveDescriptor;
-			FMassEntityHandle* Entities;
-			int32 EntityCount;
-		};
+			void AddColumns(TypedElementDataStorage::RowHandle Row, TConstArrayView<const UScriptStruct*> ColumnTypes) 
+			{
+				struct FAddedColumns
+				{
+					FMassArchetypeCompositionDescriptor AddDescriptor;
+					FMassEntityHandle Entity;
+				};
 
-		FTypedElementDatabaseScratchBuffer& ScratchBuffer = Environment.GetScratchBuffer();
-		FRemovedColumns* RemovedColumns = ScratchBuffer.Emplace<FRemovedColumns>();
-		TedsColumnsToMassDescriptorIfActiveTable(RemovedColumns->RemoveDescriptor, ColumnTypes);
+				FAddedColumns* AddedColumns = Environment.GetScratchBuffer().template Emplace<FAddedColumns>();
+				TedsColumnsToMassDescriptor(AddedColumns->AddDescriptor, ColumnTypes);
+				AddedColumns->Entity = FMassEntityHandle::FromNumber(Row);
 
-		FMassEntityHandle* Entities = ScratchBuffer.EmplaceArray<FMassEntityHandle>(Rows.Num());
-		RemovedColumns->Entities = Entities;
-		for (TypedElementRowHandle Row : Rows)
-		{
-			*Entities = FMassEntityHandle::FromNumber(Row);
-			Entities++;
-		}
-		RemovedColumns->EntityCount = Rows.Num();
+				this->Context.Defer().template PushCommand<FMassDeferredAddCommand>(
+					[AddedColumns](FMassEntityManager& System)
+					{
+						if (System.IsEntityValid(AddedColumns->Entity))
+						{
+							System.AddCompositionToEntity_GetDelta(AddedColumns->Entity, AddedColumns->AddDescriptor);
+						}
+					});
+			}
 
-		this->Context.Defer().template PushCommand<FMassDeferredAddCommand>(
-			[RemovedColumns](FMassEntityManager& System)
+			void AddColumns(TConstArrayView<TypedElementDataStorage::RowHandle> Rows, TConstArrayView<const UScriptStruct*> ColumnTypes) 
 			{
-				FMassEntityHandle* Entities = RemovedColumns->Entities;
-				int32 Count = RemovedColumns->EntityCount;
-
-				using EntityHandleArray = TArray<FMassEntityHandle, TInlineAllocator<32>>;
-				using EntityArchetypeLookup = TMap<FMassArchetypeHandle, EntityHandleArray, TInlineSetAllocator<32>>;
-				using ArchetypeEntityArray = TArray<FMassArchetypeEntityCollection, TInlineAllocator<32>>;
-
-				// Sort rows (entities) into to matching table (archetype) bucket.
-				EntityArchetypeLookup LookupTable;
-				for (int32 Counter = 0; Counter < Count; ++Counter)
+				struct FAddedColumns
 				{
-					if (System.IsEntityValid(*Entities))
-					{
-						FMassArchetypeHandle Archetype = System.GetArchetypeForEntity(*Entities);
-						EntityHandleArray& EntityCollection = LookupTable.FindOrAdd(Archetype);
-						EntityCollection.Add(*Entities);
-					}
+					FMassArchetypeCompositionDescriptor AddDescriptor;
+					FMassEntityHandle* Entities;
+					int32 EntityCount;
+				};
+
+				FScratchBuffer& ScratchBuffer = Environment.GetScratchBuffer();
+				FAddedColumns* AddedColumns = ScratchBuffer.Emplace<FAddedColumns>();
+				TedsColumnsToMassDescriptor(AddedColumns->AddDescriptor, ColumnTypes);
+		
+				FMassEntityHandle* Entities = ScratchBuffer.EmplaceArray<FMassEntityHandle>(Rows.Num());
+				AddedColumns->Entities = Entities;
+				for (TypedElementRowHandle Row : Rows)
+				{
+					*Entities = FMassEntityHandle::FromNumber(Row);
 					Entities++;
 				}
+				AddedColumns->EntityCount = Rows.Num();
 
-				// Construct table (archetype) specific row (entity) collections.
-				ArchetypeEntityArray EntityCollections;
-				EntityCollections.Reserve(LookupTable.Num());
-				for (auto It = LookupTable.CreateConstIterator(); It; ++It)
-				{
-					// Could be more effective but the previous implementation was robust when called with duplicate rows.
-					EntityCollections.Emplace(It.Key(), It.Value(), FMassArchetypeEntityCollection::EDuplicatesHandling::FoldDuplicates);
-				}
+				this->Context.Defer().template PushCommand<FMassDeferredAddCommand>(
+					[AddedColumns](FMassEntityManager& System)
+					{
+						FMassEntityHandle* Entities = AddedColumns->Entities;
+						int32 Count = AddedColumns->EntityCount;
+						for (int32 Counter = 0; Counter < Count; ++Counter)
+						{
+							if (System.IsEntityValid(*Entities))
+							{
+								System.AddCompositionToEntity_GetDelta(*Entities++, AddedColumns->AddDescriptor);
+							}
+						}
+					});
+			}
 
+			void RemoveColumns(TypedElementDataStorage::RowHandle Row, TConstArrayView<const UScriptStruct*> ColumnTypes) 
+			{
+				struct FRemovedColumns
+				{
+					FMassArchetypeCompositionDescriptor RemoveDescriptor;
+					FMassEntityHandle Entity;
+				};
 
-				// This could be improved by adding an operation that would both combine the Fragments and Tags change in one bath operation.
-				if (!RemovedColumns->RemoveDescriptor.Fragments.IsEmpty())
+				FRemovedColumns* RemovedColumns = Environment.GetScratchBuffer().template Emplace<FRemovedColumns>();
+				TedsColumnsToMassDescriptorIfActiveTable(RemovedColumns->RemoveDescriptor, ColumnTypes);
+				if (!RemovedColumns->RemoveDescriptor.IsEmpty())
 				{
-					System.BatchChangeFragmentCompositionForEntities(EntityCollections, FMassFragmentBitSet(), RemovedColumns->RemoveDescriptor.Fragments);
+					RemovedColumns->Entity = FMassEntityHandle::FromNumber(Row);
+
+					this->Context.Defer().template PushCommand<FMassDeferredAddCommand>(
+						[RemovedColumns](FMassEntityManager& System)
+						{
+							if (System.IsEntityValid(RemovedColumns->Entity))
+							{
+								System.RemoveCompositionFromEntity(RemovedColumns->Entity, RemovedColumns->RemoveDescriptor);
+							}
+						});
 				}
-				if (!RemovedColumns->RemoveDescriptor.Tags.IsEmpty())
+			}
+
+			void RemoveColumns(TConstArrayView<TypedElementDataStorage::RowHandle> Rows, TConstArrayView<const UScriptStruct*> ColumnTypes) 
+			{
+				struct FRemovedColumns
 				{
-					System.BatchChangeTagsForEntities(EntityCollections, FMassTagBitSet(), RemovedColumns->RemoveDescriptor.Tags);
+					FMassArchetypeCompositionDescriptor RemoveDescriptor;
+					FMassEntityHandle* Entities;
+					int32 EntityCount;
+				};
+
+				FScratchBuffer& ScratchBuffer = Environment.GetScratchBuffer();
+				FRemovedColumns* RemovedColumns = ScratchBuffer.Emplace<FRemovedColumns>();
+				TedsColumnsToMassDescriptorIfActiveTable(RemovedColumns->RemoveDescriptor, ColumnTypes);
+
+				FMassEntityHandle* Entities = ScratchBuffer.EmplaceArray<FMassEntityHandle>(Rows.Num());
+				RemovedColumns->Entities = Entities;
+				for (TypedElementRowHandle Row : Rows)
+				{
+					*Entities = FMassEntityHandle::FromNumber(Row);
+					Entities++;
 				}
-			});
-	}
+				RemovedColumns->EntityCount = Rows.Num();
+
+				this->Context.Defer().template PushCommand<FMassDeferredAddCommand>(
+					[RemovedColumns](FMassEntityManager& System)
+					{
+						FMassEntityHandle* Entities = RemovedColumns->Entities;
+						int32 Count = RemovedColumns->EntityCount;
+
+						using EntityHandleArray = TArray<FMassEntityHandle, TInlineAllocator<32>>;
+						using EntityArchetypeLookup = TMap<FMassArchetypeHandle, EntityHandleArray, TInlineSetAllocator<32>>;
+						using ArchetypeEntityArray = TArray<FMassArchetypeEntityCollection, TInlineAllocator<32>>;
+
+						// Sort rows (entities) into to matching table (archetype) bucket.
+						EntityArchetypeLookup LookupTable;
+						for (int32 Counter = 0; Counter < Count; ++Counter)
+						{
+							if (System.IsEntityValid(*Entities))
+							{
+								FMassArchetypeHandle Archetype = System.GetArchetypeForEntity(*Entities);
+								EntityHandleArray& EntityCollection = LookupTable.FindOrAdd(Archetype);
+								EntityCollection.Add(*Entities);
+							}
+							Entities++;
+						}
+
+						// Construct table (archetype) specific row (entity) collections.
+						ArchetypeEntityArray EntityCollections;
+						EntityCollections.Reserve(LookupTable.Num());
+						for (auto It = LookupTable.CreateConstIterator(); It; ++It)
+						{
+							// Could be more effective but the previous implementation was robust when called with duplicate rows.
+							EntityCollections.Emplace(It.Key(), It.Value(), FMassArchetypeEntityCollection::EDuplicatesHandling::FoldDuplicates);
+						}
+
+
+						// This could be improved by adding an operation that would both combine the Fragments and Tags change in one bath operation.
+						if (!RemovedColumns->RemoveDescriptor.Fragments.IsEmpty())
+						{
+							System.BatchChangeFragmentCompositionForEntities(EntityCollections, FMassFragmentBitSet(), RemovedColumns->RemoveDescriptor.Fragments);
+						}
+						if (!RemovedColumns->RemoveDescriptor.Tags.IsEmpty())
+						{
+							System.BatchChangeTagsForEntities(EntityCollections, FMassTagBitSet(), RemovedColumns->RemoveDescriptor.Tags);
+						}
+					});
+			}
 	
-	TypedElementDataStorage::RowHandle AddRow(TypedElementDataStorage::TableHandle Table) 
-	{
-		FMassEntityHandle EntityHandle = Environment.GetMassEntityManager().ReserveEntity();
-		FMassArchetypeHandle ArchetypeHandle = Environment.LookupMassArchetype(Table);
+			TypedElementDataStorage::RowHandle AddRow(TypedElementDataStorage::TableHandle Table) 
+			{
+				FMassEntityHandle EntityHandle = Environment.GetMassEntityManager().ReserveEntity();
+				FMassArchetypeHandle ArchetypeHandle = Environment.LookupMassArchetype(Table);
 
-		if (!ArchetypeHandle.IsValid())
-		{
-			return TypedElementDataStorage::InvalidRowHandle;
-		}
+				if (!ArchetypeHandle.IsValid())
+				{
+					return TypedElementDataStorage::InvalidRowHandle;
+				}
 		
-		struct CommandInfo
-		{
-			FMassEntityHandle Entity;
-			FMassArchetypeHandle Archetype;
-		};
+				struct CommandInfo
+				{
+					FMassEntityHandle Entity;
+					FMassArchetypeHandle Archetype;
+				};
 
-		CommandInfo CommandDataTmp{
-			.Entity = EntityHandle,
-			.Archetype = MoveTemp(ArchetypeHandle)
-		};
+				CommandInfo CommandDataTmp{
+					.Entity = EntityHandle,
+					.Archetype = MoveTemp(ArchetypeHandle)
+				};
 
-		this->Context.Defer().template PushCommand<FMassDeferredCreateCommand>(
-			[CommandData = MoveTemp(CommandDataTmp)](FMassEntityManager& System)
-			{
-				const FMassArchetypeSharedFragmentValues SharedFragmentValues;
-				System.BuildEntity(CommandData.Entity, CommandData.Archetype, SharedFragmentValues);
-			});
+				this->Context.Defer().template PushCommand<FMassDeferredCreateCommand>(
+					[CommandData = MoveTemp(CommandDataTmp)](FMassEntityManager& System)
+					{
+						const FMassArchetypeSharedFragmentValues SharedFragmentValues;
+						System.BuildEntity(CommandData.Entity, CommandData.Archetype, SharedFragmentValues);
+					});
 		
-		const TypedElementDataStorage::RowHandle TedsRowHandle = EntityHandle.AsNumber();
-		return TedsRowHandle;
-	}
+				const TypedElementDataStorage::RowHandle TedsRowHandle = EntityHandle.AsNumber();
+				return TedsRowHandle;
+			}
 
-	void RemoveRow(TypedElementRowHandle Row) 
-	{
-		this->Context.Defer().DestroyEntity(FMassEntityHandle::FromNumber(Row));
-	}
+			void RemoveRow(TypedElementRowHandle Row) 
+			{
+				this->Context.Defer().DestroyEntity(FMassEntityHandle::FromNumber(Row));
+			}
 
-	void RemoveRows(TConstArrayView<TypedElementRowHandle> Rows) 
-	{
-		// Row handles and entities map 1:1 for data, so a reintpret_cast can be safely done to avoid
-		// having to allocate memory and iterating over the rows.
+			void RemoveRows(TConstArrayView<TypedElementRowHandle> Rows) 
+			{
+				// Row handles and entities map 1:1 for data, so a reintpret_cast can be safely done to avoid
+				// having to allocate memory and iterating over the rows.
 
-		static_assert(sizeof(FMassEntityHandle) == sizeof(TypedElementRowHandle), 
-			"Sizes of mass entity and data storage row have gone out of sync.");
-		static_assert(alignof(FMassEntityHandle) == alignof(TypedElementRowHandle),
-			"Alignment of mass entity and data storage row have gone out of sync.");
+				static_assert(sizeof(FMassEntityHandle) == sizeof(TypedElementRowHandle), 
+					"Sizes of mass entity and data storage row have gone out of sync.");
+				static_assert(alignof(FMassEntityHandle) == alignof(TypedElementRowHandle),
+					"Alignment of mass entity and data storage row have gone out of sync.");
 
-		this->Context.Defer().DestroyEntities(
-			TConstArrayView<FMassEntityHandle>(reinterpret_cast<const FMassEntityHandle*>(Rows.begin()), Rows.Num()));
-	}
+				this->Context.Defer().DestroyEntities(
+					TConstArrayView<FMassEntityHandle>(reinterpret_cast<const FMassEntityHandle*>(Rows.begin()), Rows.Num()));
+			}
 
-protected:
-	FTypedElementDatabaseEnvironment& Environment;
-};
+		protected:
+			FEnvironment& Environment;
+		};
 
-struct FMassDirectContextForwarder final : public ITypedElementDataStorageInterface::IDirectQueryContext
-{
-	explicit FMassDirectContextForwarder(FMassExecutionContext& InContext)
-		: Implementation(InContext)
-	{}
+		struct FMassDirectContextForwarder final : public ITypedElementDataStorageInterface::IDirectQueryContext
+		{
+			explicit FMassDirectContextForwarder(FMassExecutionContext& InContext)
+				: Implementation(InContext)
+			{}
 	
-	uint32 GetRowCount() const override { return Implementation.GetRowCount(); }
-	TConstArrayView<TypedElementDataStorage::RowHandle> GetRowHandles() const override { return Implementation.GetRowHandles(); }
-	const void* GetColumn(const UScriptStruct* ColumnType) const override { return Implementation.GetColumn(ColumnType); }
-	void* GetMutableColumn(const UScriptStruct* ColumnType) override { return Implementation.GetMutableColumn(ColumnType); }
-	void GetColumns(TArrayView<char*> RetrievedAddresses, TConstArrayView<TWeakObjectPtr<const UScriptStruct>> ColumnTypes, TConstArrayView<TypedElementDataStorage::EQueryAccessType> AccessTypes) override { return Implementation.GetColumns(RetrievedAddresses, ColumnTypes, AccessTypes); }
-	void GetColumnsUnguarded(int32 TypeCount, char** RetrievedAddresses, const TWeakObjectPtr<const UScriptStruct>* ColumnTypes, const TypedElementDataStorage::EQueryAccessType* AccessTypes) override { return Implementation.GetColumnsUnguarded(TypeCount, RetrievedAddresses, ColumnTypes, AccessTypes); }
-	bool HasColumn(const UScriptStruct* ColumnType) const override { return Implementation.HasColumn(ColumnType); }
-	bool HasColumn(TypedElementDataStorage::RowHandle Row, const UScriptStruct* ColumnType) const override { return Implementation.HasColumn(Row, ColumnType); }
-
-	FMassContextCommon Implementation;
-};
-
-struct FMassSubqueryContextForwarder  : public ITypedElementDataStorageInterface::ISubqueryContext
-{
-	FMassSubqueryContextForwarder(FMassExecutionContext& InContext, FTypedElementDatabaseEnvironment& InEnvironment)
-		: Implementation(InContext, InEnvironment)
-	{}
-
-	~FMassSubqueryContextForwarder() override = default;
-	uint32 GetRowCount() const override { return Implementation.GetRowCount(); }
-	TConstArrayView<TypedElementDataStorage::RowHandle> GetRowHandles() const override { return Implementation.GetRowHandles(); }
-	const void* GetColumn(const UScriptStruct* ColumnType) const override { return Implementation.GetColumn(ColumnType); }
-	void* GetMutableColumn(const UScriptStruct* ColumnType) override { return Implementation.GetMutableColumn(ColumnType); }
-	void GetColumns(TArrayView<char*> RetrievedAddresses, TConstArrayView<TWeakObjectPtr<const UScriptStruct>> ColumnTypes, TConstArrayView<TypedElementDataStorage::EQueryAccessType> AccessTypes) override { return Implementation.GetColumns(RetrievedAddresses, ColumnTypes, AccessTypes); }
-	void GetColumnsUnguarded(int32 TypeCount, char** RetrievedAddresses, const TWeakObjectPtr<const UScriptStruct>* ColumnTypes, const TypedElementDataStorage::EQueryAccessType* AccessTypes) override { return Implementation.GetColumnsUnguarded(TypeCount, RetrievedAddresses, ColumnTypes, AccessTypes); }
-	bool HasColumn(const UScriptStruct* ColumnType) const override { return Implementation.HasColumn(ColumnType); }
-	bool HasColumn(TypedElementDataStorage::RowHandle Row, const UScriptStruct* ColumnType) const override { return Implementation.HasColumn(Row, ColumnType); }
-	uint64 GetUpdateCycleId() const override { return Implementation.GetUpdateCycleId(); }
-	bool IsRowAvailable(TypedElementDataStorage::RowHandle Row) const override { return Implementation.IsRowAvailable(Row); }
-	bool IsRowAssigned(TypedElementDataStorage::RowHandle Row) const override { return Implementation.IsRowAssigned(Row); }
-	void ActivateQueries(FName ActivationName) override  { return Implementation.ActivateQueries(ActivationName); }
-	TypedElementDataStorage::RowHandle AddRow(TypedElementDataStorage::TableHandle Table) override  { return Implementation.AddRow(Table); }
-	void RemoveRow(TypedElementDataStorage::RowHandle Row) override { return Implementation.RemoveRow(Row); }
-	void RemoveRows(TConstArrayView<TypedElementDataStorage::RowHandle> Rows) override  { return Implementation.RemoveRows(Rows); }
-	void AddColumns(TypedElementDataStorage::RowHandle Row, TConstArrayView<const UScriptStruct*> ColumnTypes) override { return Implementation.AddColumns(Row, ColumnTypes); }
-	void AddColumns(TConstArrayView<TypedElementDataStorage::RowHandle> Rows, TConstArrayView<const UScriptStruct*> ColumnTypes) override  { return Implementation.AddColumns(Rows, ColumnTypes); }
-	void* AddColumnUninitialized(TypedElementDataStorage::RowHandle Row, const UScriptStruct* ColumnType) override  { return Implementation.AddColumnUninitialized(Row, ColumnType); }
-	void* AddColumnUninitialized(TypedElementDataStorage::RowHandle Row, const UScriptStruct* ObjectType, ObjectCopyOrMove Relocator) override { return Implementation.AddColumnUninitialized(Row, ObjectType, Relocator); }
-	void RemoveColumns(TypedElementDataStorage::RowHandle Row, TConstArrayView<const UScriptStruct*> ColumnTypes) override { Implementation.RemoveColumns(Row, ColumnTypes); }
-	void RemoveColumns(TConstArrayView<TypedElementDataStorage::RowHandle> Rows, TConstArrayView<const UScriptStruct*> ColumnTypes) override { Implementation.RemoveColumns(Rows, ColumnTypes); }
-
-	FMassWithEnvironmentContextCommon Implementation;
-};
-
-struct FMassQueryContextImplementation final : FMassWithEnvironmentContextCommon
-{
-	FMassQueryContextImplementation(
-		ITypedElementDataStorageInterface::FQueryDescription& InQueryDescription,
-		FMassExecutionContext& InContext, 
-		FTypedElementExtendedQueryStore& InQueryStore,
-		FTypedElementDatabaseEnvironment& InEnvironment)
-		: FMassWithEnvironmentContextCommon(InContext, InEnvironment)
-		, QueryDescription(InQueryDescription)
-		, QueryStore(InQueryStore)
-	{
+			uint32 GetRowCount() const override { return Implementation.GetRowCount(); }
+			TConstArrayView<TypedElementDataStorage::RowHandle> GetRowHandles() const override { return Implementation.GetRowHandles(); }
+			const void* GetColumn(const UScriptStruct* ColumnType) const override { return Implementation.GetColumn(ColumnType); }
+			void* GetMutableColumn(const UScriptStruct* ColumnType) override { return Implementation.GetMutableColumn(ColumnType); }
+			void GetColumns(TArrayView<char*> RetrievedAddresses, TConstArrayView<TWeakObjectPtr<const UScriptStruct>> ColumnTypes, TConstArrayView<TypedElementDataStorage::EQueryAccessType> AccessTypes) override { return Implementation.GetColumns(RetrievedAddresses, ColumnTypes, AccessTypes); }
+			void GetColumnsUnguarded(int32 TypeCount, char** RetrievedAddresses, const TWeakObjectPtr<const UScriptStruct>* ColumnTypes, const TypedElementDataStorage::EQueryAccessType* AccessTypes) override { return Implementation.GetColumnsUnguarded(TypeCount, RetrievedAddresses, ColumnTypes, AccessTypes); }
+			bool HasColumn(const UScriptStruct* ColumnType) const override { return Implementation.HasColumn(ColumnType); }
+			bool HasColumn(TypedElementDataStorage::RowHandle Row, const UScriptStruct* ColumnType) const override { return Implementation.HasColumn(Row, ColumnType); }
+
+			FMassContextCommon Implementation;
+		};
+
+		struct FMassSubqueryContextForwarder  : public ITypedElementDataStorageInterface::ISubqueryContext
+		{
+			FMassSubqueryContextForwarder(FMassExecutionContext& InContext, FEnvironment& InEnvironment)
+				: Implementation(InContext, InEnvironment)
+			{}
+
+			~FMassSubqueryContextForwarder() override = default;
+			uint32 GetRowCount() const override { return Implementation.GetRowCount(); }
+			TConstArrayView<TypedElementDataStorage::RowHandle> GetRowHandles() const override { return Implementation.GetRowHandles(); }
+			const void* GetColumn(const UScriptStruct* ColumnType) const override { return Implementation.GetColumn(ColumnType); }
+			void* GetMutableColumn(const UScriptStruct* ColumnType) override { return Implementation.GetMutableColumn(ColumnType); }
+			void GetColumns(TArrayView<char*> RetrievedAddresses, TConstArrayView<TWeakObjectPtr<const UScriptStruct>> ColumnTypes, TConstArrayView<TypedElementDataStorage::EQueryAccessType> AccessTypes) override { return Implementation.GetColumns(RetrievedAddresses, ColumnTypes, AccessTypes); }
+			void GetColumnsUnguarded(int32 TypeCount, char** RetrievedAddresses, const TWeakObjectPtr<const UScriptStruct>* ColumnTypes, const TypedElementDataStorage::EQueryAccessType* AccessTypes) override { return Implementation.GetColumnsUnguarded(TypeCount, RetrievedAddresses, ColumnTypes, AccessTypes); }
+			bool HasColumn(const UScriptStruct* ColumnType) const override { return Implementation.HasColumn(ColumnType); }
+			bool HasColumn(TypedElementDataStorage::RowHandle Row, const UScriptStruct* ColumnType) const override { return Implementation.HasColumn(Row, ColumnType); }
+			uint64 GetUpdateCycleId() const override { return Implementation.GetUpdateCycleId(); }
+			bool IsRowAvailable(TypedElementDataStorage::RowHandle Row) const override { return Implementation.IsRowAvailable(Row); }
+			bool IsRowAssigned(TypedElementDataStorage::RowHandle Row) const override { return Implementation.IsRowAssigned(Row); }
+			void ActivateQueries(FName ActivationName) override  { return Implementation.ActivateQueries(ActivationName); }
+			TypedElementDataStorage::RowHandle AddRow(TypedElementDataStorage::TableHandle Table) override  { return Implementation.AddRow(Table); }
+			void RemoveRow(TypedElementDataStorage::RowHandle Row) override { return Implementation.RemoveRow(Row); }
+			void RemoveRows(TConstArrayView<TypedElementDataStorage::RowHandle> Rows) override  { return Implementation.RemoveRows(Rows); }
+			void AddColumns(TypedElementDataStorage::RowHandle Row, TConstArrayView<const UScriptStruct*> ColumnTypes) override { return Implementation.AddColumns(Row, ColumnTypes); }
+			void AddColumns(TConstArrayView<TypedElementDataStorage::RowHandle> Rows, TConstArrayView<const UScriptStruct*> ColumnTypes) override  { return Implementation.AddColumns(Rows, ColumnTypes); }
+			void* AddColumnUninitialized(TypedElementDataStorage::RowHandle Row, const UScriptStruct* ColumnType) override  { return Implementation.AddColumnUninitialized(Row, ColumnType); }
+			void* AddColumnUninitialized(TypedElementDataStorage::RowHandle Row, const UScriptStruct* ObjectType, ObjectCopyOrMove Relocator) override { return Implementation.AddColumnUninitialized(Row, ObjectType, Relocator); }
+			void RemoveColumns(TypedElementDataStorage::RowHandle Row, TConstArrayView<const UScriptStruct*> ColumnTypes) override { Implementation.RemoveColumns(Row, ColumnTypes); }
+			void RemoveColumns(TConstArrayView<TypedElementDataStorage::RowHandle> Rows, TConstArrayView<const UScriptStruct*> ColumnTypes) override { Implementation.RemoveColumns(Rows, ColumnTypes); }
+
+			FMassWithEnvironmentContextCommon Implementation;
+		};
+
+		struct FMassQueryContextImplementation final : FMassWithEnvironmentContextCommon
+		{
+			FMassQueryContextImplementation(
+				ITypedElementDataStorageInterface::FQueryDescription& InQueryDescription,
+				FMassExecutionContext& InContext, 
+				FExtendedQueryStore& InQueryStore,
+				FEnvironment& InEnvironment)
+				: FMassWithEnvironmentContextCommon(InContext, InEnvironment)
+				, QueryDescription(InQueryDescription)
+				, QueryStore(InQueryStore)
+			{
 		
-	}
+			}
 
-	~FMassQueryContextImplementation() = default;
+			~FMassQueryContextImplementation() = default;
 
-	UObject* GetMutableDependency(const UClass* DependencyClass)
-	{
-		return Context.GetMutableSubsystem<USubsystem>(const_cast<UClass*>(DependencyClass));
-	}
+			UObject* GetMutableDependency(const UClass* DependencyClass)
+			{
+				return Context.GetMutableSubsystem<USubsystem>(const_cast<UClass*>(DependencyClass));
+			}
 
-	const UObject* GetDependency(const UClass* DependencyClass)
-	{
-		return Context.GetSubsystem<USubsystem>(const_cast<UClass*>(DependencyClass));
-	}
+			const UObject* GetDependency(const UClass* DependencyClass)
+			{
+				return Context.GetSubsystem<USubsystem>(const_cast<UClass*>(DependencyClass));
+			}
 	
-	void GetDependencies(TArrayView<UObject*> RetrievedAddresses, TConstArrayView<TWeakObjectPtr<const UClass>> SubsystemTypes,
-		TConstArrayView<ITypedElementDataStorageInterface::EQueryAccessType> AccessTypes)
-	{
-		checkf(RetrievedAddresses.Num() == SubsystemTypes.Num(), TEXT("Unable to retrieve a batch of subsystem as the number of addresses "
-			"doesn't match the number of requested subsystem types."));
+			void GetDependencies(TArrayView<UObject*> RetrievedAddresses, TConstArrayView<TWeakObjectPtr<const UClass>> SubsystemTypes,
+				TConstArrayView<ITypedElementDataStorageInterface::EQueryAccessType> AccessTypes)
+			{
+				checkf(RetrievedAddresses.Num() == SubsystemTypes.Num(), TEXT("Unable to retrieve a batch of subsystem as the number of addresses "
+					"doesn't match the number of requested subsystem types."));
 
-		GetDependenciesUnguarded(RetrievedAddresses.Num(), RetrievedAddresses.GetData(), SubsystemTypes.GetData(), AccessTypes.GetData());
-	}
+				GetDependenciesUnguarded(RetrievedAddresses.Num(), RetrievedAddresses.GetData(), SubsystemTypes.GetData(), AccessTypes.GetData());
+			}
 
-	void GetDependenciesUnguarded(int32 SubsystemCount, UObject** RetrievedAddresses, const TWeakObjectPtr<const UClass>* DependencyTypes,
-		const ITypedElementDataStorageInterface::EQueryAccessType* AccessTypes)
-	{
-		for (int32 Index = 0; Index < SubsystemCount; ++Index)
-		{
-			checkf(DependencyTypes->IsValid(), TEXT("Attempting to retrieve a subsystem that's no longer valid."));
-			*RetrievedAddresses = *AccessTypes == ITypedElementDataStorageInterface::EQueryAccessType::ReadWrite
-				? Context.GetMutableSubsystem<USubsystem>(const_cast<UClass*>(DependencyTypes->Get()))
-				: const_cast<USubsystem*>(Context.GetSubsystem<USubsystem>(const_cast<UClass*>(DependencyTypes->Get())));
-
-			++RetrievedAddresses;
-			++DependencyTypes;
-			++AccessTypes;
-		}
-	}
+			void GetDependenciesUnguarded(int32 SubsystemCount, UObject** RetrievedAddresses, const TWeakObjectPtr<const UClass>* DependencyTypes,
+				const ITypedElementDataStorageInterface::EQueryAccessType* AccessTypes)
+			{
+				for (int32 Index = 0; Index < SubsystemCount; ++Index)
+				{
+					checkf(DependencyTypes->IsValid(), TEXT("Attempting to retrieve a subsystem that's no longer valid."));
+					*RetrievedAddresses = *AccessTypes == ITypedElementDataStorageInterface::EQueryAccessType::ReadWrite
+						? Context.GetMutableSubsystem<USubsystem>(const_cast<UClass*>(DependencyTypes->Get()))
+						: const_cast<USubsystem*>(Context.GetSubsystem<USubsystem>(const_cast<UClass*>(DependencyTypes->Get())));
+
+					++RetrievedAddresses;
+					++DependencyTypes;
+					++AccessTypes;
+				}
+			}
 	
-	TypedElementDataStorage::RowHandle FindIndexedRow(TypedElementDataStorage::IndexHash Index) const
-	{
-		using namespace UE::EditorDataStorage;
+			TypedElementDataStorage::RowHandle FindIndexedRow(TypedElementDataStorage::IndexHash Index) const
+			{
+				using namespace UE::Editor::DataStorage;
 
-		EGlobalLockScope Scope = FGlobalLock::GetLockStatus(EGlobalLockScope::Internal) == EGlobalLockStatus::Unlocked
-			? EGlobalLockScope::Public // There's no internal lock so use a public lock instead.
-			: EGlobalLockScope::Internal; // There's an internal lock set so use that.
-		return Environment.GetIndexTable().FindIndexedRow(Scope, Index);
-	}
+				EGlobalLockScope Scope = FGlobalLock::GetLockStatus(EGlobalLockScope::Internal) == EGlobalLockStatus::Unlocked
+					? EGlobalLockScope::Public // There's no internal lock so use a public lock instead.
+					: EGlobalLockScope::Internal; // There's an internal lock set so use that.
+				return Environment.GetIndexTable().FindIndexedRow(Scope, Index);
+			}
 
-	TypedElementDataStorage::FQueryResult RunQuery(TypedElementQueryHandle Query)
-	{
-		const FTypedElementExtendedQueryStore::Handle Handle(Query);
-		// This can be safely called because there's not callback, which means no columns are accessed, even for select queries.
-		return QueryStore.RunQuery(Context.GetEntityManagerChecked(), Handle);
-	}
+			TypedElementDataStorage::FQueryResult RunQuery(TypedElementQueryHandle Query)
+			{
+				const FExtendedQueryStore::Handle Handle(Query);
+				// This can be safely called because there's not callback, which means no columns are accessed, even for select queries.
+				return QueryStore.RunQuery(Context.GetEntityManagerChecked(), Handle);
+			}
 
-	TypedElementDataStorage::FQueryResult RunSubquery(int32 SubqueryIndex)
-	{
-		return SubqueryIndex < QueryDescription.Subqueries.Num() ?
-			RunQuery(QueryDescription.Subqueries[SubqueryIndex]) :
-			TypedElementDataStorage::FQueryResult{};
-	}
+			TypedElementDataStorage::FQueryResult RunSubquery(int32 SubqueryIndex)
+			{
+				return SubqueryIndex < QueryDescription.Subqueries.Num() ?
+					RunQuery(QueryDescription.Subqueries[SubqueryIndex]) :
+					TypedElementDataStorage::FQueryResult{};
+			}
 
-	TypedElementDataStorage::FQueryResult RunSubquery(int32 SubqueryIndex, TypedElementDataStorage::SubqueryCallbackRef Callback)
-	{
-		if (SubqueryIndex < QueryDescription.Subqueries.Num())
-		{
-			const TypedElementQueryHandle SubqueryHandle = QueryDescription.Subqueries[SubqueryIndex];
-			const FTypedElementExtendedQueryStore::Handle StorageHandle(SubqueryHandle);
-			return QueryStore.RunQuery(Context.GetEntityManagerChecked(), Environment, Context, StorageHandle, Callback);
-		}
-		else
-		{
-			return TypedElementDataStorage::FQueryResult{};
-		}
-	}
+			TypedElementDataStorage::FQueryResult RunSubquery(int32 SubqueryIndex, TypedElementDataStorage::SubqueryCallbackRef Callback)
+			{
+				if (SubqueryIndex < QueryDescription.Subqueries.Num())
+				{
+					const TypedElementQueryHandle SubqueryHandle = QueryDescription.Subqueries[SubqueryIndex];
+					const FExtendedQueryStore::Handle StorageHandle(SubqueryHandle);
+					return QueryStore.RunQuery(Context.GetEntityManagerChecked(), Environment, Context, StorageHandle, Callback);
+				}
+				else
+				{
+					return TypedElementDataStorage::FQueryResult{};
+				}
+			}
 
-	TypedElementDataStorage::FQueryResult RunSubquery(int32 SubqueryIndex, TypedElementDataStorage::RowHandle Row,
-		TypedElementDataStorage::SubqueryCallbackRef Callback)
-	{
-		if (SubqueryIndex < QueryDescription.Subqueries.Num())
-		{
-			const TypedElementQueryHandle SubqueryHandle = QueryDescription.Subqueries[SubqueryIndex];
-			const FTypedElementExtendedQueryStore::Handle StorageHandle(SubqueryHandle);
-			return QueryStore.RunQuery(Context.GetEntityManagerChecked(), Environment, Context, StorageHandle, Row, Callback);
-		}
-		else
-		{
-			return TypedElementDataStorage::FQueryResult{};
-		}
-	}
+			TypedElementDataStorage::FQueryResult RunSubquery(int32 SubqueryIndex, TypedElementDataStorage::RowHandle Row,
+				TypedElementDataStorage::SubqueryCallbackRef Callback)
+			{
+				if (SubqueryIndex < QueryDescription.Subqueries.Num())
+				{
+					const TypedElementQueryHandle SubqueryHandle = QueryDescription.Subqueries[SubqueryIndex];
+					const FExtendedQueryStore::Handle StorageHandle(SubqueryHandle);
+					return QueryStore.RunQuery(Context.GetEntityManagerChecked(), Environment, Context, StorageHandle, Row, Callback);
+				}
+				else
+				{
+					return TypedElementDataStorage::FQueryResult{};
+				}
+			}
 	
-	ITypedElementDataStorageInterface::FQueryDescription& QueryDescription;
-	FTypedElementExtendedQueryStore& QueryStore;
-};
+			ITypedElementDataStorageInterface::FQueryDescription& QueryDescription;
+			FExtendedQueryStore& QueryStore;
+		};
 
-struct FMassContextForwarder final : public ITypedElementDataStorageInterface::IQueryContext
-{
-	FMassContextForwarder(
-		ITypedElementDataStorageInterface::FQueryDescription& InQueryDescription,
-		FMassExecutionContext& InContext, 
-		FTypedElementExtendedQueryStore& InQueryStore,
-		FTypedElementDatabaseEnvironment& InEnvironment)
-			: Implementation(InQueryDescription, InContext, InQueryStore, InEnvironment)
-	{}
+		struct FMassContextForwarder final : public ITypedElementDataStorageInterface::IQueryContext
+		{
+			FMassContextForwarder(
+				ITypedElementDataStorageInterface::FQueryDescription& InQueryDescription,
+				FMassExecutionContext& InContext, 
+				FExtendedQueryStore& InQueryStore,
+				FEnvironment& InEnvironment)
+					: Implementation(InQueryDescription, InContext, InQueryStore, InEnvironment)
+			{}
 	
-	uint32 GetRowCount() const override { return Implementation.GetRowCount(); }
-	TConstArrayView<TypedElementDataStorage::RowHandle> GetRowHandles() const override { return Implementation.GetRowHandles(); }
-	const void* GetColumn(const UScriptStruct* ColumnType) const override { return Implementation.GetColumn(ColumnType); }
-	void* GetMutableColumn(const UScriptStruct* ColumnType) override { return Implementation.GetMutableColumn(ColumnType); }
-	void GetColumns(TArrayView<char*> RetrievedAddresses, TConstArrayView<TWeakObjectPtr<const UScriptStruct>> ColumnTypes, TConstArrayView<TypedElementDataStorage::EQueryAccessType> AccessTypes) override { return Implementation.GetColumns(RetrievedAddresses, ColumnTypes, AccessTypes); }
-	void GetColumnsUnguarded(int32 TypeCount, char** RetrievedAddresses, const TWeakObjectPtr<const UScriptStruct>* ColumnTypes, const TypedElementDataStorage::EQueryAccessType* AccessTypes) override { return Implementation.GetColumnsUnguarded(TypeCount, RetrievedAddresses, ColumnTypes, AccessTypes); }
-	bool HasColumn(const UScriptStruct* ColumnType) const override { return Implementation.HasColumn(ColumnType); }
-	bool HasColumn(TypedElementDataStorage::RowHandle Row, const UScriptStruct* ColumnType) const override { return Implementation.HasColumn(Row, ColumnType); }
-	uint64 GetUpdateCycleId() const override { return Implementation.GetUpdateCycleId(); }
-	bool IsRowAvailable(TypedElementDataStorage::RowHandle Row) const override { return Implementation.IsRowAvailable(Row); }
-	bool IsRowAssigned(TypedElementDataStorage::RowHandle Row) const override { return Implementation.IsRowAssigned(Row); }
-	void ActivateQueries(FName ActivationName) override  { return Implementation.ActivateQueries(ActivationName); }
-	TypedElementDataStorage::RowHandle AddRow(TypedElementDataStorage::TableHandle Table) override  { return Implementation.AddRow(Table); }
-	void RemoveRow(TypedElementDataStorage::RowHandle Row) override { return Implementation.RemoveRow(Row); }
-	void RemoveRows(TConstArrayView<TypedElementDataStorage::RowHandle> Rows) override  { return Implementation.RemoveRows(Rows); }
-	void AddColumns(TypedElementDataStorage::RowHandle Row, TConstArrayView<const UScriptStruct*> ColumnTypes) override { return Implementation.AddColumns(Row, ColumnTypes); }
-	void AddColumns(TConstArrayView<TypedElementDataStorage::RowHandle> Rows, TConstArrayView<const UScriptStruct*> ColumnTypes) override  { return Implementation.AddColumns(Rows, ColumnTypes); }
-	void* AddColumnUninitialized(TypedElementDataStorage::RowHandle Row, const UScriptStruct* ColumnType) override  { return Implementation.AddColumnUninitialized(Row, ColumnType); }
-	void* AddColumnUninitialized(TypedElementDataStorage::RowHandle Row, const UScriptStruct* ObjectType, ObjectCopyOrMove Relocator) override { return Implementation.AddColumnUninitialized(Row, ObjectType, Relocator); }
-	void RemoveColumns(TypedElementDataStorage::RowHandle Row, TConstArrayView<const UScriptStruct*> ColumnTypes) override { Implementation.RemoveColumns(Row, ColumnTypes); }
-	void RemoveColumns(TConstArrayView<TypedElementDataStorage::RowHandle> Rows, TConstArrayView<const UScriptStruct*> ColumnTypes) override { Implementation.RemoveColumns(Rows, ColumnTypes); }
+			uint32 GetRowCount() const override { return Implementation.GetRowCount(); }
+			TConstArrayView<TypedElementDataStorage::RowHandle> GetRowHandles() const override { return Implementation.GetRowHandles(); }
+			const void* GetColumn(const UScriptStruct* ColumnType) const override { return Implementation.GetColumn(ColumnType); }
+			void* GetMutableColumn(const UScriptStruct* ColumnType) override { return Implementation.GetMutableColumn(ColumnType); }
+			void GetColumns(TArrayView<char*> RetrievedAddresses, TConstArrayView<TWeakObjectPtr<const UScriptStruct>> ColumnTypes, TConstArrayView<TypedElementDataStorage::EQueryAccessType> AccessTypes) override { return Implementation.GetColumns(RetrievedAddresses, ColumnTypes, AccessTypes); }
+			void GetColumnsUnguarded(int32 TypeCount, char** RetrievedAddresses, const TWeakObjectPtr<const UScriptStruct>* ColumnTypes, const TypedElementDataStorage::EQueryAccessType* AccessTypes) override { return Implementation.GetColumnsUnguarded(TypeCount, RetrievedAddresses, ColumnTypes, AccessTypes); }
+			bool HasColumn(const UScriptStruct* ColumnType) const override { return Implementation.HasColumn(ColumnType); }
+			bool HasColumn(TypedElementDataStorage::RowHandle Row, const UScriptStruct* ColumnType) const override { return Implementation.HasColumn(Row, ColumnType); }
+			uint64 GetUpdateCycleId() const override { return Implementation.GetUpdateCycleId(); }
+			bool IsRowAvailable(TypedElementDataStorage::RowHandle Row) const override { return Implementation.IsRowAvailable(Row); }
+			bool IsRowAssigned(TypedElementDataStorage::RowHandle Row) const override { return Implementation.IsRowAssigned(Row); }
+			void ActivateQueries(FName ActivationName) override  { return Implementation.ActivateQueries(ActivationName); }
+			TypedElementDataStorage::RowHandle AddRow(TypedElementDataStorage::TableHandle Table) override  { return Implementation.AddRow(Table); }
+			void RemoveRow(TypedElementDataStorage::RowHandle Row) override { return Implementation.RemoveRow(Row); }
+			void RemoveRows(TConstArrayView<TypedElementDataStorage::RowHandle> Rows) override  { return Implementation.RemoveRows(Rows); }
+			void AddColumns(TypedElementDataStorage::RowHandle Row, TConstArrayView<const UScriptStruct*> ColumnTypes) override { return Implementation.AddColumns(Row, ColumnTypes); }
+			void AddColumns(TConstArrayView<TypedElementDataStorage::RowHandle> Rows, TConstArrayView<const UScriptStruct*> ColumnTypes) override  { return Implementation.AddColumns(Rows, ColumnTypes); }
+			void* AddColumnUninitialized(TypedElementDataStorage::RowHandle Row, const UScriptStruct* ColumnType) override  { return Implementation.AddColumnUninitialized(Row, ColumnType); }
+			void* AddColumnUninitialized(TypedElementDataStorage::RowHandle Row, const UScriptStruct* ObjectType, ObjectCopyOrMove Relocator) override { return Implementation.AddColumnUninitialized(Row, ObjectType, Relocator); }
+			void RemoveColumns(TypedElementDataStorage::RowHandle Row, TConstArrayView<const UScriptStruct*> ColumnTypes) override { Implementation.RemoveColumns(Row, ColumnTypes); }
+			void RemoveColumns(TConstArrayView<TypedElementDataStorage::RowHandle> Rows, TConstArrayView<const UScriptStruct*> ColumnTypes) override { Implementation.RemoveColumns(Rows, ColumnTypes); }
 	
-	const UObject* GetDependency(const UClass* DependencyClass) override { return Implementation.GetDependency(DependencyClass); }
-	UObject* GetMutableDependency(const UClass* DependencyClass) override { return Implementation.GetMutableDependency(DependencyClass); }
-	void GetDependencies(TArrayView<UObject*> RetrievedAddresses, TConstArrayView<TWeakObjectPtr<const UClass>> DependencyTypes, TConstArrayView<TypedElementDataStorage::EQueryAccessType> AccessTypes) override { return Implementation.GetDependencies(RetrievedAddresses, DependencyTypes, AccessTypes); }
-	TypedElementDataStorage::RowHandle FindIndexedRow(TypedElementDataStorage::IndexHash Index) const override { return Implementation.FindIndexedRow(Index); }
-	TypedElementDataStorage::FQueryResult RunQuery(TypedElementDataStorage::QueryHandle Query) override { return Implementation.RunQuery(Query); }
-	TypedElementDataStorage::FQueryResult RunSubquery(int32 SubqueryIndex) override { return Implementation.RunSubquery(SubqueryIndex); }
-	TypedElementDataStorage::FQueryResult RunSubquery(int32 SubqueryIndex, TypedElementDataStorage::SubqueryCallbackRef Callback) override { return Implementation.RunSubquery(SubqueryIndex, Callback); }
-	TypedElementDataStorage::FQueryResult RunSubquery(int32 SubqueryIndex, TypedElementDataStorage::RowHandle Row, TypedElementDataStorage::SubqueryCallbackRef Callback) override { return Implementation.RunSubquery(SubqueryIndex, Row, Callback); }
+			const UObject* GetDependency(const UClass* DependencyClass) override { return Implementation.GetDependency(DependencyClass); }
+			UObject* GetMutableDependency(const UClass* DependencyClass) override { return Implementation.GetMutableDependency(DependencyClass); }
+			void GetDependencies(TArrayView<UObject*> RetrievedAddresses, TConstArrayView<TWeakObjectPtr<const UClass>> DependencyTypes, TConstArrayView<TypedElementDataStorage::EQueryAccessType> AccessTypes) override { return Implementation.GetDependencies(RetrievedAddresses, DependencyTypes, AccessTypes); }
+			TypedElementDataStorage::RowHandle FindIndexedRow(TypedElementDataStorage::IndexHash Index) const override { return Implementation.FindIndexedRow(Index); }
+			TypedElementDataStorage::FQueryResult RunQuery(TypedElementDataStorage::QueryHandle Query) override { return Implementation.RunQuery(Query); }
+			TypedElementDataStorage::FQueryResult RunSubquery(int32 SubqueryIndex) override { return Implementation.RunSubquery(SubqueryIndex); }
+			TypedElementDataStorage::FQueryResult RunSubquery(int32 SubqueryIndex, TypedElementDataStorage::SubqueryCallbackRef Callback) override { return Implementation.RunSubquery(SubqueryIndex, Callback); }
+			TypedElementDataStorage::FQueryResult RunSubquery(int32 SubqueryIndex, TypedElementDataStorage::RowHandle Row, TypedElementDataStorage::SubqueryCallbackRef Callback) override { return Implementation.RunSubquery(SubqueryIndex, Row, Callback); }
 	
-	FMassQueryContextImplementation Implementation;
-};
+			FMassQueryContextImplementation Implementation;
+		};
+	} // namespace Processors::Private
 
-/**
- * FPhasePreOrPostAmbleExecutor
- */
-FPhasePreOrPostAmbleExecutor::FPhasePreOrPostAmbleExecutor(FMassEntityManager& EntityManager, float DeltaTime)
-	: Context(EntityManager, DeltaTime)
-{
-	Context.SetDeferredCommandBuffer(MakeShared<FMassCommandBuffer>());
-}
+	/**
+	 * FPhasePreOrPostAmbleExecutor
+	 */
+	FPhasePreOrPostAmbleExecutor::FPhasePreOrPostAmbleExecutor(FMassEntityManager& EntityManager, float DeltaTime)
+		: Context(EntityManager, DeltaTime)
+	{
+		Context.SetDeferredCommandBuffer(MakeShared<FMassCommandBuffer>());
+	}
 
-FPhasePreOrPostAmbleExecutor::~FPhasePreOrPostAmbleExecutor()
-{
-	Context.FlushDeferred();
-}
+	FPhasePreOrPostAmbleExecutor::~FPhasePreOrPostAmbleExecutor()
+	{
+		Context.FlushDeferred();
+	}
 
-void FPhasePreOrPostAmbleExecutor::ExecuteQuery(
-	ITypedElementDataStorageInterface::FQueryDescription& Description, 
-	FTypedElementExtendedQueryStore& QueryStore, 
-	FTypedElementDatabaseEnvironment& Environment,
-	FMassEntityQuery& NativeQuery,
-	ITypedElementDataStorageInterface::QueryCallbackRef Callback)
-{
-	if (Description.Callback.ActivationCount > 0)
+	void FPhasePreOrPostAmbleExecutor::ExecuteQuery(
+		ITypedElementDataStorageInterface::FQueryDescription& Description,
+		FExtendedQueryStore& QueryStore,
+		FEnvironment& Environment,
+		FMassEntityQuery& NativeQuery,
+		ITypedElementDataStorageInterface::QueryCallbackRef Callback)
 	{
-		NativeQuery.ForEachEntityChunk(Context.GetEntityManagerChecked(), Context,
-			[&Callback, &QueryStore, &Environment, &Description](FMassExecutionContext& ExecutionContext)
-			{
-				if (FTypedElementQueryProcessorData::PrepareCachedDependenciesOnQuery(Description, ExecutionContext))
+		if (Description.Callback.ActivationCount > 0)
+		{
+			NativeQuery.ForEachEntityChunk(Context.GetEntityManagerChecked(), Context,
+				[&Callback, &QueryStore, &Environment, &Description](FMassExecutionContext& ExecutionContext)
 				{
-					FMassContextForwarder QueryContext(Description, ExecutionContext, QueryStore, Environment);
-					Callback(Description, QueryContext);
+					if (FTypedElementQueryProcessorData::PrepareCachedDependenciesOnQuery(Description, ExecutionContext))
+					{
+						Processors::Private::FMassContextForwarder QueryContext(Description, ExecutionContext, QueryStore, Environment);
+						Callback(Description, QueryContext);
+					}
 				}
-			}
-		);
+			);
+		}
 	}
-}
-
-
+} // namespace UE::Editor::DataStorage
 
 /**
  * FTypedElementQueryProcessorData
@@ -716,10 +720,10 @@ FTypedElementQueryProcessorData::FTypedElementQueryProcessorData(UMassProcessor&
 
 bool FTypedElementQueryProcessorData::CommonQueryConfiguration(
 	UMassProcessor& InOwner,
-	FTypedElementExtendedQuery& InQuery,
-	FTypedElementExtendedQueryStore::Handle InQueryHandle,
-	FTypedElementExtendedQueryStore& InQueryStore,
-	FTypedElementDatabaseEnvironment& InEnvironment,
+	FExtendedQuery& InQuery,
+	FExtendedQueryStore::Handle InQueryHandle,
+	FExtendedQueryStore& InQueryStore,
+	FEnvironment& InEnvironment,
 	TArrayView<FMassEntityQuery> Subqueries)
 {
 	ParentQuery = InQueryHandle;
@@ -733,8 +737,8 @@ bool FTypedElementQueryProcessorData::CommonQueryConfiguration(
 		int32 CurrentSubqueryIndex = 0;
 		for (TypedElementQueryHandle SubqueryHandle : InQuery.Description.Subqueries)
 		{
-			const FTypedElementExtendedQueryStore::Handle SubqueryStoreHandle(SubqueryHandle);
-			if (const FTypedElementExtendedQuery* Subquery = InQueryStore.Get(SubqueryStoreHandle))
+			const FExtendedQueryStore::Handle SubqueryStoreHandle(SubqueryHandle);
+			if (const FExtendedQuery* Subquery = InQueryStore.Get(SubqueryStoreHandle))
 			{
 				if (ensureMsgf(Subquery->NativeQuery.CheckValidity(), TEXT("Provided subquery isn't valid. This can be because it couldn't be "
 					"constructed properly or because it's been bound to a callback.")))
@@ -778,7 +782,7 @@ EMassProcessingPhase FTypedElementQueryProcessorData::MapToMassProcessingPhase(I
 
 FString FTypedElementQueryProcessorData::GetProcessorName() const
 {
-	if (const FTypedElementExtendedQuery* StoredQuery = QueryStore ? QueryStore->Get(ParentQuery) : nullptr)
+	if (const FExtendedQuery* StoredQuery = QueryStore ? QueryStore->Get(ParentQuery) : nullptr)
 	{
 		return StoredQuery->Description.Callback.Name.ToString();
 	}
@@ -791,7 +795,7 @@ FString FTypedElementQueryProcessorData::GetProcessorName() const
 void FTypedElementQueryProcessorData::DebugOutputDescription(FOutputDevice& Ar, int32 Indent) const
 {
 #if WITH_MASSENTITY_DEBUG
-	if (const FTypedElementExtendedQuery* StoredQuery = QueryStore ? QueryStore->Get(ParentQuery) : nullptr)
+	if (const FExtendedQuery* StoredQuery = QueryStore ? QueryStore->Get(ParentQuery) : nullptr)
 	{
 		const ITypedElementDataStorageInterface::FQueryDescription& Description = StoredQuery->Description;
 		const ITypedElementDataStorageInterface::FQueryDescription::FCallbackData& Callback = Description.Callback;
@@ -872,7 +876,7 @@ TypedElementDataStorage::FQueryResult FTypedElementQueryProcessorData::Execute(
 	TypedElementDataStorage::FQueryDescription& Description,
 	FMassEntityQuery& NativeQuery, 
 	FMassEntityManager& EntityManager,
-	FTypedElementDatabaseEnvironment& Environment)
+	FEnvironment& Environment)
 {
 	ITypedElementDataStorageInterface::FQueryResult Result;
 	Result.Completed = ITypedElementDataStorageInterface::FQueryResult::ECompletion::Fully;
@@ -885,7 +889,7 @@ TypedElementDataStorage::FQueryResult FTypedElementQueryProcessorData::Execute(
 			[&Result, &Callback, &Description](FMassExecutionContext& Context)
 			{
 				// No need to cache any subsystem dependencies as these are not accessible from a direct query.
-				FMassDirectContextForwarder QueryContext(Context);
+				UE::Editor::DataStorage::Processors::Private::FMassDirectContextForwarder QueryContext(Context);
 				Callback(Description, QueryContext);
 				Result.Count += Context.GetNumEntities();
 			}
@@ -899,7 +903,7 @@ TypedElementDataStorage::FQueryResult FTypedElementQueryProcessorData::Execute(
 	TypedElementDataStorage::FQueryDescription& Description,
 	FMassEntityQuery& NativeQuery,
 	FMassEntityManager& EntityManager,
-	FTypedElementDatabaseEnvironment& Environment,
+	FEnvironment& Environment,
 	FMassExecutionContext& ParentContext)
 {
 	ITypedElementDataStorageInterface::FQueryResult Result;
@@ -915,7 +919,7 @@ TypedElementDataStorage::FQueryResult FTypedElementQueryProcessorData::Execute(
 			[&Result, &Callback, &Description, &Environment](FMassExecutionContext& Context)
 			{
 				// No need to cache any subsystem dependencies as these are not accessible from a subquery.
-				FMassSubqueryContextForwarder QueryContext(Context, Environment);
+				UE::Editor::DataStorage::Processors::Private::FMassSubqueryContextForwarder QueryContext(Context, Environment);
 				Callback(Description, QueryContext);
 				Result.Count += Context.GetNumEntities();
 			}
@@ -930,7 +934,7 @@ TypedElementDataStorage::FQueryResult FTypedElementQueryProcessorData::Execute(
 	TypedElementDataStorage::RowHandle RowHandle,
 	FMassEntityQuery& NativeQuery,
 	FMassEntityManager& EntityManager,
-	FTypedElementDatabaseEnvironment& Environment,
+	FEnvironment& Environment,
 	FMassExecutionContext& ParentContext)
 {
 	ITypedElementDataStorageInterface::FQueryResult Result;
@@ -949,7 +953,7 @@ TypedElementDataStorage::FQueryResult FTypedElementQueryProcessorData::Execute(
 			[&Result, &Callback, &Description, &Environment, &EntityManager, RowHandle](FMassExecutionContext& Context)
 			{
 				// No need to cache any subsystem dependencies as these are not accessible from a subquery.
-				FMassSubqueryContextForwarder QueryContext(Context, Environment);
+				UE::Editor::DataStorage::Processors::Private::FMassSubqueryContextForwarder QueryContext(Context, Environment);
 				Callback(Description, QueryContext);
 				Result.Count += Context.GetNumEntities();
 			}
@@ -961,7 +965,7 @@ TypedElementDataStorage::FQueryResult FTypedElementQueryProcessorData::Execute(
 
 void FTypedElementQueryProcessorData::Execute(FMassEntityManager& EntityManager, FMassExecutionContext& Context)
 {
-	FTypedElementExtendedQuery* StoredQuery = QueryStore->GetMutable(ParentQuery);
+	FExtendedQuery* StoredQuery = QueryStore->GetMutable(ParentQuery);
 	checkf(StoredQuery, TEXT("A query callback was registered for execution without an associated query."));
 	
 	ITypedElementDataStorageInterface::FQueryDescription& Description = StoredQuery->Description;
@@ -972,7 +976,7 @@ void FTypedElementQueryProcessorData::Execute(FMassEntityManager& EntityManager,
 			{
 				if (PrepareCachedDependenciesOnQuery(Description, Context))
 				{
-					FMassContextForwarder QueryContext(Description, Context, *QueryStore, *Environment);
+					UE::Editor::DataStorage::Processors::Private::FMassContextForwarder QueryContext(Description, Context, *QueryStore, *Environment);
 					Description.Callback.Function(Description, QueryContext);
 				}
 			}
@@ -999,10 +1003,10 @@ FMassEntityQuery& UTypedElementQueryProcessorCallbackAdapterProcessorBase::GetQu
 }
 
 bool UTypedElementQueryProcessorCallbackAdapterProcessorBase::ConfigureQueryCallback(
-	FTypedElementExtendedQuery& Query,
-	FTypedElementExtendedQueryStore::Handle QueryHandle,
-	FTypedElementExtendedQueryStore& QueryStore,
-	FTypedElementDatabaseEnvironment& Environment)
+	FExtendedQuery& Query,
+	FExtendedQueryStore::Handle QueryHandle,
+	FExtendedQueryStore& QueryStore,
+	FEnvironment& Environment)
 {
 	return ConfigureQueryCallbackData(Query, QueryHandle, QueryStore, Environment, {});
 }
@@ -1017,10 +1021,10 @@ bool UTypedElementQueryProcessorCallbackAdapterProcessorBase::ShouldAllowQueryBa
 }
 
 bool UTypedElementQueryProcessorCallbackAdapterProcessorBase::ConfigureQueryCallbackData(
-	FTypedElementExtendedQuery& Query,
-	FTypedElementExtendedQueryStore::Handle QueryHandle,
-	FTypedElementExtendedQueryStore& QueryStore,
-	FTypedElementDatabaseEnvironment& Environment,
+	FExtendedQuery& Query,
+	FExtendedQueryStore::Handle QueryHandle,
+	FExtendedQueryStore& QueryStore,
+	FEnvironment& Environment,
 	TArrayView<FMassEntityQuery> Subqueries)
 {
 	bool Result = Data.CommonQueryConfiguration(*this, Query, QueryHandle, QueryStore, Environment, Subqueries);
@@ -1066,49 +1070,49 @@ void UTypedElementQueryProcessorCallbackAdapterProcessorBase::Execute(FMassEntit
 }
 
 bool UTypedElementQueryProcessorCallbackAdapterProcessorWith1Subquery::ConfigureQueryCallback(
-	FTypedElementExtendedQuery& Query, FTypedElementExtendedQueryStore::Handle QueryHandle, FTypedElementExtendedQueryStore& QueryStore, FTypedElementDatabaseEnvironment& Environment)
+	FExtendedQuery& Query, FExtendedQueryStore::Handle QueryHandle, FExtendedQueryStore& QueryStore, FEnvironment& Environment)
 {
 	return ConfigureQueryCallbackData(Query, QueryHandle, QueryStore, Environment, NativeSubqueries);
 }
 
 bool UTypedElementQueryProcessorCallbackAdapterProcessorWith2Subqueries::ConfigureQueryCallback(
-	FTypedElementExtendedQuery& Query, FTypedElementExtendedQueryStore::Handle QueryHandle, FTypedElementExtendedQueryStore& QueryStore, FTypedElementDatabaseEnvironment& Environment)
+	FExtendedQuery& Query, FExtendedQueryStore::Handle QueryHandle, FExtendedQueryStore& QueryStore, FEnvironment& Environment)
 {
 	return ConfigureQueryCallbackData(Query, QueryHandle, QueryStore, Environment, NativeSubqueries);
 }
 
 bool UTypedElementQueryProcessorCallbackAdapterProcessorWith3Subqueries::ConfigureQueryCallback(
-	FTypedElementExtendedQuery& Query, FTypedElementExtendedQueryStore::Handle QueryHandle, FTypedElementExtendedQueryStore& QueryStore, FTypedElementDatabaseEnvironment& Environment)
+	FExtendedQuery& Query, FExtendedQueryStore::Handle QueryHandle, FExtendedQueryStore& QueryStore, FEnvironment& Environment)
 {
 	return ConfigureQueryCallbackData(Query, QueryHandle, QueryStore, Environment, NativeSubqueries);
 }
 
 bool UTypedElementQueryProcessorCallbackAdapterProcessorWith4Subqueries::ConfigureQueryCallback(
-	FTypedElementExtendedQuery& Query, FTypedElementExtendedQueryStore::Handle QueryHandle, FTypedElementExtendedQueryStore& QueryStore, FTypedElementDatabaseEnvironment& Environment)
+	FExtendedQuery& Query, FExtendedQueryStore::Handle QueryHandle, FExtendedQueryStore& QueryStore, FEnvironment& Environment)
 {
 	return ConfigureQueryCallbackData(Query, QueryHandle, QueryStore, Environment, NativeSubqueries);
 }
 
 bool UTypedElementQueryProcessorCallbackAdapterProcessorWith5Subqueries::ConfigureQueryCallback(
-	FTypedElementExtendedQuery& Query, FTypedElementExtendedQueryStore::Handle QueryHandle, FTypedElementExtendedQueryStore& QueryStore, FTypedElementDatabaseEnvironment& Environment)
+	FExtendedQuery& Query, FExtendedQueryStore::Handle QueryHandle, FExtendedQueryStore& QueryStore, FEnvironment& Environment)
 {
 	return ConfigureQueryCallbackData(Query, QueryHandle, QueryStore, Environment, NativeSubqueries);
 }
 
 bool UTypedElementQueryProcessorCallbackAdapterProcessorWith6Subqueries::ConfigureQueryCallback(
-	FTypedElementExtendedQuery& Query, FTypedElementExtendedQueryStore::Handle QueryHandle, FTypedElementExtendedQueryStore& QueryStore, FTypedElementDatabaseEnvironment& Environment)
+	FExtendedQuery& Query, FExtendedQueryStore::Handle QueryHandle, FExtendedQueryStore& QueryStore, FEnvironment& Environment)
 {
 	return ConfigureQueryCallbackData(Query, QueryHandle, QueryStore, Environment, NativeSubqueries);
 }
 
 bool UTypedElementQueryProcessorCallbackAdapterProcessorWith7Subqueries::ConfigureQueryCallback(
-	FTypedElementExtendedQuery& Query, FTypedElementExtendedQueryStore::Handle QueryHandle, FTypedElementExtendedQueryStore& QueryStore, FTypedElementDatabaseEnvironment& Environment)
+	FExtendedQuery& Query, FExtendedQueryStore::Handle QueryHandle, FExtendedQueryStore& QueryStore, FEnvironment& Environment)
 {
 	return ConfigureQueryCallbackData(Query, QueryHandle, QueryStore, Environment, NativeSubqueries);
 }
 
 bool UTypedElementQueryProcessorCallbackAdapterProcessorWith8Subqueries::ConfigureQueryCallback(
-	FTypedElementExtendedQuery& Query, FTypedElementExtendedQueryStore::Handle QueryHandle, FTypedElementExtendedQueryStore& QueryStore, FTypedElementDatabaseEnvironment& Environment)
+	FExtendedQuery& Query, FExtendedQueryStore::Handle QueryHandle, FExtendedQueryStore& QueryStore, FEnvironment& Environment)
 {
 	return ConfigureQueryCallbackData(Query, QueryHandle, QueryStore, Environment, NativeSubqueries);
 }
@@ -1141,16 +1145,16 @@ EMassObservedOperation UTypedElementQueryObserverCallbackAdapterProcessorBase::G
 }
 
 bool UTypedElementQueryObserverCallbackAdapterProcessorBase::ConfigureQueryCallback(
-	FTypedElementExtendedQuery& Query, FTypedElementExtendedQueryStore::Handle QueryHandle, FTypedElementExtendedQueryStore& QueryStore, FTypedElementDatabaseEnvironment& Environment)
+	FExtendedQuery& Query, FExtendedQueryStore::Handle QueryHandle, FExtendedQueryStore& QueryStore, FEnvironment& Environment)
 {
 	return ConfigureQueryCallbackData(Query, QueryHandle, QueryStore, Environment, {});
 }
 
 bool UTypedElementQueryObserverCallbackAdapterProcessorBase::ConfigureQueryCallbackData(
-	FTypedElementExtendedQuery& Query,
-	FTypedElementExtendedQueryStore::Handle QueryHandle,
-	FTypedElementExtendedQueryStore& QueryStore,
-	FTypedElementDatabaseEnvironment& Environment, TArrayView<FMassEntityQuery> Subqueries)
+	FExtendedQuery& Query,
+	FExtendedQueryStore::Handle QueryHandle,
+	FExtendedQueryStore& QueryStore,
+	FEnvironment& Environment, TArrayView<FMassEntityQuery> Subqueries)
 {
 	bool Result = Data.CommonQueryConfiguration(*this, Query, QueryHandle, QueryStore, Environment, Subqueries);
 
@@ -1225,49 +1229,49 @@ void UTypedElementQueryObserverCallbackAdapterProcessorBase::Execute(FMassEntity
 }
 
 bool UTypedElementQueryObserverCallbackAdapterProcessorWith1Subquery::ConfigureQueryCallback(
-	FTypedElementExtendedQuery& Query, FTypedElementExtendedQueryStore::Handle QueryHandle, FTypedElementExtendedQueryStore& QueryStore, FTypedElementDatabaseEnvironment& Environment)
+	FExtendedQuery& Query, FExtendedQueryStore::Handle QueryHandle, FExtendedQueryStore& QueryStore, FEnvironment& Environment)
 {
 	return ConfigureQueryCallbackData(Query, QueryHandle, QueryStore, Environment, NativeSubqueries);
 }
 
 bool UTypedElementQueryObserverCallbackAdapterProcessorWith2Subqueries::ConfigureQueryCallback(
-	FTypedElementExtendedQuery& Query, FTypedElementExtendedQueryStore::Handle QueryHandle, FTypedElementExtendedQueryStore& QueryStore, FTypedElementDatabaseEnvironment& Environment)
+	FExtendedQuery& Query, FExtendedQueryStore::Handle QueryHandle, FExtendedQueryStore& QueryStore, FEnvironment& Environment)
 {
 	return ConfigureQueryCallbackData(Query, QueryHandle, QueryStore, Environment, NativeSubqueries);
 }
 
 bool UTypedElementQueryObserverCallbackAdapterProcessorWith3Subqueries::ConfigureQueryCallback(
-	FTypedElementExtendedQuery& Query, FTypedElementExtendedQueryStore::Handle QueryHandle, FTypedElementExtendedQueryStore& QueryStore, FTypedElementDatabaseEnvironment& Environment)
+	FExtendedQuery& Query, FExtendedQueryStore::Handle QueryHandle, FExtendedQueryStore& QueryStore, FEnvironment& Environment)
 {
 	return ConfigureQueryCallbackData(Query, QueryHandle, QueryStore, Environment, NativeSubqueries);
 }
 
 bool UTypedElementQueryObserverCallbackAdapterProcessorWith4Subqueries::ConfigureQueryCallback(
-	FTypedElementExtendedQuery& Query, FTypedElementExtendedQueryStore::Handle QueryHandle, FTypedElementExtendedQueryStore& QueryStore, FTypedElementDatabaseEnvironment& Environment)
+	FExtendedQuery& Query, FExtendedQueryStore::Handle QueryHandle, FExtendedQueryStore& QueryStore, FEnvironment& Environment)
 {
 	return ConfigureQueryCallbackData(Query, QueryHandle, QueryStore, Environment, NativeSubqueries);
 }
 
 bool UTypedElementQueryObserverCallbackAdapterProcessorWith5Subqueries::ConfigureQueryCallback(
-	FTypedElementExtendedQuery& Query, FTypedElementExtendedQueryStore::Handle QueryHandle, FTypedElementExtendedQueryStore& QueryStore, FTypedElementDatabaseEnvironment& Environment)
+	FExtendedQuery& Query, FExtendedQueryStore::Handle QueryHandle, FExtendedQueryStore& QueryStore, FEnvironment& Environment)
 {
 	return ConfigureQueryCallbackData(Query, QueryHandle, QueryStore, Environment, NativeSubqueries);
 }
 
 bool UTypedElementQueryObserverCallbackAdapterProcessorWith6Subqueries::ConfigureQueryCallback(
-	FTypedElementExtendedQuery& Query, FTypedElementExtendedQueryStore::Handle QueryHandle, FTypedElementExtendedQueryStore& QueryStore, FTypedElementDatabaseEnvironment& Environment)
+	FExtendedQuery& Query, FExtendedQueryStore::Handle QueryHandle, FExtendedQueryStore& QueryStore, FEnvironment& Environment)
 {
 	return ConfigureQueryCallbackData(Query, QueryHandle, QueryStore, Environment, NativeSubqueries);
 }
 
 bool UTypedElementQueryObserverCallbackAdapterProcessorWith7Subqueries::ConfigureQueryCallback(
-	FTypedElementExtendedQuery& Query, FTypedElementExtendedQueryStore::Handle QueryHandle, FTypedElementExtendedQueryStore& QueryStore, FTypedElementDatabaseEnvironment& Environment)
+	FExtendedQuery& Query, FExtendedQueryStore::Handle QueryHandle, FExtendedQueryStore& QueryStore, FEnvironment& Environment)
 {
 	return ConfigureQueryCallbackData(Query, QueryHandle, QueryStore, Environment, NativeSubqueries);
 }
 
 bool UTypedElementQueryObserverCallbackAdapterProcessorWith8Subqueries::ConfigureQueryCallback(
-	FTypedElementExtendedQuery& Query, FTypedElementExtendedQueryStore::Handle QueryHandle, FTypedElementExtendedQueryStore& QueryStore, FTypedElementDatabaseEnvironment& Environment)
+	FExtendedQuery& Query, FExtendedQueryStore::Handle QueryHandle, FExtendedQueryStore& QueryStore, FEnvironment& Environment)
 {
 	return ConfigureQueryCallbackData(Query, QueryHandle, QueryStore, Environment, NativeSubqueries);
 }
