[Teds] Rename native classes and move into UE::EditorDataStorage namespaces

diff --git a/Engine/Plugins/Experimental/EditorDataStorage/Source/TedsCore/Private/Processors/TypedElementProcessorAdaptors.h b/Engine/Plugins/Experimental/EditorDataStorage/Source/TedsCore/Private/Processors/TypedElementProcessorAdaptors.h
--- a/Engine/Plugins/Experimental/EditorDataStorage/Source/TedsCore/Private/Processors/TypedElementProcessorAdaptors.h
+++ b/Engine/Plugins/Experimental/EditorDataStorage/Source/TedsCore/Private/Processors/TypedElementProcessorAdaptors.h
@@ -12,39 +12,46 @@
 
 #include "TypedElementProcessorAdaptors.generated.h"
 
-class FTypedElementDatabaseEnvironment;
-struct FTypedElementExtendedQuery;
-class FTypedElementExtendedQueryStore;
-
-struct FPhasePreOrPostAmbleExecutor
+namespace UE::Editor::DataStorage
 {
-	FPhasePreOrPostAmbleExecutor(FMassEntityManager& EntityManager, float DeltaTime);
-	~FPhasePreOrPostAmbleExecutor();
-
-	void ExecuteQuery(
-		ITypedElementDataStorageInterface::FQueryDescription& Description,
-		FTypedElementExtendedQueryStore& QueryStore,
-		FTypedElementDatabaseEnvironment& Environment,
-		FMassEntityQuery& NativeQuery,
-		ITypedElementDataStorageInterface::QueryCallbackRef Callback);
-
-	FMassExecutionContext Context;
-};
+	class FEnvironment;
+	struct FExtendedQuery;
+	class FExtendedQueryStore;
+
+	struct FPhasePreOrPostAmbleExecutor
+	{
+		FPhasePreOrPostAmbleExecutor(FMassEntityManager& EntityManager, float DeltaTime);
+		~FPhasePreOrPostAmbleExecutor();
+
+		void ExecuteQuery(
+			ITypedElementDataStorageInterface::FQueryDescription& Description,
+			FExtendedQueryStore& QueryStore,
+			FEnvironment& Environment,
+			FMassEntityQuery& NativeQuery,
+			ITypedElementDataStorageInterface::QueryCallbackRef Callback);
+
+		FMassExecutionContext Context;
+	};
+}
 
 USTRUCT()
 struct FTypedElementQueryProcessorData
 {
-	GENERATED_BODY()
+	GENERATED_BODY();
+
+	using FExtendedQuery = UE::Editor::DataStorage::FExtendedQuery;
+	using FExtendedQueryStore = UE::Editor::DataStorage::FExtendedQueryStore;
+	using FEnvironment = UE::Editor::DataStorage::FEnvironment;
 
 	FTypedElementQueryProcessorData() = default;
 	explicit FTypedElementQueryProcessorData(UMassProcessor& Owner);
 
 	bool CommonQueryConfiguration(
 		UMassProcessor& InOwner,
-		FTypedElementExtendedQuery& InQuery,
-		FTypedElementExtendedQueryStore::Handle InQueryHandle,
-		FTypedElementExtendedQueryStore& InQueryStore,
-		FTypedElementDatabaseEnvironment& InEnvironment,
+		FExtendedQuery& InQuery,
+		FExtendedQueryStore::Handle InQueryHandle,
+		FExtendedQueryStore& InQueryStore,
+		FEnvironment& InEnvironment,
 		TArrayView<FMassEntityQuery> Subqueries);
 	static EMassProcessingPhase MapToMassProcessingPhase(ITypedElementDataStorageInterface::EQueryTickPhase Phase);
 	FString GetProcessorName() const;
@@ -55,13 +62,13 @@ struct FTypedElementQueryProcessorData
 		TypedElementDataStorage::FQueryDescription& Description,
 		FMassEntityQuery& NativeQuery, 
 		FMassEntityManager& EntityManager,
-		FTypedElementDatabaseEnvironment& Environment);
+		FEnvironment& Environment);
 	static TypedElementDataStorage::FQueryResult Execute(
 		TypedElementDataStorage::SubqueryCallbackRef& Callback,
 		TypedElementDataStorage::FQueryDescription& Description,
 		FMassEntityQuery& NativeQuery,
 		FMassEntityManager& EntityManager,
-		FTypedElementDatabaseEnvironment& Environment,
+		FEnvironment& Environment,
 		FMassExecutionContext& ParentContext);
 	static TypedElementDataStorage::FQueryResult Execute(
 		TypedElementDataStorage::SubqueryCallbackRef& Callback,
@@ -69,16 +76,16 @@ struct FTypedElementQueryProcessorData
 		TypedElementDataStorage::RowHandle RowHandle,
 		FMassEntityQuery& NativeQuery,
 		FMassEntityManager& EntityManager,
-		FTypedElementDatabaseEnvironment& Environment,
+		FEnvironment& Environment,
 		FMassExecutionContext& ParentContext);
 	void Execute(FMassEntityManager& EntityManager, FMassExecutionContext& Context);
 
 	static bool PrepareCachedDependenciesOnQuery(
 		ITypedElementDataStorageInterface::FQueryDescription& Description, FMassExecutionContext& Context);
 
-	FTypedElementExtendedQueryStore::Handle ParentQuery;
-	FTypedElementExtendedQueryStore* QueryStore{ nullptr };
-	FTypedElementDatabaseEnvironment* Environment{ nullptr };
+	FExtendedQueryStore::Handle ParentQuery;
+	FExtendedQueryStore* QueryStore{ nullptr };
+	FEnvironment* Environment{ nullptr };
 	FMassEntityQuery NativeQuery;
 };
 
@@ -89,25 +96,29 @@ UCLASS()
 class UTypedElementQueryProcessorCallbackAdapterProcessorBase : public UMassProcessor
 {
 	GENERATED_BODY()
+	
+	using FExtendedQuery = UE::Editor::DataStorage::FExtendedQuery;
+	using FExtendedQueryStore = UE::Editor::DataStorage::FExtendedQueryStore;
+	using FEnvironment = UE::Editor::DataStorage::FEnvironment;
 
 public:
 	UTypedElementQueryProcessorCallbackAdapterProcessorBase();
 
 	FMassEntityQuery& GetQuery();
 	virtual bool ConfigureQueryCallback(
-		FTypedElementExtendedQuery& Query,
-		FTypedElementExtendedQueryStore::Handle QueryHandle,
-		FTypedElementExtendedQueryStore& QueryStore,
-		FTypedElementDatabaseEnvironment& Environment);
+		FExtendedQuery& Query,
+		FExtendedQueryStore::Handle QueryHandle,
+		FExtendedQueryStore& QueryStore,
+		FEnvironment& Environment);
 
 	virtual bool ShouldAllowQueryBasedPruning(const bool bRuntimeMode) const override;
 
 protected:
 	bool ConfigureQueryCallbackData(
-		FTypedElementExtendedQuery& Query,
-		FTypedElementExtendedQueryStore::Handle QueryHandle,
-		FTypedElementExtendedQueryStore& QueryStore,
-		FTypedElementDatabaseEnvironment& Environment,
+		FExtendedQuery& Query,
+		FExtendedQueryStore::Handle QueryHandle,
+		FExtendedQueryStore& QueryStore,
+		FEnvironment& Environment,
 		TArrayView<FMassEntityQuery> Subqueries);
 	void ConfigureQueries() override;
 	void Execute(FMassEntityManager& EntityManager, FMassExecutionContext& TargetParentQuery) override;
@@ -140,12 +151,16 @@ class UTypedElementQueryProcessorCallbackAdapterProcessorWith1Subquery final : p
 {
 	GENERATED_BODY()
 
+	using FExtendedQuery = UE::Editor::DataStorage::FExtendedQuery;
+	using FExtendedQueryStore = UE::Editor::DataStorage::FExtendedQueryStore;
+	using FEnvironment = UE::Editor::DataStorage::FEnvironment;
+
 public:
-	bool ConfigureQueryCallback(
-		FTypedElementExtendedQuery& Query,
-		FTypedElementExtendedQueryStore::Handle QueryHandle,
-		FTypedElementExtendedQueryStore& QueryStore,
-		FTypedElementDatabaseEnvironment& Environment) override;
+	virtual bool ConfigureQueryCallback(
+		FExtendedQuery& Query,
+		FExtendedQueryStore::Handle QueryHandle,
+		FExtendedQueryStore& QueryStore,
+		FEnvironment& Environment) override;
 
 private:
 	UPROPERTY(transient)
@@ -157,12 +172,16 @@ class UTypedElementQueryProcessorCallbackAdapterProcessorWith2Subqueries final :
 {
 	GENERATED_BODY()
 
+	using FExtendedQuery = UE::Editor::DataStorage::FExtendedQuery;
+	using FExtendedQueryStore = UE::Editor::DataStorage::FExtendedQueryStore;
+	using FEnvironment = UE::Editor::DataStorage::FEnvironment;
+
 public:
-	bool ConfigureQueryCallback(
-		FTypedElementExtendedQuery& Query,
-		FTypedElementExtendedQueryStore::Handle QueryHandle,
-		FTypedElementExtendedQueryStore& QueryStore,
-		FTypedElementDatabaseEnvironment& Environment) override;
+	virtual bool ConfigureQueryCallback(
+		FExtendedQuery& Query,
+		FExtendedQueryStore::Handle QueryHandle,
+		FExtendedQueryStore& QueryStore,
+		FEnvironment& Environment) override;
 
 private:
 	UPROPERTY(transient)
@@ -174,12 +193,16 @@ class UTypedElementQueryProcessorCallbackAdapterProcessorWith3Subqueries final :
 {
 	GENERATED_BODY()
 
+	using FExtendedQuery = UE::Editor::DataStorage::FExtendedQuery;
+	using FExtendedQueryStore = UE::Editor::DataStorage::FExtendedQueryStore;
+	using FEnvironment = UE::Editor::DataStorage::FEnvironment;
+
 public:
-	bool ConfigureQueryCallback(
-		FTypedElementExtendedQuery& Query,
-		FTypedElementExtendedQueryStore::Handle QueryHandle,
-		FTypedElementExtendedQueryStore& QueryStore,
-		FTypedElementDatabaseEnvironment& Environment) override;
+	virtual bool ConfigureQueryCallback(
+		FExtendedQuery& Query,
+		FExtendedQueryStore::Handle QueryHandle,
+		FExtendedQueryStore& QueryStore,
+		FEnvironment& Environment) override;
 
 private:
 	UPROPERTY(transient)
@@ -191,12 +214,16 @@ class UTypedElementQueryProcessorCallbackAdapterProcessorWith4Subqueries final :
 {
 	GENERATED_BODY()
 
+	using FExtendedQuery = UE::Editor::DataStorage::FExtendedQuery;
+	using FExtendedQueryStore = UE::Editor::DataStorage::FExtendedQueryStore;
+	using FEnvironment = UE::Editor::DataStorage::FEnvironment;
+
 public:
-	bool ConfigureQueryCallback(
-		FTypedElementExtendedQuery& Query,
-		FTypedElementExtendedQueryStore::Handle QueryHandle,
-		FTypedElementExtendedQueryStore& QueryStore,
-		FTypedElementDatabaseEnvironment& Environment) override;
+	virtual bool ConfigureQueryCallback(
+		FExtendedQuery& Query,
+		FExtendedQueryStore::Handle QueryHandle,
+		FExtendedQueryStore& QueryStore,
+		FEnvironment& Environment) override;
 
 private:
 	UPROPERTY(transient)
@@ -208,12 +235,16 @@ class UTypedElementQueryProcessorCallbackAdapterProcessorWith5Subqueries final :
 {
 	GENERATED_BODY()
 
+	using FExtendedQuery = UE::Editor::DataStorage::FExtendedQuery;
+	using FExtendedQueryStore = UE::Editor::DataStorage::FExtendedQueryStore;
+	using FEnvironment = UE::Editor::DataStorage::FEnvironment;
+
 public:
-	bool ConfigureQueryCallback(
-		FTypedElementExtendedQuery& Query,
-		FTypedElementExtendedQueryStore::Handle QueryHandle,
-		FTypedElementExtendedQueryStore& QueryStore,
-		FTypedElementDatabaseEnvironment& Environment) override;
+	virtual bool ConfigureQueryCallback(
+		FExtendedQuery& Query,
+		FExtendedQueryStore::Handle QueryHandle,
+		FExtendedQueryStore& QueryStore,
+		FEnvironment& Environment) override;
 
 private:
 	UPROPERTY(transient)
@@ -225,12 +256,16 @@ class UTypedElementQueryProcessorCallbackAdapterProcessorWith6Subqueries final :
 {
 	GENERATED_BODY()
 
+	using FExtendedQuery = UE::Editor::DataStorage::FExtendedQuery;
+	using FExtendedQueryStore = UE::Editor::DataStorage::FExtendedQueryStore;
+	using FEnvironment = UE::Editor::DataStorage::FEnvironment;
+
 public:
-	bool ConfigureQueryCallback(
-		FTypedElementExtendedQuery& Query,
-		FTypedElementExtendedQueryStore::Handle QueryHandle,
-		FTypedElementExtendedQueryStore& QueryStore,
-		FTypedElementDatabaseEnvironment& Environment) override;
+	virtual bool ConfigureQueryCallback(
+		FExtendedQuery& Query,
+		FExtendedQueryStore::Handle QueryHandle,
+		FExtendedQueryStore& QueryStore,
+		FEnvironment& Environment) override;
 
 private:
 	UPROPERTY(transient)
@@ -242,12 +277,16 @@ class UTypedElementQueryProcessorCallbackAdapterProcessorWith7Subqueries final :
 {
 	GENERATED_BODY()
 
+	using FExtendedQuery = UE::Editor::DataStorage::FExtendedQuery;
+	using FExtendedQueryStore = UE::Editor::DataStorage::FExtendedQueryStore;
+	using FEnvironment = UE::Editor::DataStorage::FEnvironment;
+
 public:
-	bool ConfigureQueryCallback(
-		FTypedElementExtendedQuery& Query,
-		FTypedElementExtendedQueryStore::Handle QueryHandle,
-		FTypedElementExtendedQueryStore& QueryStore,
-		FTypedElementDatabaseEnvironment& Environment) override;
+	virtual bool ConfigureQueryCallback(
+		FExtendedQuery& Query,
+		FExtendedQueryStore::Handle QueryHandle,
+		FExtendedQueryStore& QueryStore,
+		FEnvironment& Environment) override;
 
 private:
 	UPROPERTY(transient)
@@ -259,12 +298,16 @@ class UTypedElementQueryProcessorCallbackAdapterProcessorWith8Subqueries final :
 {
 	GENERATED_BODY()
 
+	using FExtendedQuery = UE::Editor::DataStorage::FExtendedQuery;
+	using FExtendedQueryStore = UE::Editor::DataStorage::FExtendedQueryStore;
+	using FEnvironment = UE::Editor::DataStorage::FEnvironment;
+
 public:
-	bool ConfigureQueryCallback(
-		FTypedElementExtendedQuery& Query,
-		FTypedElementExtendedQueryStore::Handle QueryHandle,
-		FTypedElementExtendedQueryStore& QueryStore,
-		FTypedElementDatabaseEnvironment& Environment) override;
+	virtual bool ConfigureQueryCallback(
+		FExtendedQuery& Query,
+		FExtendedQueryStore::Handle QueryHandle,
+		FExtendedQueryStore& QueryStore,
+		FEnvironment& Environment) override;
 
 private:
 	UPROPERTY(transient)
@@ -279,6 +322,10 @@ class UTypedElementQueryObserverCallbackAdapterProcessorBase : public UMassObser
 {
 	GENERATED_BODY()
 
+	using FExtendedQuery = UE::Editor::DataStorage::FExtendedQuery;
+	using FExtendedQueryStore = UE::Editor::DataStorage::FExtendedQueryStore;
+	using FEnvironment = UE::Editor::DataStorage::FEnvironment;
+
 public:
 	UTypedElementQueryObserverCallbackAdapterProcessorBase();
 
@@ -286,18 +333,25 @@ public:
 	const UScriptStruct* GetObservedType() const;
 	EMassObservedOperation GetObservedOperation() const;
 	virtual bool ConfigureQueryCallback(
-		FTypedElementExtendedQuery& Query, FTypedElementExtendedQueryStore::Handle QueryHandle, FTypedElementExtendedQueryStore& QueryStore, FTypedElementDatabaseEnvironment& Environment);
+		FExtendedQuery& Query,
+		FExtendedQueryStore::Handle QueryHandle,
+		FExtendedQueryStore& QueryStore,
+		FEnvironment& Environment);
 
 protected:
-	bool ConfigureQueryCallbackData(FTypedElementExtendedQuery& Query, FTypedElementExtendedQueryStore::Handle QueryHandle,
-	                                FTypedElementExtendedQueryStore& QueryStore, FTypedElementDatabaseEnvironment& Environment, TArrayView<FMassEntityQuery> Subqueries);
-	void ConfigureQueries() override;
-	void Execute(FMassEntityManager& EntityManager, FMassExecutionContext& TargetParentQuery) override;
+	bool ConfigureQueryCallbackData(
+		FExtendedQuery& Query,
+		FExtendedQueryStore::Handle QueryHandle,
+		FExtendedQueryStore& QueryStore,
+		FEnvironment& Environment,
+		TArrayView<FMassEntityQuery> Subqueries);
+	virtual void ConfigureQueries() override;
+	virtual void Execute(FMassEntityManager& EntityManager, FMassExecutionContext& TargetParentQuery) override;
 
-	void PostInitProperties() override;
-	void Register() override;
-	FString GetProcessorName() const override;
-	void DebugOutputDescription(FOutputDevice& Ar, int32 Indent) const override;
+	virtual void PostInitProperties() override;
+	virtual void Register() override;
+	virtual FString GetProcessorName() const override;
+	virtual void DebugOutputDescription(FOutputDevice& Ar, int32 Indent) const override;
 
 private:
 	UPROPERTY(transient)
@@ -323,12 +377,16 @@ class UTypedElementQueryObserverCallbackAdapterProcessorWith1Subquery final : pu
 {
 	GENERATED_BODY()
 
+	using FExtendedQuery = UE::Editor::DataStorage::FExtendedQuery;
+	using FExtendedQueryStore = UE::Editor::DataStorage::FExtendedQueryStore;
+	using FEnvironment = UE::Editor::DataStorage::FEnvironment;
+
 public:
-	bool ConfigureQueryCallback(
-		FTypedElementExtendedQuery& Query,
-		FTypedElementExtendedQueryStore::Handle QueryHandle,
-		FTypedElementExtendedQueryStore& QueryStore,
-		FTypedElementDatabaseEnvironment& Environment) override;
+	virtual bool ConfigureQueryCallback(
+		FExtendedQuery& Query,
+		FExtendedQueryStore::Handle QueryHandle,
+		FExtendedQueryStore& QueryStore,
+		FEnvironment& Environment) override;
 
 private:
 	UPROPERTY(transient)
@@ -340,12 +398,16 @@ class UTypedElementQueryObserverCallbackAdapterProcessorWith2Subqueries final :
 {
 	GENERATED_BODY()
 
+	using FExtendedQuery = UE::Editor::DataStorage::FExtendedQuery;
+	using FExtendedQueryStore = UE::Editor::DataStorage::FExtendedQueryStore;
+	using FEnvironment = UE::Editor::DataStorage::FEnvironment;
+
 public:
-	bool ConfigureQueryCallback(
-		FTypedElementExtendedQuery& Query,
-		FTypedElementExtendedQueryStore::Handle QueryHandle,
-		FTypedElementExtendedQueryStore& QueryStore,
-		FTypedElementDatabaseEnvironment& Environment) override;
+	virtual bool ConfigureQueryCallback(
+		FExtendedQuery& Query,
+		FExtendedQueryStore::Handle QueryHandle,
+		FExtendedQueryStore& QueryStore,
+		FEnvironment& Environment) override;
 
 private:
 	UPROPERTY(transient)
@@ -357,12 +419,16 @@ class UTypedElementQueryObserverCallbackAdapterProcessorWith3Subqueries final :
 {
 	GENERATED_BODY()
 
+	using FExtendedQuery = UE::Editor::DataStorage::FExtendedQuery;
+	using FExtendedQueryStore = UE::Editor::DataStorage::FExtendedQueryStore;
+	using FEnvironment = UE::Editor::DataStorage::FEnvironment;
+
 public:
-	bool ConfigureQueryCallback(
-		FTypedElementExtendedQuery& Query,
-		FTypedElementExtendedQueryStore::Handle QueryHandle,
-		FTypedElementExtendedQueryStore& QueryStore,
-		FTypedElementDatabaseEnvironment& Environment) override;
+	virtual bool ConfigureQueryCallback(
+		FExtendedQuery& Query,
+		FExtendedQueryStore::Handle QueryHandle,
+		FExtendedQueryStore& QueryStore,
+		FEnvironment& Environment) override;
 
 private:
 	UPROPERTY(transient)
@@ -374,12 +440,16 @@ class UTypedElementQueryObserverCallbackAdapterProcessorWith4Subqueries final :
 {
 	GENERATED_BODY()
 
+	using FExtendedQuery = UE::Editor::DataStorage::FExtendedQuery;
+	using FExtendedQueryStore = UE::Editor::DataStorage::FExtendedQueryStore;
+	using FEnvironment = UE::Editor::DataStorage::FEnvironment;
+
 public:
-	bool ConfigureQueryCallback(
-		FTypedElementExtendedQuery& Query,
-		FTypedElementExtendedQueryStore::Handle QueryHandle,
-		FTypedElementExtendedQueryStore& QueryStore,
-		FTypedElementDatabaseEnvironment& Environment) override;
+	virtual bool ConfigureQueryCallback(
+		FExtendedQuery& Query,
+		FExtendedQueryStore::Handle QueryHandle,
+		FExtendedQueryStore& QueryStore,
+		FEnvironment& Environment) override;
 
 private:
 	UPROPERTY(transient)
@@ -391,16 +461,20 @@ class UTypedElementQueryObserverCallbackAdapterProcessorWith5Subqueries final :
 {
 	GENERATED_BODY()
 
+	using FExtendedQuery = UE::Editor::DataStorage::FExtendedQuery;
+	using FExtendedQueryStore = UE::Editor::DataStorage::FExtendedQueryStore;
+	using FEnvironment = UE::Editor::DataStorage::FEnvironment;
+
 public:
-	bool ConfigureQueryCallback(
-		FTypedElementExtendedQuery& Query,
-		FTypedElementExtendedQueryStore::Handle QueryHandle,
-		FTypedElementExtendedQueryStore& QueryStore,
-		FTypedElementDatabaseEnvironment& Environment) override;
+	virtual bool ConfigureQueryCallback(
+		FExtendedQuery& Query,
+		FExtendedQueryStore::Handle QueryHandle,
+		FExtendedQueryStore& QueryStore,
+		FEnvironment& Environment) override;
 
 private:
 	UPROPERTY(transient)
-		FMassEntityQuery NativeSubqueries[5];
+	FMassEntityQuery NativeSubqueries[5];
 };
 
 UCLASS()
@@ -408,12 +482,16 @@ class UTypedElementQueryObserverCallbackAdapterProcessorWith6Subqueries final :
 {
 	GENERATED_BODY()
 
+	using FExtendedQuery = UE::Editor::DataStorage::FExtendedQuery;
+	using FExtendedQueryStore = UE::Editor::DataStorage::FExtendedQueryStore;
+	using FEnvironment = UE::Editor::DataStorage::FEnvironment;
+
 public:
-	bool ConfigureQueryCallback(
-		FTypedElementExtendedQuery& Query,
-		FTypedElementExtendedQueryStore::Handle QueryHandle,
-		FTypedElementExtendedQueryStore& QueryStore,
-		FTypedElementDatabaseEnvironment& Environment) override;
+	virtual bool ConfigureQueryCallback(
+		FExtendedQuery& Query,
+		FExtendedQueryStore::Handle QueryHandle,
+		FExtendedQueryStore& QueryStore,
+		FEnvironment& Environment) override;
 
 private:
 	UPROPERTY(transient)
@@ -425,12 +503,16 @@ class UTypedElementQueryObserverCallbackAdapterProcessorWith7Subqueries final :
 {
 	GENERATED_BODY()
 
+	using FExtendedQuery = UE::Editor::DataStorage::FExtendedQuery;
+	using FExtendedQueryStore = UE::Editor::DataStorage::FExtendedQueryStore;
+	using FEnvironment = UE::Editor::DataStorage::FEnvironment;
+
 public:
-	bool ConfigureQueryCallback(
-		FTypedElementExtendedQuery& Query,
-		FTypedElementExtendedQueryStore::Handle QueryHandle,
-		FTypedElementExtendedQueryStore& QueryStore,
-		FTypedElementDatabaseEnvironment& Environment) override;
+	virtual bool ConfigureQueryCallback(
+		FExtendedQuery& Query,
+		FExtendedQueryStore::Handle QueryHandle,
+		FExtendedQueryStore& QueryStore,
+		FEnvironment& Environment) override;
 
 private:
 	UPROPERTY(transient)
@@ -442,12 +524,16 @@ class UTypedElementQueryObserverCallbackAdapterProcessorWith8Subqueries final :
 {
 	GENERATED_BODY()
 
+	using FExtendedQuery = UE::Editor::DataStorage::FExtendedQuery;
+	using FExtendedQueryStore = UE::Editor::DataStorage::FExtendedQueryStore;
+	using FEnvironment = UE::Editor::DataStorage::FEnvironment;
+
 public:
-	bool ConfigureQueryCallback(
-		FTypedElementExtendedQuery& Query,
-		FTypedElementExtendedQueryStore::Handle QueryHandle,
-		FTypedElementExtendedQueryStore& QueryStore,
-		FTypedElementDatabaseEnvironment& Environment) override;
+	virtual bool ConfigureQueryCallback(
+		FExtendedQuery& Query,
+		FExtendedQueryStore::Handle QueryHandle,
+		FExtendedQueryStore& QueryStore,
+		FEnvironment& Environment) override;
 
 private:
 	UPROPERTY(transient)
diff --git a/Engine/Plugins/Experimental/EditorDataStorage/Source/TedsCore/Private/Queries/TypedElementExtendedQueryStore.cpp b/Engine/Plugins/Experimental/EditorDataStorage/Source/TedsCore/Private/Queries/TypedElementExtendedQueryStore.cpp
--- a/Engine/Plugins/Experimental/EditorDataStorage/Source/TedsCore/Private/Queries/TypedElementExtendedQueryStore.cpp
+++ b/Engine/Plugins/Experimental/EditorDataStorage/Source/TedsCore/Private/Queries/TypedElementExtendedQueryStore.cpp
@@ -9,903 +9,906 @@
 #include "Misc/OutputDevice.h"
 #include "Processors/TypedElementProcessorAdaptors.h"
 
-const ITypedElementDataStorageInterface::FQueryDescription FTypedElementExtendedQueryStore::EmptyDescription{};
+namespace UE::Editor::DataStorage
+{
+	const ITypedElementDataStorageInterface::FQueryDescription FExtendedQueryStore::EmptyDescription{};
+
+	FExtendedQueryStore::Handle FExtendedQueryStore::RegisterQuery(
+		ITypedElementDataStorageInterface::FQueryDescription Query,
+		FEnvironment& Environment,
+		FMassEntityManager& EntityManager,
+		FMassProcessingPhaseManager& PhaseManager)
+	{
+		FExtendedQueryStore::Handle Result = Queries.Emplace();
+		FExtendedQuery& StoredQuery = GetMutableChecked(Result);
+		StoredQuery.Description = MoveTemp(Query);
+
+		FMassEntityQuery& NativeQuery = SetupNativeQuery(StoredQuery.Description, StoredQuery);
+		bool bContinueSetup = SetupSelectedColumns(StoredQuery.Description, NativeQuery);
+		bContinueSetup = bContinueSetup && SetupChunkFilters(Result, StoredQuery.Description, Environment, NativeQuery);
+		bContinueSetup = bContinueSetup && SetupConditions(StoredQuery.Description, NativeQuery);
+		bContinueSetup = bContinueSetup && SetupDependencies(StoredQuery.Description, NativeQuery);
+		bContinueSetup = bContinueSetup && SetupTickGroupDefaults(StoredQuery.Description);
+		bContinueSetup = bContinueSetup && SetupProcessors(Result, StoredQuery, Environment, EntityManager, PhaseManager);
+		bContinueSetup = bContinueSetup && SetupActivatable(Result, Query);
+
+		if (!bContinueSetup)
+		{
+			// This will also make the handle invalid.
+			Queries.Remove(Result);
+		}
 
-FTypedElementExtendedQueryStore::Handle FTypedElementExtendedQueryStore::RegisterQuery(
-	ITypedElementDataStorageInterface::FQueryDescription Query, 
-	FTypedElementDatabaseEnvironment& Environment,
-	FMassEntityManager& EntityManager, 
-	FMassProcessingPhaseManager& PhaseManager)
-{
-	FTypedElementExtendedQueryStore::Handle Result = Queries.Emplace();
-	FTypedElementExtendedQuery& StoredQuery = GetMutableChecked(Result);
-	StoredQuery.Description = MoveTemp(Query);
+		return Result;
+	}
 
-	FMassEntityQuery& NativeQuery =		SetupNativeQuery(StoredQuery.Description, StoredQuery);
-	bool bContinueSetup =				SetupSelectedColumns(StoredQuery.Description, NativeQuery);
-	bContinueSetup = bContinueSetup &&	SetupChunkFilters(Result, StoredQuery.Description, Environment, NativeQuery);
-	bContinueSetup = bContinueSetup &&	SetupConditions(StoredQuery.Description, NativeQuery);
-	bContinueSetup = bContinueSetup &&	SetupDependencies(StoredQuery.Description, NativeQuery);
-	bContinueSetup = bContinueSetup &&	SetupTickGroupDefaults(StoredQuery.Description);
-	bContinueSetup = bContinueSetup &&	SetupProcessors(Result, StoredQuery, Environment, EntityManager, PhaseManager);
-	bContinueSetup = bContinueSetup &&	SetupActivatable(Result, Query);
-	
-	if (!bContinueSetup)
+	void FExtendedQueryStore::UnregisterQuery(Handle Query, FMassEntityManager& EntityManager, FMassProcessingPhaseManager& PhaseManager)
 	{
-		// This will also make the handle invalid.
-		Queries.Remove(Result);
+		if (FExtendedQuery* QueryData = Get(Query))
+		{
+			UnregisterQueryData(Query, *QueryData, EntityManager, PhaseManager);
+			Queries.Remove(Query);
+		}
 	}
 
-	return Result;
-}
-
-void FTypedElementExtendedQueryStore::UnregisterQuery(Handle Query, FMassEntityManager& EntityManager, FMassProcessingPhaseManager& PhaseManager)
-{
-	if (FTypedElementExtendedQuery* QueryData = Get(Query))
+	void FExtendedQueryStore::Clear(FMassEntityManager& EntityManager, FMassProcessingPhaseManager& PhaseManager)
 	{
-		UnregisterQueryData(Query, *QueryData, EntityManager, PhaseManager);
-		Queries.Remove(Query);
+		TickGroupDescriptions.Empty();
+
+		Queries.ListAliveEntries([this, &EntityManager, &PhaseManager](Handle Query, FExtendedQuery& QueryData)
+			{
+				if (QueryData.Processor && QueryData.Processor->IsA<UTypedElementQueryObserverCallbackAdapterProcessorBase>())
+				{
+					// Observers can't be unregistered at this point, so skip these for now.
+					return;
+				}
+				UnregisterQueryData(Query, QueryData, EntityManager, PhaseManager);
+			});
 	}
-}
 
-void FTypedElementExtendedQueryStore::Clear(FMassEntityManager& EntityManager, FMassProcessingPhaseManager& PhaseManager)
-{
-	TickGroupDescriptions.Empty();
+	void FExtendedQueryStore::RegisterTickGroup(FName GroupName, ITypedElementDataStorageInterface::EQueryTickPhase Phase,
+		FName BeforeGroup, FName AfterGroup, bool bRequiresMainThread)
+	{
+		FTickGroupDescription& Group = TickGroupDescriptions.FindOrAdd({ GroupName, Phase });
 
-	Queries.ListAliveEntries([this, &EntityManager, &PhaseManager](Handle Query, FTypedElementExtendedQuery& QueryData)
+		if (!BeforeGroup.IsNone() && Group.BeforeGroups.Find(BeforeGroup) == INDEX_NONE)
 		{
-			if (QueryData.Processor && QueryData.Processor->IsA<UTypedElementQueryObserverCallbackAdapterProcessorBase>())
-			{
-				// Observers can't be unregistered at this point, so skip these for now.
-				return;
-			}
-			UnregisterQueryData(Query, QueryData, EntityManager, PhaseManager);
-		});
-}
+			Group.BeforeGroups.Add(BeforeGroup);
+		}
 
-void FTypedElementExtendedQueryStore::RegisterTickGroup(FName GroupName, ITypedElementDataStorageInterface::EQueryTickPhase Phase,
-	FName BeforeGroup, FName AfterGroup, bool bRequiresMainThread)
-{
-	FTickGroupDescription& Group = TickGroupDescriptions.FindOrAdd({ GroupName, Phase });
+		if (!AfterGroup.IsNone() && Group.AfterGroups.Find(AfterGroup) == INDEX_NONE)
+		{
+			Group.AfterGroups.Add(AfterGroup);
+		}
 
-	if (!BeforeGroup.IsNone() && Group.BeforeGroups.Find(BeforeGroup) == INDEX_NONE)
-	{
-		Group.BeforeGroups.Add(BeforeGroup);
+		if (bRequiresMainThread)
+		{
+			Group.bRequiresMainThread = true;
+		}
 	}
 
-	if (!AfterGroup.IsNone() && Group.AfterGroups.Find(AfterGroup) == INDEX_NONE)
+	void FExtendedQueryStore::UnregisterTickGroup(FName GroupName, ITypedElementDataStorageInterface::EQueryTickPhase Phase)
 	{
-		Group.AfterGroups.Add(AfterGroup);
+		TickGroupDescriptions.Remove({ GroupName, Phase });
 	}
 
-	if (bRequiresMainThread)
+	FExtendedQuery* FExtendedQueryStore::Get(Handle Entry)
 	{
-		Group.bRequiresMainThread = true;
+		return IsAlive(Entry) ? &Queries.Get(Entry) : nullptr;
 	}
-}
 
-void FTypedElementExtendedQueryStore::UnregisterTickGroup(FName GroupName, ITypedElementDataStorageInterface::EQueryTickPhase Phase)
-{
-	TickGroupDescriptions.Remove({ GroupName, Phase });
-}
-
-FTypedElementExtendedQuery* FTypedElementExtendedQueryStore::Get(Handle Entry)
-{
-	return IsAlive(Entry) ? &Queries.Get(Entry) : nullptr;
-}
-
-FTypedElementExtendedQuery* FTypedElementExtendedQueryStore::GetMutable(Handle Entry)
-{
-	return IsAlive(Entry) ? &Queries.GetMutable(Entry) : nullptr;
-}
-
-const FTypedElementExtendedQuery* FTypedElementExtendedQueryStore::Get(Handle Entry) const
-{
-	return IsAlive(Entry) ? &Queries.Get(Entry) : nullptr;
-}
+	FExtendedQuery* FExtendedQueryStore::GetMutable(Handle Entry)
+	{
+		return IsAlive(Entry) ? &Queries.GetMutable(Entry) : nullptr;
+	}
 
-FTypedElementExtendedQuery& FTypedElementExtendedQueryStore::GetChecked(Handle Entry)
-{
-	return Queries.Get(Entry);
-}
+	const FExtendedQuery* FExtendedQueryStore::Get(Handle Entry) const
+	{
+		return IsAlive(Entry) ? &Queries.Get(Entry) : nullptr;
+	}
 
-FTypedElementExtendedQuery& FTypedElementExtendedQueryStore::GetMutableChecked(Handle Entry)
-{
-	return Queries.GetMutable(Entry);
-}
+	FExtendedQuery& FExtendedQueryStore::GetChecked(Handle Entry)
+	{
+		return Queries.Get(Entry);
+	}
 
-const FTypedElementExtendedQuery& FTypedElementExtendedQueryStore::GetChecked(Handle Entry) const
-{
-	return Queries.Get(Entry);
-}
+	FExtendedQuery& FExtendedQueryStore::GetMutableChecked(Handle Entry)
+	{
+		return Queries.GetMutable(Entry);
+	}
 
-const ITypedElementDataStorageInterface::FQueryDescription& FTypedElementExtendedQueryStore::GetQueryDescription(Handle Query) const
-{
-	const FTypedElementExtendedQuery* QueryData = Get(Query);
-	return QueryData ? QueryData->Description : EmptyDescription;
-}
+	const FExtendedQuery& FExtendedQueryStore::GetChecked(Handle Entry) const
+	{
+		return Queries.Get(Entry);
+	}
 
-bool FTypedElementExtendedQueryStore::IsAlive(Handle Entry) const
-{
-	return Queries.IsAlive(Entry);
-}
+	const ITypedElementDataStorageInterface::FQueryDescription& FExtendedQueryStore::GetQueryDescription(Handle Query) const
+	{
+		const FExtendedQuery* QueryData = Get(Query);
+		return QueryData ? QueryData->Description : EmptyDescription;
+	}
 
-void FTypedElementExtendedQueryStore::ListAliveEntries(const ListAliveEntriesConstCallback& Callback) const
-{
-	Queries.ListAliveEntries(Callback);
-}
+	bool FExtendedQueryStore::IsAlive(Handle Entry) const
+	{
+		return Queries.IsAlive(Entry);
+	}
 
-void FTypedElementExtendedQueryStore::UpdateActivatableQueries()
-{
-	// Update activatable counts and remove any queries that have completed.
-	for (Handle Query : ActiveActivatables)
+	void FExtendedQueryStore::ListAliveEntries(const ListAliveEntriesConstCallback& Callback) const
 	{
-		FTypedElementExtendedQuery& QueryData = Queries.GetMutable(Query);
-		checkf(QueryData.Description.Callback.ActivationCount > 0,
-			TEXT("Attempting to decrement the query '%s' which is already at zero."), *QueryData.Description.Callback.Name.ToString());
-		QueryData.Description.Callback.ActivationCount--;
+		Queries.ListAliveEntries(Callback);
 	}
-	ActiveActivatables.Reset();
 
-	// Queue up the next batch of activatables.
-	for (Handle Query : PendingActivatables)
+	void FExtendedQueryStore::UpdateActivatableQueries()
 	{
-		FTypedElementExtendedQuery& QueryData = Queries.GetMutable(Query);
-		if (QueryData.Description.Callback.ActivationCount == 0)
+		// Update activatable counts and remove any queries that have completed.
+		for (Handle Query : ActiveActivatables)
+		{
+			FExtendedQuery& QueryData = Queries.GetMutable(Query);
+			checkf(QueryData.Description.Callback.ActivationCount > 0,
+				TEXT("Attempting to decrement the query '%s' which is already at zero."), *QueryData.Description.Callback.Name.ToString());
+			QueryData.Description.Callback.ActivationCount--;
+		}
+		ActiveActivatables.Reset();
+
+		// Queue up the next batch of activatables.
+		for (Handle Query : PendingActivatables)
 		{
-			QueryData.Description.Callback.ActivationCount = 1;
-			ActiveActivatables.Add(Query);
+			FExtendedQuery& QueryData = Queries.GetMutable(Query);
+			if (QueryData.Description.Callback.ActivationCount == 0)
+			{
+				QueryData.Description.Callback.ActivationCount = 1;
+				ActiveActivatables.Add(Query);
+			}
 		}
+		PendingActivatables.Reset();
 	}
-	PendingActivatables.Reset();
-}
 
-void FTypedElementExtendedQueryStore::ActivateQueries(FName ActivationName)
-{
-	for (TMultiMap<FName, Handle>::TKeyIterator QueryIt = ActivatableMapping.CreateKeyIterator(ActivationName); QueryIt; ++QueryIt)
+	void FExtendedQueryStore::ActivateQueries(FName ActivationName)
 	{
-		Handle Query = QueryIt.Value();
-		if (Queries.IsAlive(Query))
+		for (TMultiMap<FName, Handle>::TKeyIterator QueryIt = ActivatableMapping.CreateKeyIterator(ActivationName); QueryIt; ++QueryIt)
 		{
+			Handle Query = QueryIt.Value();
+			if (Queries.IsAlive(Query))
+			{
 #if DO_ENSURE
-			FTypedElementExtendedQuery& QueryData = Queries.GetMutable(Query);
-			checkf(!QueryData.Description.Callback.ActivationName.IsNone(),
-				TEXT("Attempting to enable the query '%s' which isn't activatable."), *QueryData.Description.Callback.Name.ToString());
+				FExtendedQuery& QueryData = Queries.GetMutable(Query);
+				checkf(!QueryData.Description.Callback.ActivationName.IsNone(),
+					TEXT("Attempting to enable the query '%s' which isn't activatable."), *QueryData.Description.Callback.Name.ToString());
 #endif
-			PendingActivatables.Add(Query);
+				PendingActivatables.Add(Query);
+			}
 		}
 	}
-}
 
 
-TypedElementDataStorage::FQueryResult FTypedElementExtendedQueryStore::RunQuery(FMassEntityManager& EntityManager, Handle Query)
-{
-	using ActionType = ITypedElementDataStorageInterface::FQueryDescription::EActionType;
-	using CompletionType = ITypedElementDataStorageInterface::FQueryResult::ECompletion;
+	TypedElementDataStorage::FQueryResult FExtendedQueryStore::RunQuery(FMassEntityManager& EntityManager, Handle Query)
+	{
+		using ActionType = ITypedElementDataStorageInterface::FQueryDescription::EActionType;
+		using CompletionType = ITypedElementDataStorageInterface::FQueryResult::ECompletion;
 
-	ITypedElementDataStorageInterface::FQueryResult Result;
+		ITypedElementDataStorageInterface::FQueryResult Result;
 
-	if (FTypedElementExtendedQuery* QueryData = Get(Query))
-	{
-		switch (QueryData->Description.Action)
+		if (FExtendedQuery* QueryData = Get(Query))
 		{
-		case ActionType::None:
-			Result.Completed = CompletionType::Fully;
-			break;
-		case ActionType::Select:
-			// Fall through: There's nothing to callback to, so only return the total count.
-		case ActionType::Count:
-			Result.Count = QueryData->NativeQuery.GetNumMatchingEntities(EntityManager);
-			Result.Completed = CompletionType::Fully;
-			break;
-		default:
-			Result.Completed = CompletionType::Unsupported;
-			break;
+			switch (QueryData->Description.Action)
+			{
+			case ActionType::None:
+				Result.Completed = CompletionType::Fully;
+				break;
+			case ActionType::Select:
+				// Fall through: There's nothing to callback to, so only return the total count.
+			case ActionType::Count:
+				Result.Count = QueryData->NativeQuery.GetNumMatchingEntities(EntityManager);
+				Result.Completed = CompletionType::Fully;
+				break;
+			default:
+				Result.Completed = CompletionType::Unsupported;
+				break;
+			}
+		}
+		else
+		{
+			Result.Completed = CompletionType::Unavailable;
 		}
-	}
-	else
-	{
-		Result.Completed = CompletionType::Unavailable;
-	}
 
-	return Result;
-}
+		return Result;
+	}
 
-template<typename CallbackReference>
-TypedElementDataStorage::FQueryResult FTypedElementExtendedQueryStore::RunQueryCallbackCommon(
-	FMassEntityManager& EntityManager, 
-	FTypedElementDatabaseEnvironment& Environment,
-	FMassExecutionContext* ParentContext,
-	Handle Query,
-	CallbackReference Callback)
-{
-	using ActionType = ITypedElementDataStorageInterface::FQueryDescription::EActionType;
-	using CompletionType = ITypedElementDataStorageInterface::FQueryResult::ECompletion;
+	template<typename CallbackReference>
+	TypedElementDataStorage::FQueryResult FExtendedQueryStore::RunQueryCallbackCommon(
+		FMassEntityManager& EntityManager,
+		FEnvironment& Environment,
+		FMassExecutionContext* ParentContext,
+		Handle Query,
+		CallbackReference Callback)
+	{
+		using ActionType = ITypedElementDataStorageInterface::FQueryDescription::EActionType;
+		using CompletionType = ITypedElementDataStorageInterface::FQueryResult::ECompletion;
 
-	ITypedElementDataStorageInterface::FQueryResult Result;
+		ITypedElementDataStorageInterface::FQueryResult Result;
 
-	if (FTypedElementExtendedQuery* QueryData = Get(Query))
-	{
-		switch (QueryData->Description.Action)
+		if (FExtendedQuery* QueryData = Get(Query))
 		{
-		case ActionType::None:
-			Result.Completed = CompletionType::Fully;
-			break;
-		case ActionType::Select:
-			if (!QueryData->Processor.IsValid())
+			switch (QueryData->Description.Action)
 			{
-				if constexpr (std::is_same_v<CallbackReference, TypedElementDataStorage::DirectQueryCallbackRef>)
+			case ActionType::None:
+				Result.Completed = CompletionType::Fully;
+				break;
+			case ActionType::Select:
+				if (!QueryData->Processor.IsValid())
 				{
-					Result = FTypedElementQueryProcessorData::Execute(
-						Callback, QueryData->Description, QueryData->NativeQuery, EntityManager, Environment);
+					if constexpr (std::is_same_v<CallbackReference, TypedElementDataStorage::DirectQueryCallbackRef>)
+					{
+						Result = FTypedElementQueryProcessorData::Execute(
+							Callback, QueryData->Description, QueryData->NativeQuery, EntityManager, Environment);
+					}
+					else
+					{
+						Result = FTypedElementQueryProcessorData::Execute(
+							Callback, QueryData->Description, QueryData->NativeQuery, EntityManager, Environment, *ParentContext);
+					}
 				}
 				else
 				{
-					Result = FTypedElementQueryProcessorData::Execute(
-						Callback, QueryData->Description, QueryData->NativeQuery, EntityManager, Environment, *ParentContext);
+					Result.Completed = CompletionType::Unsupported;
 				}
-			}
-			else
-			{
+				break;
+			case ActionType::Count:
+				// Only the count is requested so no need to trigger the callback.
+				Result.Count = QueryData->NativeQuery.GetNumMatchingEntities(EntityManager);
+				Result.Completed = CompletionType::Fully;
+				break;
+			default:
 				Result.Completed = CompletionType::Unsupported;
+				break;
 			}
-			break;
-		case ActionType::Count:
-			// Only the count is requested so no need to trigger the callback.
-			Result.Count = QueryData->NativeQuery.GetNumMatchingEntities(EntityManager);
-			Result.Completed = CompletionType::Fully;
-			break;
-		default:
-			Result.Completed = CompletionType::Unsupported;
-			break;
 		}
+		else
+		{
+			Result.Completed = CompletionType::Unavailable;
+		}
+
+		return Result;
 	}
-	else
+
+	TypedElementDataStorage::FQueryResult FExtendedQueryStore::RunQuery(FMassEntityManager& EntityManager,
+		FEnvironment& Environment, Handle Query, TypedElementDataStorage::DirectQueryCallbackRef Callback)
 	{
-		Result.Completed = CompletionType::Unavailable;
+		return RunQueryCallbackCommon(EntityManager, Environment, nullptr, Query, Callback);
 	}
 
-	return Result;
-}
-
-TypedElementDataStorage::FQueryResult FTypedElementExtendedQueryStore::RunQuery(FMassEntityManager& EntityManager, 
-	FTypedElementDatabaseEnvironment& Environment, Handle Query, TypedElementDataStorage::DirectQueryCallbackRef Callback)
-{
-	return RunQueryCallbackCommon(EntityManager, Environment, nullptr, Query, Callback);
-}
-
-TypedElementDataStorage::FQueryResult FTypedElementExtendedQueryStore::RunQuery(FMassEntityManager& EntityManager, 
-	FTypedElementDatabaseEnvironment& Environment, FMassExecutionContext& ParentContext, Handle Query, 
-	TypedElementDataStorage::SubqueryCallbackRef Callback)
-{
-	return RunQueryCallbackCommon(EntityManager, Environment, &ParentContext, Query, Callback);
-}
+	TypedElementDataStorage::FQueryResult FExtendedQueryStore::RunQuery(FMassEntityManager& EntityManager,
+		FEnvironment& Environment, FMassExecutionContext& ParentContext, Handle Query,
+		TypedElementDataStorage::SubqueryCallbackRef Callback)
+	{
+		return RunQueryCallbackCommon(EntityManager, Environment, &ParentContext, Query, Callback);
+	}
 
-TypedElementDataStorage::FQueryResult FTypedElementExtendedQueryStore::RunQuery(FMassEntityManager& EntityManager, 
-	FTypedElementDatabaseEnvironment& Environment, FMassExecutionContext& ParentContext, Handle Query, TypedElementRowHandle Row,
-	TypedElementDataStorage::SubqueryCallbackRef Callback)
-{
-	using ActionType = ITypedElementDataStorageInterface::FQueryDescription::EActionType;
-	using CompletionType = ITypedElementDataStorageInterface::FQueryResult::ECompletion;
+	TypedElementDataStorage::FQueryResult FExtendedQueryStore::RunQuery(FMassEntityManager& EntityManager,
+		FEnvironment& Environment, FMassExecutionContext& ParentContext, Handle Query, TypedElementRowHandle Row,
+		TypedElementDataStorage::SubqueryCallbackRef Callback)
+	{
+		using ActionType = ITypedElementDataStorageInterface::FQueryDescription::EActionType;
+		using CompletionType = ITypedElementDataStorageInterface::FQueryResult::ECompletion;
 
-	ITypedElementDataStorageInterface::FQueryResult Result;
+		ITypedElementDataStorageInterface::FQueryResult Result;
 
-	if (FTypedElementExtendedQuery* QueryData = Get(Query))
-	{
-		switch (QueryData->Description.Action)
+		if (FExtendedQuery* QueryData = Get(Query))
 		{
-		case ActionType::None:
-			Result.Completed = CompletionType::Fully;
-			break;
-		case ActionType::Select:
-			if (!QueryData->Processor.IsValid())
-			{
-				Result = FTypedElementQueryProcessorData::Execute(
-					Callback, QueryData->Description, Row, QueryData->NativeQuery, EntityManager, Environment, ParentContext);
-			}
-			else
+			switch (QueryData->Description.Action)
 			{
+			case ActionType::None:
+				Result.Completed = CompletionType::Fully;
+				break;
+			case ActionType::Select:
+				if (!QueryData->Processor.IsValid())
+				{
+					Result = FTypedElementQueryProcessorData::Execute(
+						Callback, QueryData->Description, Row, QueryData->NativeQuery, EntityManager, Environment, ParentContext);
+				}
+				else
+				{
+					Result.Completed = CompletionType::Unsupported;
+				}
+				break;
+			case ActionType::Count:
+				// Only the count is requested so no need to trigger the callback.
+				Result.Count = 1;
+				Result.Completed = CompletionType::Fully;
+				break;
+			default:
 				Result.Completed = CompletionType::Unsupported;
+				break;
 			}
-			break;
-		case ActionType::Count:
-			// Only the count is requested so no need to trigger the callback.
-			Result.Count = 1;
-			Result.Completed = CompletionType::Fully;
-			break;
-		default:
-			Result.Completed = CompletionType::Unsupported;
-			break;
 		}
+		else
+		{
+			Result.Completed = CompletionType::Unavailable;
+		}
+
+		return Result;
 	}
-	else
+	void FExtendedQueryStore::RunPhasePreambleQueries(FMassEntityManager& EntityManager,
+		FEnvironment& Environment, ITypedElementDataStorageInterface::EQueryTickPhase Phase, float DeltaTime)
 	{
-		Result.Completed = CompletionType::Unavailable;
+		RunPhasePreOrPostAmbleQueries(EntityManager, Environment, Phase, DeltaTime,
+			PhasePreparationQueries[static_cast<QueryTickPhaseType>(Phase)]);
 	}
 
-	return Result;
-}
-void FTypedElementExtendedQueryStore::RunPhasePreambleQueries(FMassEntityManager& EntityManager,
-	FTypedElementDatabaseEnvironment& Environment, ITypedElementDataStorageInterface::EQueryTickPhase Phase, float DeltaTime)
-{
-	RunPhasePreOrPostAmbleQueries(EntityManager, Environment, Phase, DeltaTime, 
-		PhasePreparationQueries[static_cast<QueryTickPhaseType>(Phase)]);
-}
+	void FExtendedQueryStore::RunPhasePostambleQueries(FMassEntityManager& EntityManager,
+		FEnvironment& Environment, ITypedElementDataStorageInterface::EQueryTickPhase Phase, float DeltaTime)
+	{
+		RunPhasePreOrPostAmbleQueries(EntityManager, Environment, Phase, DeltaTime,
+			PhaseFinalizationQueries[static_cast<QueryTickPhaseType>(Phase)]);
+	}
 
-void FTypedElementExtendedQueryStore::RunPhasePostambleQueries(FMassEntityManager& EntityManager,
-	FTypedElementDatabaseEnvironment& Environment, ITypedElementDataStorageInterface::EQueryTickPhase Phase, float DeltaTime)
-{
-	RunPhasePreOrPostAmbleQueries(EntityManager, Environment, Phase, DeltaTime, 
-		PhaseFinalizationQueries[static_cast<QueryTickPhaseType>(Phase)]);
-}
+	void FExtendedQueryStore::DebugPrintQueryCallbacks(FOutputDevice& Output) const
+	{
+		Output.Log(TEXT("The Typed Elements Data Storage has the following query callbacks:"));
+		Queries.ListAliveEntries(
+			[&Output](Handle QueryHandle, const FExtendedQuery& Query)
+			{
+				if (Query.Processor)
+				{
+					Output.Logf(TEXT("    [%s] %s"),
+						IsValid(Query.Processor.Get()) ? TEXT("Valid") : TEXT("Invalid"),
+						*(Query.Processor->GetProcessorName()));
+				}
+			});
 
-void FTypedElementExtendedQueryStore::DebugPrintQueryCallbacks(FOutputDevice& Output) const
-{
-	Output.Log(TEXT("The Typed Elements Data Storage has the following query callbacks:"));
-	Queries.ListAliveEntries(
-		[&Output](Handle QueryHandle, const FTypedElementExtendedQuery& Query)
+		for (QueryTickPhaseType PhaseId = 0; PhaseId < static_cast<QueryTickPhaseType>(MaxTickPhase); ++PhaseId)
 		{
-			if (Query.Processor)
+			for (Handle QueryHandle : PhasePreparationQueries[PhaseId])
 			{
-				Output.Logf(TEXT("    [%s] %s"),
-					IsValid(Query.Processor.Get()) ? TEXT("Valid") : TEXT("Invalid"),
-					*(Query.Processor->GetProcessorName()));
+				const FExtendedQuery& QueryData = GetChecked(QueryHandle);
+				Output.Logf(TEXT("    [Valid] %s [Editor Phase Preamble]"), *QueryData.Description.Callback.Name.ToString());
+			}
+			for (Handle QueryHandle : PhaseFinalizationQueries[PhaseId])
+			{
+				const FExtendedQuery& QueryData = GetChecked(QueryHandle);
+				Output.Logf(TEXT("    [Valid] %s [Editor Phase Postamble]"), *QueryData.Description.Callback.Name.ToString());
 			}
-		});
-
-	for (QueryTickPhaseType PhaseId = 0; PhaseId < static_cast<QueryTickPhaseType>(MaxTickPhase); ++PhaseId)
-	{
-		for (Handle QueryHandle : PhasePreparationQueries[PhaseId])
-		{
-			const FTypedElementExtendedQuery& QueryData = GetChecked(QueryHandle);
-			Output.Logf(TEXT("    [Valid] %s [Editor Phase Preamble]"), *QueryData.Description.Callback.Name.ToString());
-		}
-		for (Handle QueryHandle : PhaseFinalizationQueries[PhaseId])
-		{
-			const FTypedElementExtendedQuery& QueryData = GetChecked(QueryHandle);
-			Output.Logf(TEXT("    [Valid] %s [Editor Phase Postamble]"), *QueryData.Description.Callback.Name.ToString());
 		}
-	}
 
-	Output.Log(TEXT("End of Typed Elements Data Storage query callback list."));
-}
+		Output.Log(TEXT("End of Typed Elements Data Storage query callback list."));
+	}
 
-FMassEntityQuery& FTypedElementExtendedQueryStore::SetupNativeQuery(
-	ITypedElementDataStorageInterface::FQueryDescription& Query, FTypedElementExtendedQuery& StoredQuery)
-{
-	/**
-	 * Mass verifies that queries that are used by processors are on the processor themselves. It does this by taking the address of the query
-	 * and seeing if it's within the start and end address of the processor. When a dynamic array is used those addresses are going to be
-	 * elsewhere, so the two options are to store a single fixed size array on a processor or have multiple instances. With Mass' queries being
-	 * not an insignificant size it's preferable to have several variants with queries to allow the choice for the minimal size. Unfortunately
-	 * UHT doesn't allow for templates so it had to be done in an explicit way.
-	 */
+	FMassEntityQuery& FExtendedQueryStore::SetupNativeQuery(
+		ITypedElementDataStorageInterface::FQueryDescription& Query, FExtendedQuery& StoredQuery)
+	{
+		/**
+		 * Mass verifies that queries that are used by processors are on the processor themselves. It does this by taking the address of the query
+		 * and seeing if it's within the start and end address of the processor. When a dynamic array is used those addresses are going to be
+		 * elsewhere, so the two options are to store a single fixed size array on a processor or have multiple instances. With Mass' queries being
+		 * not an insignificant size it's preferable to have several variants with queries to allow the choice for the minimal size. Unfortunately
+		 * UHT doesn't allow for templates so it had to be done in an explicit way.
+		 */
 
-	using DSI = ITypedElementDataStorageInterface;
+		using DSI = ITypedElementDataStorageInterface;
 
-	if (Query.Action == DSI::FQueryDescription::EActionType::Select)
-	{
-		switch (Query.Callback.Type)
-		{
-		case DSI::EQueryCallbackType::None:
-			break;
-		case DSI::EQueryCallbackType::Processor:
+		if (Query.Action == DSI::FQueryDescription::EActionType::Select)
 		{
-			UTypedElementQueryProcessorCallbackAdapterProcessorBase* Processor;
-			switch (Query.Subqueries.Num())
+			switch (Query.Callback.Type)
 			{
-			case 0:
-				Processor = NewObject<UTypedElementQueryProcessorCallbackAdapterProcessor>();
-				break;
-			case 1:
-				Processor = NewObject<UTypedElementQueryProcessorCallbackAdapterProcessorWith1Subquery>();
-				break;
-			case 2:
-				Processor = NewObject<UTypedElementQueryProcessorCallbackAdapterProcessorWith2Subqueries>();
+			case DSI::EQueryCallbackType::None:
 				break;
-			case 3:
-				Processor = NewObject<UTypedElementQueryProcessorCallbackAdapterProcessorWith3Subqueries>();
-				break;
-			case 4:
-				Processor = NewObject<UTypedElementQueryProcessorCallbackAdapterProcessorWith4Subqueries>();
-				break;
-			case 5:
-				Processor = NewObject<UTypedElementQueryProcessorCallbackAdapterProcessorWith5Subqueries>();
-				break;
-			case 6:
-				Processor = NewObject<UTypedElementQueryProcessorCallbackAdapterProcessorWith6Subqueries>();
-				break;
-			case 7:
-				Processor = NewObject<UTypedElementQueryProcessorCallbackAdapterProcessorWith7Subqueries>();
+			case DSI::EQueryCallbackType::Processor:
+			{
+				UTypedElementQueryProcessorCallbackAdapterProcessorBase* Processor;
+				switch (Query.Subqueries.Num())
+				{
+				case 0:
+					Processor = NewObject<UTypedElementQueryProcessorCallbackAdapterProcessor>();
+					break;
+				case 1:
+					Processor = NewObject<UTypedElementQueryProcessorCallbackAdapterProcessorWith1Subquery>();
+					break;
+				case 2:
+					Processor = NewObject<UTypedElementQueryProcessorCallbackAdapterProcessorWith2Subqueries>();
+					break;
+				case 3:
+					Processor = NewObject<UTypedElementQueryProcessorCallbackAdapterProcessorWith3Subqueries>();
+					break;
+				case 4:
+					Processor = NewObject<UTypedElementQueryProcessorCallbackAdapterProcessorWith4Subqueries>();
+					break;
+				case 5:
+					Processor = NewObject<UTypedElementQueryProcessorCallbackAdapterProcessorWith5Subqueries>();
+					break;
+				case 6:
+					Processor = NewObject<UTypedElementQueryProcessorCallbackAdapterProcessorWith6Subqueries>();
+					break;
+				case 7:
+					Processor = NewObject<UTypedElementQueryProcessorCallbackAdapterProcessorWith7Subqueries>();
+					break;
+				case 8:
+					Processor = NewObject<UTypedElementQueryProcessorCallbackAdapterProcessorWith8Subqueries>();
+					break;
+				default:
+					checkf(false, TEXT("The current Typed Elements Data Storage backend doesn't support %i subqueries per processor query."),
+						Query.Subqueries.Num());
+					return StoredQuery.NativeQuery;
+				}
+				StoredQuery.Processor.Reset(Processor);
+				return Processor->GetQuery();
+			}
+			case DSI::EQueryCallbackType::ObserveAdd:
+				// Fall-through
+			case DSI::EQueryCallbackType::ObserveRemove:
+			{
+				UTypedElementQueryObserverCallbackAdapterProcessorBase* Observer;
+				switch (Query.Subqueries.Num())
+				{
+				case 0:
+					Observer = NewObject<UTypedElementQueryObserverCallbackAdapterProcessor>();
+					break;
+				case 1:
+					Observer = NewObject<UTypedElementQueryObserverCallbackAdapterProcessorWith1Subquery>();
+					break;
+				case 2:
+					Observer = NewObject<UTypedElementQueryObserverCallbackAdapterProcessorWith2Subqueries>();
+					break;
+				case 3:
+					Observer = NewObject<UTypedElementQueryObserverCallbackAdapterProcessorWith3Subqueries>();
+					break;
+				case 4:
+					Observer = NewObject<UTypedElementQueryObserverCallbackAdapterProcessorWith4Subqueries>();
+					break;
+				case 5:
+					Observer = NewObject<UTypedElementQueryObserverCallbackAdapterProcessorWith5Subqueries>();
+					break;
+				case 6:
+					Observer = NewObject<UTypedElementQueryObserverCallbackAdapterProcessorWith6Subqueries>();
+					break;
+				case 7:
+					Observer = NewObject<UTypedElementQueryObserverCallbackAdapterProcessorWith7Subqueries>();
+					break;
+				case 8:
+					Observer = NewObject<UTypedElementQueryObserverCallbackAdapterProcessorWith8Subqueries>();
+					break;
+				default:
+					checkf(false, TEXT("The current Typed Elements Data Storage backend doesn't support %i subqueries per observer query."),
+						Query.Subqueries.Num());
+					return StoredQuery.NativeQuery;
+				}
+				StoredQuery.Processor.Reset(Observer);
+				return Observer->GetQuery();
+			}
+			case DSI::EQueryCallbackType::PhasePreparation:
 				break;
-			case 8:
-				Processor = NewObject<UTypedElementQueryProcessorCallbackAdapterProcessorWith8Subqueries>();
+			case DSI::EQueryCallbackType::PhaseFinalization:
 				break;
 			default:
-				checkf(false, TEXT("The current Typed Elements Data Storage backend doesn't support %i subqueries per processor query."), 
-					Query.Subqueries.Num());
-				return StoredQuery.NativeQuery;
+				checkf(false, TEXT("Unsupported query callback type %i."), static_cast<int>(Query.Callback.Type));
+				break;
 			}
-			StoredQuery.Processor.Reset(Processor);
-			return Processor->GetQuery();
 		}
-		case DSI::EQueryCallbackType::ObserveAdd:
-			// Fall-through
-		case DSI::EQueryCallbackType::ObserveRemove:
+		return StoredQuery.NativeQuery;
+	}
+
+	bool FExtendedQueryStore::SetupSelectedColumns(
+		ITypedElementDataStorageInterface::FQueryDescription& Query, FMassEntityQuery& NativeQuery)
+	{
+		using DSI = ITypedElementDataStorageInterface;
+
+		switch (Query.Action)
 		{
-			UTypedElementQueryObserverCallbackAdapterProcessorBase* Observer;
-			switch (Query.Subqueries.Num())
+		case DSI::FQueryDescription::EActionType::None:
+			return true;
+		case DSI::FQueryDescription::EActionType::Select:
+		{
+			const int32 SelectionCount = Query.SelectionTypes.Num();
+			if (ensureMsgf(SelectionCount == Query.SelectionAccessTypes.Num(),
+				TEXT("The number of query selection types (%i) doesn't match the number of selection access types (%i)."),
+				SelectionCount, Query.SelectionAccessTypes.Num()))
 			{
-			case 0:
-				Observer = NewObject<UTypedElementQueryObserverCallbackAdapterProcessor>();
-				break;
-			case 1:
-				Observer = NewObject<UTypedElementQueryObserverCallbackAdapterProcessorWith1Subquery>();
-				break;
-			case 2:
-				Observer = NewObject<UTypedElementQueryObserverCallbackAdapterProcessorWith2Subqueries>();
-				break;
-			case 3:
-				Observer = NewObject<UTypedElementQueryObserverCallbackAdapterProcessorWith3Subqueries>();
-				break;
-			case 4:
-				Observer = NewObject<UTypedElementQueryObserverCallbackAdapterProcessorWith4Subqueries>();
-				break;
-			case 5:
-				Observer = NewObject<UTypedElementQueryObserverCallbackAdapterProcessorWith5Subqueries>();
-				break;
-			case 6:
-				Observer = NewObject<UTypedElementQueryObserverCallbackAdapterProcessorWith6Subqueries>();
-				break;
-			case 7:
-				Observer = NewObject<UTypedElementQueryObserverCallbackAdapterProcessorWith7Subqueries>();
-				break;
-			case 8:
-				Observer = NewObject<UTypedElementQueryObserverCallbackAdapterProcessorWith8Subqueries>();
-				break;
-			default:
-				checkf(false, TEXT("The current Typed Elements Data Storage backend doesn't support %i subqueries per observer query."),
-					Query.Subqueries.Num());
-				return StoredQuery.NativeQuery;
+				for (int SelectionIndex = 0; SelectionIndex < SelectionCount; ++SelectionIndex)
+				{
+					TWeakObjectPtr<const UScriptStruct>& Type = Query.SelectionTypes[SelectionIndex];
+					DSI::EQueryAccessType AccessType = Query.SelectionAccessTypes[SelectionIndex];
+					if (ensureMsgf(Type.IsValid(), TEXT("Provided query selection type can not be null.")) &&
+						ensureMsgf(
+							Type->IsChildOf(FTypedElementDataStorageColumn::StaticStruct()) ||
+							Type->IsChildOf(FMassFragment::StaticStruct()),
+							TEXT("Provided query selection type '%s' is not based on FTypedElementDataStorageColumn or another supported base type."),
+							*Type->GetStructPathName().ToString()))
+					{
+						NativeQuery.AddRequirement(Type.Get(), ConvertToNativeAccessType(AccessType), ConvertToNativePresenceType(AccessType));
+					}
+					else
+					{
+						return false;
+					}
+				}
+				return true;
 			}
-			StoredQuery.Processor.Reset(Observer);
-			return Observer->GetQuery();
+			return false;
+		}
+		case DSI::FQueryDescription::EActionType::Count:
+		{
+			bool bIsSelectionEmpty = Query.SelectionTypes.IsEmpty();
+			bool bIsAccessTypesEmpty = Query.SelectionAccessTypes.IsEmpty();
+			checkf(bIsSelectionEmpty, TEXT("Count queries for the Typed Elements Data Storage can't have entries for selection."));
+			checkf(bIsAccessTypesEmpty, TEXT("Count queries for the Typed Elements Data Storage can't have entries for selection."));
+			return bIsSelectionEmpty && bIsAccessTypesEmpty;
 		}
-		case DSI::EQueryCallbackType::PhasePreparation:
-			break;
-		case DSI::EQueryCallbackType::PhaseFinalization:
-			break;
 		default:
-			checkf(false, TEXT("Unsupported query callback type %i."), static_cast<int>(Query.Callback.Type));
-			break;
+			checkf(false, TEXT("Unexpected query action: %i."), static_cast<int32>(Query.Action));
+			return false;
 		}
 	}
-	return StoredQuery.NativeQuery;
-}
 
-bool FTypedElementExtendedQueryStore::SetupSelectedColumns(
-	ITypedElementDataStorageInterface::FQueryDescription& Query, FMassEntityQuery& NativeQuery)
-{
-	using DSI = ITypedElementDataStorageInterface;
-
-	switch (Query.Action)
-	{
-	case DSI::FQueryDescription::EActionType::None:
-		return true;
-	case DSI::FQueryDescription::EActionType::Select:
+	bool FExtendedQueryStore::SetupConditions(
+		ITypedElementDataStorageInterface::FQueryDescription& Query, FMassEntityQuery& NativeQuery)
 	{
-		const int32 SelectionCount = Query.SelectionTypes.Num();
-		if(ensureMsgf(SelectionCount == Query.SelectionAccessTypes.Num(),
-			TEXT("The number of query selection types (%i) doesn't match the number of selection access types (%i)."),
-				SelectionCount, Query.SelectionAccessTypes.Num()))
+		using DSI = ITypedElementDataStorageInterface;
+
+		if (Query.ConditionTypes.IsEmpty())
 		{
-			for (int SelectionIndex = 0; SelectionIndex < SelectionCount; ++SelectionIndex)
+			return true;
+		}
+
+		if (ensureMsgf(Query.ConditionTypes.Num() == Query.ConditionOperators.Num(),
+			TEXT("The types and operators for a typed element query have gone out of sync.")))
+		{
+			const DSI::FQueryDescription::FOperator* Operand = Query.ConditionOperators.GetData();
+			for (DSI::FQueryDescription::EOperatorType Type : Query.ConditionTypes)
 			{
-				TWeakObjectPtr<const UScriptStruct>& Type = Query.SelectionTypes[SelectionIndex];
-				DSI::EQueryAccessType AccessType = Query.SelectionAccessTypes[SelectionIndex];
-				if (ensureMsgf(Type.IsValid(), TEXT("Provided query selection type can not be null.")) &&
-					ensureMsgf(
-						Type->IsChildOf(FTypedElementDataStorageColumn::StaticStruct()) ||
-						Type->IsChildOf(FMassFragment::StaticStruct()),
-						TEXT("Provided query selection type '%s' is not based on FTypedElementDataStorageColumn or another supported base type."),
-						*Type->GetStructPathName().ToString()))
+				EMassFragmentPresence Presence;
+				switch (Type)
 				{
-					NativeQuery.AddRequirement(Type.Get(), ConvertToNativeAccessType(AccessType), ConvertToNativePresenceType(AccessType));
+				case DSI::FQueryDescription::EOperatorType::SimpleAll:
+					Presence = EMassFragmentPresence::All;
+					break;
+				case DSI::FQueryDescription::EOperatorType::SimpleAny:
+					Presence = EMassFragmentPresence::Any;
+					break;
+				case DSI::FQueryDescription::EOperatorType::SimpleNone:
+					Presence = EMassFragmentPresence::None;
+					break;
+				default:
+					continue;
 				}
-				else
+
+				if (Operand->Type->IsChildOf(FMassTag::StaticStruct()))
 				{
-					return false;
+					NativeQuery.AddTagRequirement(*(Operand->Type), Presence);
+				}
+				else if (Operand->Type->IsChildOf(FMassFragment::StaticStruct()))
+				{
+					NativeQuery.AddRequirement(Operand->Type.Get(), EMassFragmentAccess::None, Presence);
 				}
+
+				++Operand;
 			}
 			return true;
 		}
 		return false;
 	}
-	case DSI::FQueryDescription::EActionType::Count:
-	{
-		bool bIsSelectionEmpty = Query.SelectionTypes.IsEmpty();
-		bool bIsAccessTypesEmpty = Query.SelectionAccessTypes.IsEmpty();
-		checkf(bIsSelectionEmpty, TEXT("Count queries for the Typed Elements Data Storage can't have entries for selection."));
-		checkf(bIsAccessTypesEmpty, TEXT("Count queries for the Typed Elements Data Storage can't have entries for selection."));
-		return bIsSelectionEmpty && bIsAccessTypesEmpty;
-	}
-	default:
-		checkf(false, TEXT("Unexpected query action: %i."), static_cast<int32>(Query.Action));
-		return false;
-	}
-}
 
-bool FTypedElementExtendedQueryStore::SetupConditions(
-	ITypedElementDataStorageInterface::FQueryDescription& Query, FMassEntityQuery& NativeQuery)
-{
-	using DSI = ITypedElementDataStorageInterface;
-
-	if (Query.ConditionTypes.IsEmpty())
-	{
-		return true;
-	}
+	static const FName DynamicTagDataParamName = TEXT("TagName");
 
-	if (ensureMsgf(Query.ConditionTypes.Num() == Query.ConditionOperators.Num(),
-		TEXT("The types and operators for a typed element query have gone out of sync.")))
+	bool FExtendedQueryStore::SetupChunkFilters(
+		Handle QueryHandle,
+		ITypedElementDataStorageInterface::FQueryDescription& Query,
+		FEnvironment& Environment,
+		FMassEntityQuery& NativeQuery)
 	{
-		const DSI::FQueryDescription::FOperator* Operand = Query.ConditionOperators.GetData();
-		for (DSI::FQueryDescription::EOperatorType Type : Query.ConditionTypes)
+		if (Query.DynamicTags.IsEmpty())
 		{
-			EMassFragmentPresence Presence;
-			switch (Type)
-			{
-			case DSI::FQueryDescription::EOperatorType::SimpleAll:
-				Presence = EMassFragmentPresence::All;
-				break;
-			case DSI::FQueryDescription::EOperatorType::SimpleAny:
-				Presence = EMassFragmentPresence::Any;
-				break;
-			case DSI::FQueryDescription::EOperatorType::SimpleNone:
-				Presence = EMassFragmentPresence::None;
-				break;
-			default:
-				continue;
-			}
+			return true;
+		}
 
-			if (Operand->Type->IsChildOf(FMassTag::StaticStruct()))
+		Algo::SortBy(Query.DynamicTags, [](const TypedElementDataStorage::FQueryDescription::FDynamicTagData& DynamicTagData)
 			{
-				NativeQuery.AddTagRequirement(*(Operand->Type), Presence);
-			}
-			else if (Operand->Type->IsChildOf(FMassFragment::StaticStruct()))
+				return DynamicTagData.Tag.GetName();
+			}, FNameFastLess());
+		// Check if there are any duplicate groups. Not yet supported until we can match multiple MatchTags
+		UE::Editor::DataStorage::FDynamicTag PreviousTag = Query.DynamicTags[0].Tag;
+		for (int32 Index = 1, End = Query.DynamicTags.Num(); Index < End; ++Index)
+		{
+			if (Query.DynamicTags[Index].Tag == PreviousTag)
 			{
-				NativeQuery.AddRequirement(Operand->Type.Get(), EMassFragmentAccess::None, Presence);
+				return false;
 			}
-
-			++Operand;
+			PreviousTag = Query.DynamicTags[Index].Tag;
 		}
-		return true;
-	}
-	return false;
-}
-
-static const FName DynamicTagDataParamName = TEXT("TagName");
 
-bool FTypedElementExtendedQueryStore::SetupChunkFilters(
-	Handle QueryHandle,
-	ITypedElementDataStorageInterface::FQueryDescription& Query,
-	FTypedElementDatabaseEnvironment& Environment,
-	FMassEntityQuery& NativeQuery)
-{
-	if (Query.DynamicTags.IsEmpty())
-	{
-		return true;
-	}
+		struct FGroupTagPair
+		{
+			const UScriptStruct* ColumnType;
+			FName Value;
+		};
 
-	Algo::SortBy(Query.DynamicTags, [](const TypedElementDataStorage::FQueryDescription::FDynamicTagData& DynamicTagData)
-	{
-		return DynamicTagData.Tag.GetName();
-	}, FNameFastLess());
-	// Check if there are any duplicate groups. Not yet supported until we can match multiple MatchTags
-	UE::EditorDataStorage::FDynamicTag PreviousTag = Query.DynamicTags[0].Tag;
-	for (int32 Index = 1, End = Query.DynamicTags.Num(); Index < End; ++Index)
-	{
-		if (Query.DynamicTags[Index].Tag == PreviousTag)
+		TArray<FGroupTagPair> GroupTagPairsTemp;
+		GroupTagPairsTemp.Reserve(Query.DynamicTags.Num());
+		for (int32 Index = 0, End = Query.DynamicTags.Num(); Index < End; ++Index)
 		{
-			return false;
+			const UScriptStruct* ColumnType = Environment.GenerateColumnType(Query.DynamicTags[Index].Tag);
+			GroupTagPairsTemp.Emplace(FGroupTagPair
+				{
+					.ColumnType = ColumnType,
+					.Value = Query.DynamicTags[Index].MatchValue
+				});
 		}
-		PreviousTag = Query.DynamicTags[Index].Tag;
-	}
-	
-	struct FGroupTagPair
-	{
-		const UScriptStruct* ColumnType;
-		FName Value;
-	};
-
-	TArray<FGroupTagPair> GroupTagPairsTemp;
-	GroupTagPairsTemp.Reserve(Query.DynamicTags.Num());
-	for (int32 Index = 0, End = Query.DynamicTags.Num(); Index < End; ++Index)
-	{
-		const UScriptStruct* ColumnType = Environment.GenerateColumnType(Query.DynamicTags[Index].Tag);
-		GroupTagPairsTemp.Emplace(FGroupTagPair
-			{
-				.ColumnType = ColumnType,
-				.Value = Query.DynamicTags[Index].MatchValue
-			});
-	}
 
-	check(!GroupTagPairsTemp.IsEmpty());
+		check(!GroupTagPairsTemp.IsEmpty());
 
-	for (const FGroupTagPair& Element : GroupTagPairsTemp)
-	{
-		NativeQuery.AddConstSharedRequirement(Element.ColumnType);
-	}
-	
-	auto ChunkFilterFunction = [GroupTagPairs = MoveTemp(GroupTagPairsTemp)](const FMassExecutionContext& MassContext) -> bool
-	{
-		for (const FGroupTagPair& GroupTagPair : GroupTagPairs)
+		for (const FGroupTagPair& Element : GroupTagPairsTemp)
 		{
-			const void* SharedFragmentData = MassContext.GetConstSharedFragmentPtr(*GroupTagPair.ColumnType);
+			NativeQuery.AddConstSharedRequirement(Element.ColumnType);
+		}
 
-			if (SharedFragmentData)
+		auto ChunkFilterFunction = [GroupTagPairs = MoveTemp(GroupTagPairsTemp)](const FMassExecutionContext& MassContext) -> bool
+		{
+			for (const FGroupTagPair& GroupTagPair : GroupTagPairs)
 			{
-				// TODO: Use reflection / knowledge of key offset to read the tag and compare it
-				// For now... we know it is at offset 0
-				const FTypedElementDataStorageDynamicTagTemplate* TagOverlay = static_cast<const FTypedElementDataStorageDynamicTagTemplate*>(SharedFragmentData);
-				// NAME_None will match any presence of the shared fragment
-				// otherwise, match the specific tag only
-				if (GroupTagPair.Value != NAME_None && TagOverlay->Value != GroupTagPair.Value)
+				const void* SharedFragmentData = MassContext.GetConstSharedFragmentPtr(*GroupTagPair.ColumnType);
+
+				if (SharedFragmentData)
+				{
+					// TODO: Use reflection / knowledge of key offset to read the tag and compare it
+					// For now... we know it is at offset 0
+					const FDynamicTagColumn* TagOverlay = static_cast<const FDynamicTagColumn*>(SharedFragmentData);
+					// NAME_None will match any presence of the shared fragment
+					// otherwise, match the specific tag only
+					if (GroupTagPair.Value != NAME_None && TagOverlay->Value != GroupTagPair.Value)
+					{
+						return false;
+					}
+				}
+				else
 				{
 					return false;
 				}
 			}
-			else
-			{
-				return false;
-			}
-		}
-		
+
+			return true;
+		};
+
+		NativeQuery.SetChunkFilter(ChunkFilterFunction);
 		return true;
-	};
-	
-	NativeQuery.SetChunkFilter(ChunkFilterFunction);
-	return true;
-}
-
-bool FTypedElementExtendedQueryStore::SetupDependencies(
-	ITypedElementDataStorageInterface::FQueryDescription& Query, FMassEntityQuery& NativeQuery)
-{
-	using DSI = ITypedElementDataStorageInterface;
+	}
 
-	const int32 DependencyCount = Query.DependencyTypes.Num();
-	if (ensureMsgf(DependencyCount == Query.DependencyFlags.Num() && DependencyCount == Query.CachedDependencies.Num(),
-		TEXT("The number of query dependencies (%i) doesn't match the number of dependency access types (%i) and/or cached dependencies count (%i)."),
-		DependencyCount, Query.DependencyFlags.Num(), Query.CachedDependencies.Num()))
+	bool FExtendedQueryStore::SetupDependencies(
+		ITypedElementDataStorageInterface::FQueryDescription& Query, FMassEntityQuery& NativeQuery)
 	{
-		for (int32 DependencyIndex = 0; DependencyIndex < DependencyCount; ++DependencyIndex)
+		using DSI = ITypedElementDataStorageInterface;
+
+		const int32 DependencyCount = Query.DependencyTypes.Num();
+		if (ensureMsgf(DependencyCount == Query.DependencyFlags.Num() && DependencyCount == Query.CachedDependencies.Num(),
+			TEXT("The number of query dependencies (%i) doesn't match the number of dependency access types (%i) and/or cached dependencies count (%i)."),
+			DependencyCount, Query.DependencyFlags.Num(), Query.CachedDependencies.Num()))
 		{
-			TWeakObjectPtr<const UClass>& Type = Query.DependencyTypes[DependencyIndex];
-			if (ensureMsgf(Type.IsValid(), TEXT("Provided query dependency type can not be null.")) &&
-				ensureMsgf(Type->IsChildOf<USubsystem>(), TEXT("Provided query dependency type '%s' is not based on USubSystem."),
-					*Type->GetStructPathName().ToString()))
+			for (int32 DependencyIndex = 0; DependencyIndex < DependencyCount; ++DependencyIndex)
 			{
-				DSI::EQueryDependencyFlags Flags = Query.DependencyFlags[DependencyIndex];
-				NativeQuery.AddSubsystemRequirement(
-					const_cast<UClass*>(Type.Get()),
-					EnumHasAllFlags(Flags, DSI::EQueryDependencyFlags::ReadOnly) ? EMassFragmentAccess::ReadOnly : EMassFragmentAccess::ReadWrite,
-					EnumHasAllFlags(Flags, DSI::EQueryDependencyFlags::GameThreadBound));
-			}
-			else
-			{
-				return false;
+				TWeakObjectPtr<const UClass>& Type = Query.DependencyTypes[DependencyIndex];
+				if (ensureMsgf(Type.IsValid(), TEXT("Provided query dependency type can not be null.")) &&
+					ensureMsgf(Type->IsChildOf<USubsystem>(), TEXT("Provided query dependency type '%s' is not based on USubSystem."),
+						*Type->GetStructPathName().ToString()))
+				{
+					DSI::EQueryDependencyFlags Flags = Query.DependencyFlags[DependencyIndex];
+					NativeQuery.AddSubsystemRequirement(
+						const_cast<UClass*>(Type.Get()),
+						EnumHasAllFlags(Flags, DSI::EQueryDependencyFlags::ReadOnly) ? EMassFragmentAccess::ReadOnly : EMassFragmentAccess::ReadWrite,
+						EnumHasAllFlags(Flags, DSI::EQueryDependencyFlags::GameThreadBound));
+				}
+				else
+				{
+					return false;
+				}
 			}
+			return true;
 		}
-		return true;
+		return false;
 	}
-	return false;
-}
 
-bool FTypedElementExtendedQueryStore::SetupTickGroupDefaults(ITypedElementDataStorageInterface::FQueryDescription& Query)
-{
-	const FTickGroupDescription* TickGroup = TickGroupDescriptions.Find({ Query.Callback.Group, Query.Callback.Phase });
-	if (TickGroup)
+	bool FExtendedQueryStore::SetupTickGroupDefaults(ITypedElementDataStorageInterface::FQueryDescription& Query)
 	{
-		for (auto It = Query.Callback.BeforeGroups.CreateIterator(); It; ++It)
+		const FTickGroupDescription* TickGroup = TickGroupDescriptions.Find({ Query.Callback.Group, Query.Callback.Phase });
+		if (TickGroup)
 		{
-			if (TickGroup->BeforeGroups.Contains(*It))
+			for (auto It = Query.Callback.BeforeGroups.CreateIterator(); It; ++It)
 			{
-				It.RemoveCurrentSwap();
+				if (TickGroup->BeforeGroups.Contains(*It))
+				{
+					It.RemoveCurrentSwap();
+				}
 			}
-		}
-		Query.Callback.BeforeGroups.Append(TickGroup->BeforeGroups);
-		
-		for (auto It = Query.Callback.AfterGroups.CreateIterator(); It; ++It)
-		{
-			if (TickGroup->AfterGroups.Contains(*It))
+			Query.Callback.BeforeGroups.Append(TickGroup->BeforeGroups);
+
+			for (auto It = Query.Callback.AfterGroups.CreateIterator(); It; ++It)
 			{
-				It.RemoveCurrentSwap();
+				if (TickGroup->AfterGroups.Contains(*It))
+				{
+					It.RemoveCurrentSwap();
+				}
 			}
+			Query.Callback.AfterGroups.Append(TickGroup->AfterGroups);
+
+			Query.Callback.bForceToGameThread |= TickGroup->bRequiresMainThread;
 		}
-		Query.Callback.AfterGroups.Append(TickGroup->AfterGroups);
-		
-		Query.Callback.bForceToGameThread |= TickGroup->bRequiresMainThread;
+		return true;
 	}
-	return true;
-}
 
-bool FTypedElementExtendedQueryStore::SetupProcessors(Handle QueryHandle, FTypedElementExtendedQuery& StoredQuery, 
-	FTypedElementDatabaseEnvironment& Environment, FMassEntityManager& EntityManager, FMassProcessingPhaseManager& PhaseManager)
-{
-	using DSI = ITypedElementDataStorageInterface;
-
-	// Register Phase processors locally.
-	switch (StoredQuery.Description.Callback.Type)
+	bool FExtendedQueryStore::SetupProcessors(Handle QueryHandle, FExtendedQuery& StoredQuery,
+		FEnvironment& Environment, FMassEntityManager& EntityManager, FMassProcessingPhaseManager& PhaseManager)
 	{
-	case DSI::EQueryCallbackType::PhasePreparation:
-		RegisterPreambleQuery(StoredQuery.Description.Callback.Phase, QueryHandle);
-		break;
-	case DSI::EQueryCallbackType::PhaseFinalization:
-		RegisterPostambleQuery(StoredQuery.Description.Callback.Phase, QueryHandle);
-		break;
-	}
+		using DSI = ITypedElementDataStorageInterface;
 
-	// Register regular processors and observer with Mass.
-	if (StoredQuery.Processor)
-	{
-		if (StoredQuery.Processor->IsA<UTypedElementQueryProcessorCallbackAdapterProcessorBase>())
+		// Register Phase processors locally.
+		switch (StoredQuery.Description.Callback.Type)
 		{
-			if (static_cast<UTypedElementQueryProcessorCallbackAdapterProcessorBase*>(StoredQuery.Processor.Get())->
-				ConfigureQueryCallback(StoredQuery, QueryHandle, *this, Environment))
-			{
-				PhaseManager.RegisterDynamicProcessor(*StoredQuery.Processor);
-			}
-			else
-			{
-				return false;
-			}
+		case DSI::EQueryCallbackType::PhasePreparation:
+			RegisterPreambleQuery(StoredQuery.Description.Callback.Phase, QueryHandle);
+			break;
+		case DSI::EQueryCallbackType::PhaseFinalization:
+			RegisterPostambleQuery(StoredQuery.Description.Callback.Phase, QueryHandle);
+			break;
 		}
-		else if (StoredQuery.Processor->IsA<UTypedElementQueryObserverCallbackAdapterProcessorBase>())
+
+		// Register regular processors and observer with Mass.
+		if (StoredQuery.Processor)
 		{
-			if (UTypedElementQueryObserverCallbackAdapterProcessorBase* Observer =
-				static_cast<UTypedElementQueryObserverCallbackAdapterProcessorBase*>(StoredQuery.Processor.Get()))
+			if (StoredQuery.Processor->IsA<UTypedElementQueryProcessorCallbackAdapterProcessorBase>())
+			{
+				if (static_cast<UTypedElementQueryProcessorCallbackAdapterProcessorBase*>(StoredQuery.Processor.Get())->
+					ConfigureQueryCallback(StoredQuery, QueryHandle, *this, Environment))
+				{
+					PhaseManager.RegisterDynamicProcessor(*StoredQuery.Processor);
+				}
+				else
+				{
+					return false;
+				}
+			}
+			else if (StoredQuery.Processor->IsA<UTypedElementQueryObserverCallbackAdapterProcessorBase>())
 			{
-				Observer->ConfigureQueryCallback(StoredQuery, QueryHandle, *this, Environment);
-				EntityManager.GetObserverManager().AddObserverInstance(*Observer->GetObservedType(), Observer->GetObservedOperation(), *Observer);
+				if (UTypedElementQueryObserverCallbackAdapterProcessorBase* Observer =
+					static_cast<UTypedElementQueryObserverCallbackAdapterProcessorBase*>(StoredQuery.Processor.Get()))
+				{
+					Observer->ConfigureQueryCallback(StoredQuery, QueryHandle, *this, Environment);
+					EntityManager.GetObserverManager().AddObserverInstance(*Observer->GetObservedType(), Observer->GetObservedOperation(), *Observer);
+				}
+				else
+				{
+					return false;
+				}
 			}
 			else
 			{
+				checkf(false, TEXT("Query processor %s is of unsupported type %s."),
+					*StoredQuery.Description.Callback.Name.ToString(), *StoredQuery.Processor->GetSparseClassDataStruct()->GetName());
 				return false;
 			}
 		}
-		else
+		return true;
+	}
+
+	bool FExtendedQueryStore::SetupActivatable(Handle QueryHandle, ITypedElementDataStorageInterface::FQueryDescription& Query)
+	{
+		if (!Query.Callback.ActivationName.IsNone())
 		{
-			checkf(false, TEXT("Query processor %s is of unsupported type %s."),
-				*StoredQuery.Description.Callback.Name.ToString(), *StoredQuery.Processor->GetSparseClassDataStruct()->GetName());
-			return false;
+			ActivatableMapping.Add(Query.Callback.ActivationName, QueryHandle);
 		}
+		return true;
 	}
-	return true;
-}
 
-bool FTypedElementExtendedQueryStore::SetupActivatable(Handle QueryHandle, ITypedElementDataStorageInterface::FQueryDescription& Query)
-{
-	if (!Query.Callback.ActivationName.IsNone())
+	EMassFragmentAccess FExtendedQueryStore::ConvertToNativeAccessType(ITypedElementDataStorageInterface::EQueryAccessType AccessType)
 	{
-		ActivatableMapping.Add(Query.Callback.ActivationName, QueryHandle);
+		switch (AccessType)
+		{
+		case ITypedElementDataStorageInterface::EQueryAccessType::ReadOnly:
+			// Fall through
+		case ITypedElementDataStorageInterface::EQueryAccessType::OptionalReadOnly:
+			return EMassFragmentAccess::ReadOnly;
+		case ITypedElementDataStorageInterface::EQueryAccessType::ReadWrite:
+			return EMassFragmentAccess::ReadWrite;
+		default:
+			checkf(false, TEXT("Invalid query access type: %i."), static_cast<uint32>(AccessType));
+			return EMassFragmentAccess::MAX;
+		}
 	}
-	return true;
-}
 
-EMassFragmentAccess FTypedElementExtendedQueryStore::ConvertToNativeAccessType(ITypedElementDataStorageInterface::EQueryAccessType AccessType)
-{
-	switch (AccessType)
+	EMassFragmentPresence FExtendedQueryStore::ConvertToNativePresenceType(ITypedElementDataStorageInterface::EQueryAccessType AccessType)
 	{
-	case ITypedElementDataStorageInterface::EQueryAccessType::ReadOnly:
-		// Fall through
-	case ITypedElementDataStorageInterface::EQueryAccessType::OptionalReadOnly:
-		return EMassFragmentAccess::ReadOnly;
-	case ITypedElementDataStorageInterface::EQueryAccessType::ReadWrite:
-		return EMassFragmentAccess::ReadWrite;
-	default:
-		checkf(false, TEXT("Invalid query access type: %i."), static_cast<uint32>(AccessType));
-		return EMassFragmentAccess::MAX;
+		switch (AccessType)
+		{
+		case ITypedElementDataStorageInterface::EQueryAccessType::ReadOnly:
+			return EMassFragmentPresence::All;
+		case ITypedElementDataStorageInterface::EQueryAccessType::OptionalReadOnly:
+			return EMassFragmentPresence::Optional;
+		case ITypedElementDataStorageInterface::EQueryAccessType::ReadWrite:
+			return EMassFragmentPresence::All;
+		default:
+			checkf(false, TEXT("Invalid query access type: %i."), static_cast<uint32>(AccessType));
+			return EMassFragmentPresence::MAX;
+		}
 	}
-}
 
-EMassFragmentPresence FTypedElementExtendedQueryStore::ConvertToNativePresenceType(ITypedElementDataStorageInterface::EQueryAccessType AccessType)
-{
-	switch (AccessType)
+	void FExtendedQueryStore::RegisterPreambleQuery(ITypedElementDataStorageInterface::EQueryTickPhase Phase, Handle Query)
 	{
-	case ITypedElementDataStorageInterface::EQueryAccessType::ReadOnly:
-		return EMassFragmentPresence::All;
-	case ITypedElementDataStorageInterface::EQueryAccessType::OptionalReadOnly:
-		return EMassFragmentPresence::Optional;
-	case ITypedElementDataStorageInterface::EQueryAccessType::ReadWrite:
-		return EMassFragmentPresence::All;
-	default:
-		checkf(false, TEXT("Invalid query access type: %i."), static_cast<uint32>(AccessType));
-		return EMassFragmentPresence::MAX;
+		PhasePreparationQueries[static_cast<QueryTickPhaseType>(Phase)].Add(Query);
 	}
-}
-
-void FTypedElementExtendedQueryStore::RegisterPreambleQuery(ITypedElementDataStorageInterface::EQueryTickPhase Phase, Handle Query)
-{
-	PhasePreparationQueries[static_cast<QueryTickPhaseType>(Phase)].Add(Query);
-}
-
-void FTypedElementExtendedQueryStore::RegisterPostambleQuery(ITypedElementDataStorageInterface::EQueryTickPhase Phase, Handle Query)
-{
-	PhaseFinalizationQueries[static_cast<QueryTickPhaseType>(Phase)].Add(Query);
-}
 
-void FTypedElementExtendedQueryStore::UnregisterPreambleQuery(ITypedElementDataStorageInterface::EQueryTickPhase Phase, Handle Query)
-{
-	int32 Index;
-	if (PhasePreparationQueries[static_cast<QueryTickPhaseType>(Phase)].Find(Query, Index))
+	void FExtendedQueryStore::RegisterPostambleQuery(ITypedElementDataStorageInterface::EQueryTickPhase Phase, Handle Query)
 	{
-		PhasePreparationQueries[static_cast<QueryTickPhaseType>(Phase)].RemoveAt(Index);
+		PhaseFinalizationQueries[static_cast<QueryTickPhaseType>(Phase)].Add(Query);
 	}
-}
 
-void FTypedElementExtendedQueryStore::UnregisterPostambleQuery(ITypedElementDataStorageInterface::EQueryTickPhase Phase, Handle Query)
-{
-	int32 Index;
-	if (PhaseFinalizationQueries[static_cast<QueryTickPhaseType>(Phase)].Find(Query, Index))
+	void FExtendedQueryStore::UnregisterPreambleQuery(ITypedElementDataStorageInterface::EQueryTickPhase Phase, Handle Query)
 	{
-		PhaseFinalizationQueries[static_cast<QueryTickPhaseType>(Phase)].RemoveAt(Index);
+		int32 Index;
+		if (PhasePreparationQueries[static_cast<QueryTickPhaseType>(Phase)].Find(Query, Index))
+		{
+			PhasePreparationQueries[static_cast<QueryTickPhaseType>(Phase)].RemoveAt(Index);
+		}
 	}
-}
 
-void FTypedElementExtendedQueryStore::RunPhasePreOrPostAmbleQueries(FMassEntityManager& EntityManager,
-	FTypedElementDatabaseEnvironment& Environment, ITypedElementDataStorageInterface::EQueryTickPhase Phase,
-	float DeltaTime, TArray<Handle>& QueryHandles)
-{
-	if (!QueryHandles.IsEmpty())
+	void FExtendedQueryStore::UnregisterPostambleQuery(ITypedElementDataStorageInterface::EQueryTickPhase Phase, Handle Query)
 	{
-		FPhasePreOrPostAmbleExecutor Executor(EntityManager, DeltaTime);
-		for (Handle Query : QueryHandles)
+		int32 Index;
+		if (PhaseFinalizationQueries[static_cast<QueryTickPhaseType>(Phase)].Find(Query, Index))
 		{
-			FTypedElementExtendedQuery& QueryData = Queries.Get(Query);
-			Executor.ExecuteQuery(QueryData.Description, *this, Environment, QueryData.NativeQuery, QueryData.Description.Callback.Function);
+			PhaseFinalizationQueries[static_cast<QueryTickPhaseType>(Phase)].RemoveAt(Index);
 		}
 	}
-}
 
-void FTypedElementExtendedQueryStore::UnregisterQueryData(Handle Query, FTypedElementExtendedQuery& QueryData, FMassEntityManager& EntityManager, FMassProcessingPhaseManager& PhaseManager)
-{
-	if (!QueryData.Description.Callback.ActivationName.IsNone())
+	void FExtendedQueryStore::RunPhasePreOrPostAmbleQueries(FMassEntityManager& EntityManager,
+		FEnvironment& Environment, ITypedElementDataStorageInterface::EQueryTickPhase Phase,
+		float DeltaTime, TArray<Handle>& QueryHandles)
 	{
-		ActivatableMapping.RemoveSingle(QueryData.Description.Callback.ActivationName, Query);
-		ActiveActivatables.RemoveSingleSwap(Query);
-		PendingActivatables.RemoveSingleSwap(Query);
+		if (!QueryHandles.IsEmpty())
+		{
+			FPhasePreOrPostAmbleExecutor Executor(EntityManager, DeltaTime);
+			for (Handle Query : QueryHandles)
+			{
+				FExtendedQuery& QueryData = Queries.Get(Query);
+				Executor.ExecuteQuery(QueryData.Description, *this, Environment, QueryData.NativeQuery, QueryData.Description.Callback.Function);
+			}
+		}
 	}
 
-	if (QueryData.Processor)
+	void FExtendedQueryStore::UnregisterQueryData(Handle Query, FExtendedQuery& QueryData, FMassEntityManager& EntityManager, FMassProcessingPhaseManager& PhaseManager)
 	{
-		if (QueryData.Processor->IsA<UTypedElementQueryProcessorCallbackAdapterProcessorBase>())
+		if (!QueryData.Description.Callback.ActivationName.IsNone())
 		{
-			PhaseManager.UnregisterDynamicProcessor(*QueryData.Processor);
+			ActivatableMapping.RemoveSingle(QueryData.Description.Callback.ActivationName, Query);
+			ActiveActivatables.RemoveSingleSwap(Query);
+			PendingActivatables.RemoveSingleSwap(Query);
 		}
-		else if (QueryData.Processor->IsA<UTypedElementQueryObserverCallbackAdapterProcessorBase>())
+
+		if (QueryData.Processor)
 		{
-			UTypedElementQueryObserverCallbackAdapterProcessorBase* Observer =
-				static_cast<UTypedElementQueryObserverCallbackAdapterProcessorBase*>(QueryData.Processor.Get());
-			if (ensure(Observer))
+			if (QueryData.Processor->IsA<UTypedElementQueryProcessorCallbackAdapterProcessorBase>())
 			{
-				EntityManager.GetObserverManager().RemoveObserverInstance(*Observer->GetObservedType(), Observer->GetObservedOperation(), *Observer);
+				PhaseManager.UnregisterDynamicProcessor(*QueryData.Processor);
 			}
+			else if (QueryData.Processor->IsA<UTypedElementQueryObserverCallbackAdapterProcessorBase>())
+			{
+				UTypedElementQueryObserverCallbackAdapterProcessorBase* Observer =
+					static_cast<UTypedElementQueryObserverCallbackAdapterProcessorBase*>(QueryData.Processor.Get());
+				if (ensure(Observer))
+				{
+					EntityManager.GetObserverManager().RemoveObserverInstance(*Observer->GetObservedType(), Observer->GetObservedOperation(), *Observer);
+				}
+			}
+			else
+			{
+				checkf(false, TEXT("Query processor %s is of unsupported type %s."),
+					*QueryData.Description.Callback.Name.ToString(), *QueryData.Processor->GetSparseClassDataStruct()->GetName());
+			}
+		}
+		else if (QueryData.Description.Callback.Type == ITypedElementDataStorageInterface::EQueryCallbackType::PhasePreparation)
+		{
+			UnregisterPreambleQuery(QueryData.Description.Callback.Phase, Query);
+		}
+		else if (QueryData.Description.Callback.Type == ITypedElementDataStorageInterface::EQueryCallbackType::PhaseFinalization)
+		{
+			UnregisterPostambleQuery(QueryData.Description.Callback.Phase, Query);
 		}
 		else
 		{
-			checkf(false, TEXT("Query processor %s is of unsupported type %s."),
-				*QueryData.Description.Callback.Name.ToString(), *QueryData.Processor->GetSparseClassDataStruct()->GetName());
+			QueryData.NativeQuery.Clear();
 		}
 	}
-	else if (QueryData.Description.Callback.Type == ITypedElementDataStorageInterface::EQueryCallbackType::PhasePreparation)
-	{
-		UnregisterPreambleQuery(QueryData.Description.Callback.Phase, Query);
-	}
-	else if (QueryData.Description.Callback.Type == ITypedElementDataStorageInterface::EQueryCallbackType::PhaseFinalization)
-	{
-		UnregisterPostambleQuery(QueryData.Description.Callback.Phase, Query);
-	}
-	else
-	{
-		QueryData.NativeQuery.Clear();
-	}
-}
+} // namespace UE::Editor::DataStorage
