[Teds] Rename native classes and move into UE::EditorDataStorage namespaces

diff --git a/Engine/Plugins/Experimental/EditorDataStorage/Source/TedsCore/Private/Queries/TypedElementExtendedQueryStore.h b/Engine/Plugins/Experimental/EditorDataStorage/Source/TedsCore/Private/Queries/TypedElementExtendedQueryStore.h
--- a/Engine/Plugins/Experimental/EditorDataStorage/Source/TedsCore/Private/Queries/TypedElementExtendedQueryStore.h
+++ b/Engine/Plugins/Experimental/EditorDataStorage/Source/TedsCore/Private/Queries/TypedElementExtendedQueryStore.h
@@ -12,182 +12,186 @@ struct FMassEntityManager;
 struct FMassProcessingPhaseManager;
 class UMassProcessor;
 class FOutputDevice;
-class FTypedElementDatabaseEnvironment;
 
-struct FTypedElementExtendedQuery
+namespace UE::Editor::DataStorage
 {
-	FMassEntityQuery NativeQuery; // Used if there's no processor bound.
-	ITypedElementDataStorageInterface::FQueryDescription Description;
-	TStrongObjectPtr<UMassProcessor> Processor;
-};
-
-/**
- * Storage and utilities for Typed Element queries after they've been processed by the Data Storage implementation.
- */
-class FTypedElementExtendedQueryStore
-{
-private:
-	using QueryStore = TTypedElementHandleStore<FTypedElementExtendedQuery>;
-public:
-	using Handle = QueryStore::Handle;
-	using ListAliveEntriesConstCallback = QueryStore::ListAliveEntriesConstCallback;
-
-	/**
-	 * @section Registration
-	 * @description A set of functions to manage the registration of queries.
-	 */
-
-	/** Adds a new query to the store and initializes the query with the provided arguments. */
-	Handle RegisterQuery(
-		ITypedElementDataStorageInterface::FQueryDescription Query, 
-		FTypedElementDatabaseEnvironment& Environment,
-		FMassEntityManager& EntityManager,
-		FMassProcessingPhaseManager& PhaseManager);
-	/** Removes the query at the given handle if still alive and otherwise does nothing. */
-	void UnregisterQuery(Handle Query, FMassEntityManager& EntityManager, FMassProcessingPhaseManager& PhaseManager);
-
-	/** Removes all data in the query store. */
-	void Clear(FMassEntityManager& EntityManager, FMassProcessingPhaseManager& PhaseManager);
-	
-	/** Register the defaults for a tick group. These will be applied on top of any settings provided with a query registration. */
-	void RegisterTickGroup(FName GroupName, ITypedElementDataStorageInterface::EQueryTickPhase Phase,
-		FName BeforeGroup, FName AfterGroup, bool bRequiresMainThread);
-	/** Removes a previously registered set of tick group defaults. */
-	void UnregisterTickGroup(FName GroupName, ITypedElementDataStorageInterface::EQueryTickPhase Phase);
-
-	/**
-	 * @section Retrieval
-	 * @description Functions to retrieve data or information on queries.
-	 */
-
-	/** Retrieves the query at the provided handle, if still alive or otherwise returns a nullptr. */
-	FTypedElementExtendedQuery* Get(Handle Entry);
-	/** Retrieves the query at the provided handle, if still alive or otherwise returns a nullptr. */
-	FTypedElementExtendedQuery* GetMutable(Handle Entry);
-	/** Retrieves the query at the provided handle, if still alive or otherwise returns a nullptr. */
-	const FTypedElementExtendedQuery* Get(Handle Entry) const;
-
-	/** Retrieves the query at the provided handle, if still alive. It's up to the caller to guarantee the query is still alive. */
-	FTypedElementExtendedQuery& GetChecked(Handle Entry);
-	/** Retrieves the query at the provided handle, if still alive. It's up to the caller to guarantee the query is still alive. */
-	FTypedElementExtendedQuery& GetMutableChecked(Handle Entry);
-	/** Retrieves the query at the provided handle, if still alive. It's up to the caller to guarantee the query is still alive. */
-	const FTypedElementExtendedQuery& GetChecked(Handle Entry) const;
-
-	/** Gets the original description used to create an extended query or an empty default if the provided query isn't alive. */
-	const ITypedElementDataStorageInterface::FQueryDescription& GetQueryDescription(Handle Query) const;
-
-	/** Checks to see if a query is still available or has been removed. */
-	bool IsAlive(Handle Entry) const;
+	class FEnvironment;
 
-	/** Calls the provided callback for each query that's available. */
-	void ListAliveEntries(const ListAliveEntriesConstCallback& Callback) const;
-
-	/**
-	 * @section activatable queries
-	 * @description Functions to manipulate activatable queries
-	 */
-
-	/** Update the active activatable queries. In practice this means decrementing any active queries that automatically decrement. */
-	void UpdateActivatableQueries();
-	/** Triggers a query to run for a single update cycle. */
-	void ActivateQueries(FName ActivationName);
-	
-	/**
-	 * @section Execution
-	 * @description Various functions to run queries.
-	 */
-
-	TypedElementDataStorage::FQueryResult RunQuery(FMassEntityManager& EntityManager, Handle Query);
-	TypedElementDataStorage::FQueryResult RunQuery(
-		FMassEntityManager& EntityManager, 
-		FTypedElementDatabaseEnvironment& Environment,
-		Handle Query,
-		TypedElementDataStorage::DirectQueryCallbackRef Callback);
-	TypedElementDataStorage::FQueryResult RunQuery(
-		FMassEntityManager& EntityManager,
-		FTypedElementDatabaseEnvironment& Environment,
-		FMassExecutionContext& ParentContext,
-		Handle Query,
-		TypedElementDataStorage::SubqueryCallbackRef Callback);
-	TypedElementDataStorage::FQueryResult RunQuery(
-		FMassEntityManager& EntityManager, 
-		FTypedElementDatabaseEnvironment& Environment,
-		FMassExecutionContext& ParentContext,
-		Handle Query,
-		TypedElementRowHandle Row, 
-		TypedElementDataStorage::SubqueryCallbackRef Callback);
-	void RunPhasePreambleQueries(
-		FMassEntityManager& EntityManager, 
-		FTypedElementDatabaseEnvironment& Environment,
-		ITypedElementDataStorageInterface::EQueryTickPhase Phase, 
-		float DeltaTime);
-	void RunPhasePostambleQueries(
-		FMassEntityManager& EntityManager, 
-		FTypedElementDatabaseEnvironment& Environment,
-		ITypedElementDataStorageInterface::EQueryTickPhase Phase, 
-		float DeltaTime);
-
-	void DebugPrintQueryCallbacks(FOutputDevice& Output) const;
-
-private:
-	using QueryTickPhaseType = std::underlying_type_t<ITypedElementDataStorageInterface::EQueryTickPhase>;
-	static constexpr QueryTickPhaseType MaxTickPhase = static_cast<QueryTickPhaseType>(ITypedElementDataStorageInterface::EQueryTickPhase::Max);
-
-	struct FTickGroupId
+	struct FExtendedQuery
 	{
-		FName Name;
-		ITypedElementDataStorageInterface::EQueryTickPhase Phase;
-
-		friend inline uint32 GetTypeHash(const FTickGroupId& Id) { return HashCombine(GetTypeHash(Id.Name), GetTypeHash(Id.Phase)); }
-		friend inline bool operator==(const FTickGroupId& Lhs, const FTickGroupId& Rhs) { return Lhs.Phase == Rhs.Phase && Lhs.Name == Rhs.Name; }
-		friend inline bool operator!=(const FTickGroupId& Lhs, const FTickGroupId& Rhs) { return Lhs.Phase != Rhs.Phase || Lhs.Name != Rhs.Name; }
+		FMassEntityQuery NativeQuery; // Used if there's no processor bound.
+		ITypedElementDataStorageInterface::FQueryDescription Description;
+		TStrongObjectPtr<UMassProcessor> Processor;
 	};
 
-	struct FTickGroupDescription
+	/**
+	 * Storage and utilities for Typed Element queries after they've been processed by the Data Storage implementation.
+	 */
+	class FExtendedQueryStore
 	{
-		TArray<FName> BeforeGroups;
-		TArray<FName> AfterGroups;
-		bool bRequiresMainThread{ false };
+	private:
+		using QueryStore = THandleStore<FExtendedQuery>;
+	public:
+		using Handle = QueryStore::Handle;
+		using ListAliveEntriesConstCallback = QueryStore::ListAliveEntriesConstCallback;
+
+		/**
+		 * @section Registration
+		 * @description A set of functions to manage the registration of queries.
+		 */
+
+		 /** Adds a new query to the store and initializes the query with the provided arguments. */
+		Handle RegisterQuery(
+			ITypedElementDataStorageInterface::FQueryDescription Query,
+			FEnvironment& Environment,
+			FMassEntityManager& EntityManager,
+			FMassProcessingPhaseManager& PhaseManager);
+		/** Removes the query at the given handle if still alive and otherwise does nothing. */
+		void UnregisterQuery(Handle Query, FMassEntityManager& EntityManager, FMassProcessingPhaseManager& PhaseManager);
+
+		/** Removes all data in the query store. */
+		void Clear(FMassEntityManager& EntityManager, FMassProcessingPhaseManager& PhaseManager);
+
+		/** Register the defaults for a tick group. These will be applied on top of any settings provided with a query registration. */
+		void RegisterTickGroup(FName GroupName, ITypedElementDataStorageInterface::EQueryTickPhase Phase,
+			FName BeforeGroup, FName AfterGroup, bool bRequiresMainThread);
+		/** Removes a previously registered set of tick group defaults. */
+		void UnregisterTickGroup(FName GroupName, ITypedElementDataStorageInterface::EQueryTickPhase Phase);
+
+		/**
+		 * @section Retrieval
+		 * @description Functions to retrieve data or information on queries.
+		 */
+
+		 /** Retrieves the query at the provided handle, if still alive or otherwise returns a nullptr. */
+		FExtendedQuery* Get(Handle Entry);
+		/** Retrieves the query at the provided handle, if still alive or otherwise returns a nullptr. */
+		FExtendedQuery* GetMutable(Handle Entry);
+		/** Retrieves the query at the provided handle, if still alive or otherwise returns a nullptr. */
+		const FExtendedQuery* Get(Handle Entry) const;
+
+		/** Retrieves the query at the provided handle, if still alive. It's up to the caller to guarantee the query is still alive. */
+		FExtendedQuery& GetChecked(Handle Entry);
+		/** Retrieves the query at the provided handle, if still alive. It's up to the caller to guarantee the query is still alive. */
+		FExtendedQuery& GetMutableChecked(Handle Entry);
+		/** Retrieves the query at the provided handle, if still alive. It's up to the caller to guarantee the query is still alive. */
+		const FExtendedQuery& GetChecked(Handle Entry) const;
+
+		/** Gets the original description used to create an extended query or an empty default if the provided query isn't alive. */
+		const ITypedElementDataStorageInterface::FQueryDescription& GetQueryDescription(Handle Query) const;
+
+		/** Checks to see if a query is still available or has been removed. */
+		bool IsAlive(Handle Entry) const;
+
+		/** Calls the provided callback for each query that's available. */
+		void ListAliveEntries(const ListAliveEntriesConstCallback& Callback) const;
+
+		/**
+		 * @section activatable queries
+		 * @description Functions to manipulate activatable queries
+		 */
+
+		 /** Update the active activatable queries. In practice this means decrementing any active queries that automatically decrement. */
+		void UpdateActivatableQueries();
+		/** Triggers a query to run for a single update cycle. */
+		void ActivateQueries(FName ActivationName);
+
+		/**
+		 * @section Execution
+		 * @description Various functions to run queries.
+		 */
+
+		TypedElementDataStorage::FQueryResult RunQuery(FMassEntityManager& EntityManager, Handle Query);
+		TypedElementDataStorage::FQueryResult RunQuery(
+			FMassEntityManager& EntityManager,
+			FEnvironment& Environment,
+			Handle Query,
+			TypedElementDataStorage::DirectQueryCallbackRef Callback);
+		TypedElementDataStorage::FQueryResult RunQuery(
+			FMassEntityManager& EntityManager,
+			FEnvironment& Environment,
+			FMassExecutionContext& ParentContext,
+			Handle Query,
+			TypedElementDataStorage::SubqueryCallbackRef Callback);
+		TypedElementDataStorage::FQueryResult RunQuery(
+			FMassEntityManager& EntityManager,
+			FEnvironment& Environment,
+			FMassExecutionContext& ParentContext,
+			Handle Query,
+			TypedElementRowHandle Row,
+			TypedElementDataStorage::SubqueryCallbackRef Callback);
+		void RunPhasePreambleQueries(
+			FMassEntityManager& EntityManager,
+			FEnvironment& Environment,
+			ITypedElementDataStorageInterface::EQueryTickPhase Phase,
+			float DeltaTime);
+		void RunPhasePostambleQueries(
+			FMassEntityManager& EntityManager,
+			FEnvironment& Environment,
+			ITypedElementDataStorageInterface::EQueryTickPhase Phase,
+			float DeltaTime);
+
+		void DebugPrintQueryCallbacks(FOutputDevice& Output) const;
+
+	private:
+		using QueryTickPhaseType = std::underlying_type_t<ITypedElementDataStorageInterface::EQueryTickPhase>;
+		static constexpr QueryTickPhaseType MaxTickPhase = static_cast<QueryTickPhaseType>(ITypedElementDataStorageInterface::EQueryTickPhase::Max);
+
+		struct FTickGroupId
+		{
+			FName Name;
+			ITypedElementDataStorageInterface::EQueryTickPhase Phase;
+
+			friend inline uint32 GetTypeHash(const FTickGroupId& Id) { return HashCombine(GetTypeHash(Id.Name), GetTypeHash(Id.Phase)); }
+			friend inline bool operator==(const FTickGroupId& Lhs, const FTickGroupId& Rhs) { return Lhs.Phase == Rhs.Phase && Lhs.Name == Rhs.Name; }
+			friend inline bool operator!=(const FTickGroupId& Lhs, const FTickGroupId& Rhs) { return Lhs.Phase != Rhs.Phase || Lhs.Name != Rhs.Name; }
+		};
+
+		struct FTickGroupDescription
+		{
+			TArray<FName> BeforeGroups;
+			TArray<FName> AfterGroups;
+			bool bRequiresMainThread{ false };
+		};
+
+		template<typename CallbackReference>
+		TypedElementDataStorage::FQueryResult RunQueryCallbackCommon(
+			FMassEntityManager& EntityManager,
+			FEnvironment& Environment,
+			FMassExecutionContext* ParentContext,
+			Handle Query,
+			CallbackReference Callback);
+
+		FMassEntityQuery& SetupNativeQuery(ITypedElementDataStorageInterface::FQueryDescription& Query, FExtendedQuery& StoredQuery);
+		bool SetupSelectedColumns(ITypedElementDataStorageInterface::FQueryDescription& Query, FMassEntityQuery& NativeQuery);
+		bool SetupConditions(ITypedElementDataStorageInterface::FQueryDescription& Query, FMassEntityQuery& NativeQuery);
+		bool SetupChunkFilters(Handle QueryHandle, ITypedElementDataStorageInterface::FQueryDescription& Query, FEnvironment& Environment, FMassEntityQuery& NativeQuery);
+		bool SetupDependencies(ITypedElementDataStorageInterface::FQueryDescription& Query, FMassEntityQuery& NativeQuery);
+		bool SetupTickGroupDefaults(ITypedElementDataStorageInterface::FQueryDescription& Query);
+		bool SetupProcessors(Handle QueryHandle, FExtendedQuery& StoredQuery, FEnvironment& Environment,
+			FMassEntityManager& EntityManager, FMassProcessingPhaseManager& PhaseManager);
+		bool SetupActivatable(Handle QueryHandle, ITypedElementDataStorageInterface::FQueryDescription& Query);
+
+		EMassFragmentAccess ConvertToNativeAccessType(ITypedElementDataStorageInterface::EQueryAccessType AccessType);
+		EMassFragmentPresence ConvertToNativePresenceType(ITypedElementDataStorageInterface::EQueryAccessType AccessType);
+
+		void RegisterPreambleQuery(ITypedElementDataStorageInterface::EQueryTickPhase Phase, Handle Query);
+		void RegisterPostambleQuery(ITypedElementDataStorageInterface::EQueryTickPhase Phase, Handle Query);
+		void UnregisterPreambleQuery(ITypedElementDataStorageInterface::EQueryTickPhase Phase, Handle Query);
+		void UnregisterPostambleQuery(ITypedElementDataStorageInterface::EQueryTickPhase Phase, Handle Query);
+		void RunPhasePreOrPostAmbleQueries(FMassEntityManager& EntityManager, FEnvironment& Environment,
+			ITypedElementDataStorageInterface::EQueryTickPhase Phase, float DeltaTime, TArray<Handle>& QueryHandles);
+
+		void UnregisterQueryData(Handle Query, FExtendedQuery& QueryData, FMassEntityManager& EntityManager, FMassProcessingPhaseManager& PhaseManager);
+
+		static const ITypedElementDataStorageInterface::FQueryDescription EmptyDescription;
+
+		QueryStore Queries;
+		TMultiMap<FName, Handle> ActivatableMapping;
+		TMap<FTickGroupId, FTickGroupDescription> TickGroupDescriptions;
+		TArray<Handle> PhasePreparationQueries[MaxTickPhase];
+		TArray<Handle> PhaseFinalizationQueries[MaxTickPhase];
+		TArray<Handle> PendingActivatables;
+		TArray<Handle> ActiveActivatables;
 	};
-
-	template<typename CallbackReference>
-	TypedElementDataStorage::FQueryResult RunQueryCallbackCommon(
-		FMassEntityManager& EntityManager, 
-		FTypedElementDatabaseEnvironment& Environment,
-		FMassExecutionContext* ParentContext,
-		Handle Query, 
-		CallbackReference Callback);
-
-	FMassEntityQuery& SetupNativeQuery(ITypedElementDataStorageInterface::FQueryDescription& Query, FTypedElementExtendedQuery& StoredQuery);
-	bool SetupSelectedColumns(ITypedElementDataStorageInterface::FQueryDescription& Query, FMassEntityQuery& NativeQuery);
-	bool SetupConditions(ITypedElementDataStorageInterface::FQueryDescription& Query, FMassEntityQuery& NativeQuery);
-	bool SetupChunkFilters(Handle QueryHandle, ITypedElementDataStorageInterface::FQueryDescription& Query, FTypedElementDatabaseEnvironment& Environment, FMassEntityQuery& NativeQuery);
-	bool SetupDependencies(ITypedElementDataStorageInterface::FQueryDescription& Query, FMassEntityQuery& NativeQuery);
-	bool SetupTickGroupDefaults(ITypedElementDataStorageInterface::FQueryDescription& Query);
-	bool SetupProcessors(Handle QueryHandle, FTypedElementExtendedQuery& StoredQuery, FTypedElementDatabaseEnvironment& Environment,
-		FMassEntityManager& EntityManager, FMassProcessingPhaseManager& PhaseManager);
-	bool SetupActivatable(Handle QueryHandle, ITypedElementDataStorageInterface::FQueryDescription& Query);
-
-	EMassFragmentAccess ConvertToNativeAccessType(ITypedElementDataStorageInterface::EQueryAccessType AccessType);
-	EMassFragmentPresence ConvertToNativePresenceType(ITypedElementDataStorageInterface::EQueryAccessType AccessType);
-
-	void RegisterPreambleQuery(ITypedElementDataStorageInterface::EQueryTickPhase Phase, Handle Query);
-	void RegisterPostambleQuery(ITypedElementDataStorageInterface::EQueryTickPhase Phase, Handle Query);
-	void UnregisterPreambleQuery(ITypedElementDataStorageInterface::EQueryTickPhase Phase, Handle Query);
-	void UnregisterPostambleQuery(ITypedElementDataStorageInterface::EQueryTickPhase Phase, Handle Query);
-	void RunPhasePreOrPostAmbleQueries(FMassEntityManager& EntityManager, FTypedElementDatabaseEnvironment& Environment,
-		ITypedElementDataStorageInterface::EQueryTickPhase Phase, float DeltaTime, TArray<Handle>& QueryHandles);
-
-	void UnregisterQueryData(Handle Query, FTypedElementExtendedQuery& QueryData, FMassEntityManager& EntityManager, FMassProcessingPhaseManager& PhaseManager);
-
-	static const ITypedElementDataStorageInterface::FQueryDescription EmptyDescription;
-
-	QueryStore Queries;
-	TMultiMap<FName, Handle> ActivatableMapping;
-	TMap<FTickGroupId, FTickGroupDescription> TickGroupDescriptions;
-	TArray<Handle> PhasePreparationQueries[MaxTickPhase];
-	TArray<Handle> PhaseFinalizationQueries[MaxTickPhase];
-	TArray<Handle> PendingActivatables;
-	TArray<Handle> ActiveActivatables;
-};
+} // namepsace UE::Editor::DataStorage
diff --git a/Engine/Plugins/Experimental/EditorDataStorage/Source/TedsCore/Private/TypedElementDatabase.cpp b/Engine/Plugins/Experimental/EditorDataStorage/Source/TedsCore/Private/TypedElementDatabase.cpp
--- a/Engine/Plugins/Experimental/EditorDataStorage/Source/TedsCore/Private/TypedElementDatabase.cpp
+++ b/Engine/Plugins/Experimental/EditorDataStorage/Source/TedsCore/Private/TypedElementDatabase.cpp
@@ -114,7 +114,7 @@ void UTypedElementDatabase::Initialize()
 	ActiveEditorPhaseManager = Mass->GetMutablePhaseManager();
 	if (ActiveEditorEntityManager && ActiveEditorPhaseManager)
 	{
-		Environment = MakeShared<FTypedElementDatabaseEnvironment>(*this, *ActiveEditorEntityManager, *ActiveEditorPhaseManager);
+		Environment = MakeShared<UE::Editor::DataStorage::FEnvironment>(*this, *ActiveEditorEntityManager, *ActiveEditorPhaseManager);
 
 		using PhaseType = std::underlying_type_t<EQueryTickPhase>;
 		for (PhaseType PhaseId = 0; PhaseId < static_cast<PhaseType>(EQueryTickPhase::Max); ++PhaseId)
@@ -472,7 +472,7 @@ void UTypedElementDatabase::RemoveRow(TypedElementDataStorage::RowHandle Row)
 	FMassEntityHandle Entity = FMassEntityHandle::FromNumber(Row);
 	if (ActiveEditorEntityManager && ActiveEditorEntityManager->IsEntityValid(Entity))
 	{
-		Environment->GetIndexTable().RemoveRow(UE::EditorDataStorage::EGlobalLockScope::Public, Row);
+		Environment->GetIndexTable().RemoveRow(UE::Editor::DataStorage::EGlobalLockScope::Public, Row);
 		if (ActiveEditorEntityManager->IsEntityBuilt(FMassEntityHandle::FromNumber(Row)))
 		{
 			ActiveEditorEntityManager->DestroyEntity(FMassEntityHandle::FromNumber(Row));
@@ -487,12 +487,12 @@ void UTypedElementDatabase::RemoveRow(TypedElementDataStorage::RowHandle Row)
 
 bool UTypedElementDatabase::IsRowAvailable(TypedElementDataStorage::RowHandle Row) const
 {
-	return ActiveEditorEntityManager ? FTypedElementDatabaseCommandBuffer::Execute_IsRowAvailable(*ActiveEditorEntityManager, Row) : false;
+	return ActiveEditorEntityManager ? UE::Editor::DataStorage::Legacy::FCommandBuffer::Execute_IsRowAvailable(*ActiveEditorEntityManager, Row) : false;
 }
 
 bool UTypedElementDatabase::IsRowAssigned(TypedElementDataStorage::RowHandle Row) const
 {
-	return ActiveEditorEntityManager ? FTypedElementDatabaseCommandBuffer::Execute_IsRowAssigned(*ActiveEditorEntityManager, Row) : false;
+	return ActiveEditorEntityManager ? UE::Editor::DataStorage::Legacy::FCommandBuffer::Execute_IsRowAssigned(*ActiveEditorEntityManager, Row) : false;
 }
 
 void UTypedElementDatabase::AddColumn(TypedElementRowHandle Row, const UScriptStruct* ColumnType)
@@ -501,7 +501,7 @@ void UTypedElementDatabase::AddColumn(TypedElementRowHandle Row, const UScriptSt
 	{
 		if (IsRowAssigned(Row))
 		{
-			FTypedElementDatabaseCommandBuffer::Execute_AddColumnCommand(*ActiveEditorEntityManager, Row, ColumnType);
+			UE::Editor::DataStorage::Legacy::FCommandBuffer::Execute_AddColumnCommand(*ActiveEditorEntityManager, Row, ColumnType);
 		}
 		else
 		{
@@ -543,7 +543,7 @@ void UTypedElementDatabase::RemoveColumn(TypedElementRowHandle Row, const UScrip
 	{
 		if (IsRowAssigned(Row))
 		{
-			FTypedElementDatabaseCommandBuffer::Execute_RemoveColumnCommand(*ActiveEditorEntityManager, Row, ColumnType);
+			UE::Editor::DataStorage::Legacy::FCommandBuffer::Execute_RemoveColumnCommand(*ActiveEditorEntityManager, Row, ColumnType);
 		}
 		else
 		{
@@ -592,7 +592,7 @@ void UTypedElementDatabase::AddColumns(TypedElementRowHandle Row, TConstArrayVie
 		{
 			if (ActiveEditorEntityManager->IsEntityActive(Entity))
 			{
-				FTypedElementDatabaseCommandBuffer::Execute_AddColumnsCommand(*ActiveEditorEntityManager, Row, FragmentsToAdd, TagsToAdd);
+				UE::Editor::DataStorage::Legacy::FCommandBuffer::Execute_AddColumnsCommand(*ActiveEditorEntityManager, Row, FragmentsToAdd, TagsToAdd);
 			}
 			else
 			{
@@ -602,7 +602,7 @@ void UTypedElementDatabase::AddColumns(TypedElementRowHandle Row, TConstArrayVie
 	}
 }
 
-void UTypedElementDatabase::AddColumn(TypedElementRowHandle Row, const UE::EditorDataStorage::FDynamicTag& Tag, const FName& InValue)
+void UTypedElementDatabase::AddColumn(TypedElementRowHandle Row, const UE::Editor::DataStorage::FDynamicTag& Tag, const FName& InValue)
 {
 	if (ActiveEditorEntityManager)
 	{
@@ -611,12 +611,12 @@ void UTypedElementDatabase::AddColumn(TypedElementRowHandle Row, const UE::Edito
 		FMassEntityHandle Entity = FMassEntityHandle::FromNumber(Row);
 		if (ActiveEditorEntityManager->IsEntityActive(Entity))
 		{
-			FTypedElementDatabaseCommandBuffer::Execute_AddSharedColumnCommand(*ActiveEditorEntityManager, Row, SharedStruct);
+			UE::Editor::DataStorage::Legacy::FCommandBuffer::Execute_AddSharedColumnCommand(*ActiveEditorEntityManager, Row, SharedStruct);
 		}
 	}
 }
 
-void UTypedElementDatabase::RemoveColumn(TypedElementRowHandle Row, const UE::EditorDataStorage::FDynamicTag& Tag)
+void UTypedElementDatabase::RemoveColumn(TypedElementRowHandle Row, const UE::Editor::DataStorage::FDynamicTag& Tag)
 {
 	if (ActiveEditorEntityManager)
 	{
@@ -624,7 +624,7 @@ void UTypedElementDatabase::RemoveColumn(TypedElementRowHandle Row, const UE::Ed
 		FMassEntityHandle Entity = FMassEntityHandle::FromNumber(Row);
 		if (ActiveEditorEntityManager->IsEntityActive(Entity))
 		{
-			FTypedElementDatabaseCommandBuffer::Execute_RemoveSharedColumnCommand(*ActiveEditorEntityManager, Row, *DynamicTagType);
+			UE::Editor::DataStorage::Legacy::FCommandBuffer::Execute_RemoveSharedColumnCommand(*ActiveEditorEntityManager, Row, *DynamicTagType);
 		}
 	}
 }
@@ -642,7 +642,7 @@ void UTypedElementDatabase::RemoveColumns(TypedElementRowHandle Row, TConstArray
 		{
 			if (ActiveEditorEntityManager->IsEntityActive(Entity))
 			{
-				FTypedElementDatabaseCommandBuffer::Execute_RemoveColumnsCommand(*ActiveEditorEntityManager, Row, FragmentsToRemove, TagsToRemove);
+				UE::Editor::DataStorage::Legacy::FCommandBuffer::Execute_RemoveColumnsCommand(*ActiveEditorEntityManager, Row, FragmentsToRemove, TagsToRemove);
 			}
 			else
 			{
@@ -672,11 +672,11 @@ void UTypedElementDatabase::AddRemoveColumns(TypedElementRowHandle Row,
 		{
 			if (bMustAddColumns)
 			{
-				FTypedElementDatabaseCommandBuffer::Execute_AddColumnsCommand(*ActiveEditorEntityManager, Row, FragmentsToAdd, TagsToAdd);
+				UE::Editor::DataStorage::Legacy::FCommandBuffer::Execute_AddColumnsCommand(*ActiveEditorEntityManager, Row, FragmentsToAdd, TagsToAdd);
 			}
 			if (bMustRemoveColumns)
 			{
-				FTypedElementDatabaseCommandBuffer::Execute_RemoveColumnsCommand(*ActiveEditorEntityManager, Row, FragmentsToRemove, TagsToRemove);
+				UE::Editor::DataStorage::Legacy::FCommandBuffer::Execute_RemoveColumnsCommand(*ActiveEditorEntityManager, Row, FragmentsToRemove, TagsToRemove);
 			}
 		}
 		else
@@ -715,7 +715,7 @@ void UTypedElementDatabase::BatchAddRemoveColumns(TConstArrayView<TypedElementRo
 			using EntityArchetypeLookup = TMap<FMassArchetypeHandle, EntityHandleArray, TInlineSetAllocator<32>>;
 			using ArchetypeEntityArray = TArray<FMassArchetypeEntityCollection, TInlineAllocator<32>>;
 
-			FTypedElementDatabaseCommandBuffer& CommandBuffer = Environment->GetDirectDeferredCommands();
+			UE::Editor::DataStorage::Legacy::FCommandBuffer& CommandBuffer = Environment->GetDirectDeferredCommands();
 			
 			// Sort rows (entities) into to matching table (archetype) bucket.
 			EntityArchetypeLookup LookupTable;
@@ -793,7 +793,7 @@ bool UTypedElementDatabase::HasColumns(TypedElementRowHandle Row, TConstArrayVie
 		}
 		else
 		{
-			const FTypedElementDatabaseCommandBuffer& CommandBuffer = Environment->GetDirectDeferredCommands();
+			const UE::Editor::DataStorage::Legacy::FCommandBuffer& CommandBuffer = Environment->GetDirectDeferredCommands();
 			const UScriptStruct* const* ColumnTypesEnd = ColumnTypes.end();
 			for (const UScriptStruct* const* ColumnType = ColumnTypes.begin(); ColumnType != ColumnTypesEnd && bHasAllColumns; ++ColumnType)
 			{
@@ -839,7 +839,7 @@ bool UTypedElementDatabase::HasColumns(TypedElementRowHandle Row, TConstArrayVie
 		}
 		else
 		{
-			const FTypedElementDatabaseCommandBuffer& CommandBuffer = Environment->GetDirectDeferredCommands();
+			const UE::Editor::DataStorage::Legacy::FCommandBuffer& CommandBuffer = Environment->GetDirectDeferredCommands();
 			const TWeakObjectPtr<const UScriptStruct>* ColumnTypesEnd = ColumnTypes.end();
 			for (const TWeakObjectPtr<const UScriptStruct>* ColumnType = ColumnTypes.begin(); ColumnType != ColumnTypesEnd && bHasAllColumns; ++ColumnType)
 			{
@@ -976,14 +976,14 @@ void UTypedElementDatabase::UnregisterQuery(TypedElementQueryHandle Query)
 {
 	if (ActiveEditorEntityManager && ActiveEditorPhaseManager)
 	{
-		const FTypedElementExtendedQueryStore::Handle StorageHandle(Query);
+		const UE::Editor::DataStorage::FExtendedQueryStore::Handle StorageHandle(Query);
 		Environment->GetQueryStore().UnregisterQuery(StorageHandle, *ActiveEditorEntityManager, *ActiveEditorPhaseManager);
 	}
 }
 
 const ITypedElementDataStorageInterface::FQueryDescription& UTypedElementDatabase::GetQueryDescription(TypedElementQueryHandle Query) const
 {
-	const FTypedElementExtendedQueryStore::Handle StorageHandle(Query);
+	const UE::Editor::DataStorage::FExtendedQueryStore::Handle StorageHandle(Query);
 	return Environment->GetQueryStore().GetQueryDescription(StorageHandle);
 }
 
@@ -1017,7 +1017,7 @@ ITypedElementDataStorageInterface::FQueryResult UTypedElementDatabase::RunQuery(
 
 	if (ActiveEditorEntityManager)
 	{
-		const FTypedElementExtendedQueryStore::Handle StorageHandle(Query);
+		const UE::Editor::DataStorage::FExtendedQueryStore::Handle StorageHandle(Query);
 		return Environment->GetQueryStore().RunQuery(*ActiveEditorEntityManager, StorageHandle);
 	}
 	else
@@ -1033,7 +1033,7 @@ ITypedElementDataStorageInterface::FQueryResult UTypedElementDatabase::RunQuery(
 
 	if (ActiveEditorEntityManager)
 	{
-		const FTypedElementExtendedQueryStore::Handle StorageHandle(Query);
+		const UE::Editor::DataStorage::FExtendedQueryStore::Handle StorageHandle(Query);
 		return Environment->GetQueryStore().RunQuery(*ActiveEditorEntityManager, *Environment, StorageHandle, Callback);
 	}
 	else
@@ -1052,29 +1052,29 @@ void UTypedElementDatabase::ActivateQueries(FName ActivationName)
 
 TypedElementDataStorage::RowHandle UTypedElementDatabase::FindIndexedRow(TypedElementDataStorage::IndexHash Index) const
 {
-	return Environment->GetIndexTable().FindIndexedRow(UE::EditorDataStorage::EGlobalLockScope::Public, Index);
+	return Environment->GetIndexTable().FindIndexedRow(UE::Editor::DataStorage::EGlobalLockScope::Public, Index);
 }
 
 void UTypedElementDatabase::IndexRow(TypedElementDataStorage::IndexHash Index, TypedElementDataStorage::RowHandle Row)
 {
-	Environment->GetIndexTable().IndexRow(UE::EditorDataStorage::EGlobalLockScope::Public, Index, Row);
+	Environment->GetIndexTable().IndexRow(UE::Editor::DataStorage::EGlobalLockScope::Public, Index, Row);
 }
 
 void UTypedElementDatabase::BatchIndexRows(
 	TConstArrayView<TPair<TypedElementDataStorage::IndexHash, TypedElementDataStorage::RowHandle>> IndexRowPairs)
 {
-	Environment->GetIndexTable().BatchIndexRows(UE::EditorDataStorage::EGlobalLockScope::Public, IndexRowPairs);
+	Environment->GetIndexTable().BatchIndexRows(UE::Editor::DataStorage::EGlobalLockScope::Public, IndexRowPairs);
 }
 
 void UTypedElementDatabase::ReindexRow(TypedElementDataStorage::IndexHash OriginalIndex, TypedElementDataStorage::IndexHash NewIndex, 
 	TypedElementDataStorage::RowHandle RowHandle)
 {
-	Environment->GetIndexTable().ReindexRow(UE::EditorDataStorage::EGlobalLockScope::Public, OriginalIndex, NewIndex, RowHandle);
+	Environment->GetIndexTable().ReindexRow(UE::Editor::DataStorage::EGlobalLockScope::Public, OriginalIndex, NewIndex, RowHandle);
 }
 
 void UTypedElementDatabase::RemoveIndex(TypedElementDataStorage::IndexHash Index)
 {
-	Environment->GetIndexTable().RemoveIndex(UE::EditorDataStorage::EGlobalLockScope::Public, Index);
+	Environment->GetIndexTable().RemoveIndex(UE::Editor::DataStorage::EGlobalLockScope::Public, Index);
 }
 
 FTypedElementOnDataStorageUpdate& UTypedElementDatabase::OnUpdate()
@@ -1112,7 +1112,7 @@ void UTypedElementDatabase::ListExtensions(TFunctionRef<void(FName)> Callback) c
 
 void UTypedElementDatabase::PreparePhase(EQueryTickPhase Phase, float DeltaTime)
 {
-	using namespace UE::EditorDataStorage;
+	using namespace UE::Editor::DataStorage;
 
 	if (ActiveEditorEntityManager)
 	{
@@ -1132,7 +1132,7 @@ void UTypedElementDatabase::PreparePhase(EQueryTickPhase Phase, float DeltaTime)
 
 void UTypedElementDatabase::FinalizePhase(EQueryTickPhase Phase, float DeltaTime)
 {
-	using namespace UE::EditorDataStorage;
+	using namespace UE::Editor::DataStorage;
 
 	if (ActiveEditorEntityManager)
 	{
@@ -1184,12 +1184,12 @@ int32 UTypedElementDatabase::GetTableChunkSize(FName TableName) const
 	}
 }
 
-TSharedPtr<FTypedElementDatabaseEnvironment> UTypedElementDatabase::GetEnvironment()
+TSharedPtr<UE::Editor::DataStorage::FEnvironment> UTypedElementDatabase::GetEnvironment()
 {
 	return Environment;
 }
 
-TSharedPtr<const FTypedElementDatabaseEnvironment> UTypedElementDatabase::GetEnvironment() const
+TSharedPtr<const UE::Editor::DataStorage::FEnvironment> UTypedElementDatabase::GetEnvironment() const
 {
 	return Environment;
 }
diff --git a/Engine/Plugins/Experimental/EditorDataStorage/Source/TedsCore/Private/TypedElementDatabase.h b/Engine/Plugins/Experimental/EditorDataStorage/Source/TedsCore/Private/TypedElementDatabase.h
--- a/Engine/Plugins/Experimental/EditorDataStorage/Source/TedsCore/Private/TypedElementDatabase.h
+++ b/Engine/Plugins/Experimental/EditorDataStorage/Source/TedsCore/Private/TypedElementDatabase.h
@@ -15,11 +15,14 @@
 
 struct FMassEntityManager;
 struct FMassProcessingPhaseManager;
-class FTypedElementDatabaseEnvironment;
 class UTypedElementDataStorageFactory;
 class FOutputDevice;
 class UWorld;
 
+namespace UE::Editor::DataStorage
+{
+	class FEnvironment;
+} // namespace UE::Editor::DataStorage
 
 UCLASS()
 class TEDSCORE_API UTypedElementDatabase
@@ -66,7 +69,7 @@ public:
 	FactoryConstIterator CreateFactoryIterator() const;
 
 	/** Returns factory instance given the type of factory */
-	const UTypedElementDataStorageFactory* FindFactory(const UClass* FactoryType) const override;
+	virtual const UTypedElementDataStorageFactory* FindFactory(const UClass* FactoryType) const override;
 	/** Helper for FindFactory(const UClass*) */
 	template<typename FactoryTypeT>
 	const FactoryTypeT* FindFactory() const;
@@ -81,80 +84,80 @@ public:
 	TSharedPtr<FMassEntityManager> GetActiveMutableEditorEntityManager();
 	TSharedPtr<const FMassEntityManager> GetActiveEditorEntityManager() const;
 
-	TypedElementDataStorage::TableHandle RegisterTable(TConstArrayView<const UScriptStruct*> ColumnList, const FName Name) override;
-	TypedElementDataStorage::TableHandle RegisterTable(
+	virtual TypedElementDataStorage::TableHandle RegisterTable(TConstArrayView<const UScriptStruct*> ColumnList, const FName Name) override;
+	virtual TypedElementDataStorage::TableHandle RegisterTable(
 		TypedElementDataStorage::TableHandle SourceTable, TConstArrayView<const UScriptStruct*> ColumnList, const FName Name) override;
-	TypedElementDataStorage::TableHandle FindTable(const FName Name) override;
+	virtual TypedElementDataStorage::TableHandle FindTable(const FName Name) override;
 
-	TypedElementDataStorage::RowHandle ReserveRow() override;
-	void BatchReserveRows(int32 Count, TFunctionRef<void(TypedElementDataStorage::RowHandle)> ReservationCallback) override;
-	void BatchReserveRows(TArrayView<TypedElementDataStorage::RowHandle> ReservedRows) override;
-	TypedElementDataStorage::RowHandle AddRow(TypedElementDataStorage::TableHandle Table, 
+	virtual TypedElementDataStorage::RowHandle ReserveRow() override;
+	virtual void BatchReserveRows(int32 Count, TFunctionRef<void(TypedElementDataStorage::RowHandle)> ReservationCallback) override;
+	virtual void BatchReserveRows(TArrayView<TypedElementDataStorage::RowHandle> ReservedRows) override;
+	virtual TypedElementDataStorage::RowHandle AddRow(TypedElementDataStorage::TableHandle Table, 
 		TypedElementDataStorage::RowCreationCallbackRef OnCreated) override;
 	TypedElementDataStorage::RowHandle AddRow(TypedElementDataStorage::TableHandle Table) override;
-	bool AddRow(TypedElementDataStorage::RowHandle ReservedRow, TypedElementDataStorage::TableHandle Table) override;
-	bool AddRow(TypedElementDataStorage::RowHandle ReservedRow, TypedElementDataStorage::TableHandle Table,
+	virtual bool AddRow(TypedElementDataStorage::RowHandle ReservedRow, TypedElementDataStorage::TableHandle Table) override;
+	virtual bool AddRow(TypedElementDataStorage::RowHandle ReservedRow, TypedElementDataStorage::TableHandle Table,
 		TypedElementDataStorage::RowCreationCallbackRef OnCreated) override;
-	bool BatchAddRow(TypedElementDataStorage::TableHandle Table, int32 Count, 
+	virtual bool BatchAddRow(TypedElementDataStorage::TableHandle Table, int32 Count,
 		TypedElementDataStorage::RowCreationCallbackRef OnCreated) override;
-	bool BatchAddRow(TypedElementDataStorage::TableHandle Table, TConstArrayView<TypedElementDataStorage::RowHandle> ReservedHandles,
+	virtual bool BatchAddRow(TypedElementDataStorage::TableHandle Table, TConstArrayView<TypedElementDataStorage::RowHandle> ReservedHandles,
 		TypedElementDataStorage::RowCreationCallbackRef OnCreated) override;
-	void RemoveRow(TypedElementDataStorage::RowHandle Row) override;
-	bool IsRowAvailable(TypedElementDataStorage::RowHandle Row) const override;
-	bool IsRowAssigned(TypedElementDataStorage::RowHandle Row) const override;
+	virtual void RemoveRow(TypedElementDataStorage::RowHandle Row) override;
+	virtual bool IsRowAvailable(TypedElementDataStorage::RowHandle Row) const override;
+	virtual bool IsRowAssigned(TypedElementDataStorage::RowHandle Row) const override;
 
-	void AddColumn(TypedElementRowHandle Row, const UScriptStruct* ColumnType) override;
-	void AddColumn(TypedElementRowHandle Row, const UE::EditorDataStorage::FDynamicTag& Tag, const FName& InValue) override;
-	void AddColumnData(TypedElementRowHandle Row, const UScriptStruct* ColumnType,
+	virtual void AddColumn(TypedElementRowHandle Row, const UScriptStruct* ColumnType) override;
+	virtual void AddColumn(TypedElementRowHandle Row, const UE::Editor::DataStorage::FDynamicTag& Tag, const FName& InValue) override;
+	virtual void AddColumnData(TypedElementRowHandle Row, const UScriptStruct* ColumnType,
 		const TypedElementDataStorage::ColumnCreationCallbackRef& Initializer,
 		TypedElementDataStorage::ColumnCopyOrMoveCallback Relocator) override;
-	void RemoveColumn(TypedElementRowHandle Row, const UScriptStruct* ColumnType) override;
-	void RemoveColumn(TypedElementRowHandle Row, const UE::EditorDataStorage::FDynamicTag& Tag) override;
-	void* GetColumnData(TypedElementRowHandle Row, const UScriptStruct* ColumnType) override;
-	const void* GetColumnData(TypedElementRowHandle Row, const UScriptStruct* ColumnType) const override;
-	void AddColumns(TypedElementRowHandle Row, TConstArrayView<const UScriptStruct*> Columns) override;
-	void RemoveColumns(TypedElementRowHandle Row, TConstArrayView<const UScriptStruct*> Columns) override;
-	void AddRemoveColumns(TypedElementRowHandle Row, TConstArrayView<const UScriptStruct*> ColumnsToAdd,
+	virtual void RemoveColumn(TypedElementRowHandle Row, const UScriptStruct* ColumnType) override;
+	virtual void RemoveColumn(TypedElementRowHandle Row, const UE::Editor::DataStorage::FDynamicTag& Tag) override;
+	virtual void* GetColumnData(TypedElementRowHandle Row, const UScriptStruct* ColumnType) override;
+	virtual const void* GetColumnData(TypedElementRowHandle Row, const UScriptStruct* ColumnType) const override;
+	virtual void AddColumns(TypedElementRowHandle Row, TConstArrayView<const UScriptStruct*> Columns) override;
+	virtual void RemoveColumns(TypedElementRowHandle Row, TConstArrayView<const UScriptStruct*> Columns) override;
+	virtual void AddRemoveColumns(TypedElementRowHandle Row, TConstArrayView<const UScriptStruct*> ColumnsToAdd,
 		TConstArrayView<const UScriptStruct*> ColumnsToRemove) override;
-	void BatchAddRemoveColumns(TConstArrayView<TypedElementRowHandle> Rows,TConstArrayView<const UScriptStruct*> ColumnsToAdd,
+	virtual void BatchAddRemoveColumns(TConstArrayView<TypedElementRowHandle> Rows,TConstArrayView<const UScriptStruct*> ColumnsToAdd,
 		TConstArrayView<const UScriptStruct*> ColumnsToRemove) override;
-	bool HasColumns(TypedElementRowHandle Row, TConstArrayView<const UScriptStruct*> ColumnTypes) const override;
-	bool HasColumns(TypedElementRowHandle Row, TConstArrayView<TWeakObjectPtr<const UScriptStruct>> ColumnTypes) const override;
-	void ListColumns(TypedElementDataStorage::RowHandle Row, TypedElementDataStorage::ColumnListCallbackRef Callback) const;
-	void ListColumns(TypedElementDataStorage::RowHandle Row, TypedElementDataStorage::ColumnListWithDataCallbackRef Callback);
-	bool MatchesColumns(TypedElementDataStorage::RowHandle Row, const TypedElementDataStorage::FQueryConditions& Conditions) const override;
+	virtual bool HasColumns(TypedElementRowHandle Row, TConstArrayView<const UScriptStruct*> ColumnTypes) const override;
+	virtual bool HasColumns(TypedElementRowHandle Row, TConstArrayView<TWeakObjectPtr<const UScriptStruct>> ColumnTypes) const override;
+	virtual void ListColumns(TypedElementDataStorage::RowHandle Row, TypedElementDataStorage::ColumnListCallbackRef Callback) const;
+	virtual void ListColumns(TypedElementDataStorage::RowHandle Row, TypedElementDataStorage::ColumnListWithDataCallbackRef Callback);
+	virtual bool MatchesColumns(TypedElementDataStorage::RowHandle Row, const TypedElementDataStorage::FQueryConditions& Conditions) const override;
 
 	void RegisterTickGroup(FName GroupName, EQueryTickPhase Phase, FName BeforeGroup, FName AfterGroup, bool bRequiresMainThread);
 	void UnregisterTickGroup(FName GroupName, EQueryTickPhase Phase);
 
 	TypedElementQueryHandle RegisterQuery(FQueryDescription&& Query) override;
-	void UnregisterQuery(TypedElementQueryHandle Query) override;
-	const FQueryDescription& GetQueryDescription(TypedElementQueryHandle Query) const override;
-	FName GetQueryTickGroupName(EQueryTickGroups Group) const override;
-	FQueryResult RunQuery(TypedElementQueryHandle Query) override;
-	FQueryResult RunQuery(TypedElementQueryHandle Query, DirectQueryCallbackRef Callback) override;
-	void ActivateQueries(FName ActivationName) override;
-
-	TypedElementDataStorage::RowHandle FindIndexedRow(TypedElementDataStorage::IndexHash Index) const override;
-	void IndexRow(TypedElementDataStorage::IndexHash Index, TypedElementDataStorage::RowHandle Row) override;
-	void BatchIndexRows(
+	virtual void UnregisterQuery(TypedElementQueryHandle Query) override;
+	virtual const FQueryDescription& GetQueryDescription(TypedElementQueryHandle Query) const override;
+	virtual FName GetQueryTickGroupName(EQueryTickGroups Group) const override;
+	virtual FQueryResult RunQuery(TypedElementQueryHandle Query) override;
+	virtual FQueryResult RunQuery(TypedElementQueryHandle Query, DirectQueryCallbackRef Callback) override;
+	virtual void ActivateQueries(FName ActivationName) override;
+
+	virtual TypedElementDataStorage::RowHandle FindIndexedRow(TypedElementDataStorage::IndexHash Index) const override;
+	virtual void IndexRow(TypedElementDataStorage::IndexHash Index, TypedElementDataStorage::RowHandle Row) override;
+	virtual void BatchIndexRows(
 		TConstArrayView<TPair<TypedElementDataStorage::IndexHash, TypedElementDataStorage::RowHandle>> IndexRowPairs) override;
-	void ReindexRow(
+	virtual void ReindexRow(
 		TypedElementDataStorage::IndexHash OriginalIndex, 
 		TypedElementDataStorage::IndexHash NewIndex, 
 		TypedElementDataStorage::RowHandle Row) override;
-	void RemoveIndex(TypedElementDataStorage::IndexHash Index) override;
+	virtual void RemoveIndex(TypedElementDataStorage::IndexHash Index) override;
 
-	FTypedElementOnDataStorageUpdate& OnUpdate() override;
-	FTypedElementOnDataStorageUpdate& OnUpdateCompleted() override;
-	bool IsAvailable() const override;
-	void* GetExternalSystemAddress(UClass* Target) override;
+	virtual FTypedElementOnDataStorageUpdate& OnUpdate() override;
+	virtual FTypedElementOnDataStorageUpdate& OnUpdateCompleted() override;
+	virtual bool IsAvailable() const override;
+	virtual void* GetExternalSystemAddress(UClass* Target) override;
 
-	bool SupportsExtension(FName Extension) const override;
-	void ListExtensions(TFunctionRef<void(FName)> Callback) const override;
+	virtual bool SupportsExtension(FName Extension) const override;
+	virtual void ListExtensions(TFunctionRef<void(FName)> Callback) const override;
 	
-	TSharedPtr<FTypedElementDatabaseEnvironment> GetEnvironment();
-	TSharedPtr<const FTypedElementDatabaseEnvironment> GetEnvironment() const;
+	TSharedPtr<UE::Editor::DataStorage::FEnvironment> GetEnvironment();
+	TSharedPtr<const UE::Editor::DataStorage::FEnvironment> GetEnvironment() const;
 
 	FMassArchetypeHandle LookupArchetype(TypedElementDataStorage::TableHandle InTableHandle) const;
 
@@ -191,7 +194,7 @@ private:
 	// Ordered array of factories by the return value of GetOrder()
 	TArray<FFactoryTypePair> Factories;
 
-	TSharedPtr<FTypedElementDatabaseEnvironment> Environment;
+	TSharedPtr<UE::Editor::DataStorage::FEnvironment> Environment;
 	
 	FTypedElementOnDataStorageUpdate OnUpdateDelegate;
 	FTypedElementOnDataStorageUpdate OnUpdateCompletedDelegate;
diff --git a/Engine/Plugins/Experimental/EditorDataStorage/Source/TedsCore/Private/TypedElementDatabaseCommandBuffer.cpp b/Engine/Plugins/Experimental/EditorDataStorage/Source/TedsCore/Private/TypedElementDatabaseCommandBuffer.cpp
--- a/Engine/Plugins/Experimental/EditorDataStorage/Source/TedsCore/Private/TypedElementDatabaseCommandBuffer.cpp
+++ b/Engine/Plugins/Experimental/EditorDataStorage/Source/TedsCore/Private/TypedElementDatabaseCommandBuffer.cpp
@@ -7,340 +7,343 @@
 #include "TypedElementDatabaseEnvironment.h"
 #include "TypedElementDataStorageSharedColumn.h"
 
-// 
-// Commands section
-//
-
-FTypedElementDatabaseCommandBuffer::FTypedElementDatabaseCommandBuffer(FTypedElementDatabaseEnvironment& InEnvironment)
-	: Environment(InEnvironment)
+namespace UE::Editor::DataStorage::Legacy
 {
-}
+	// 
+	// Commands section
+	//
 
-void* FTypedElementDatabaseCommandBuffer::GetQueuedDataColumn(
-	TypedElementDataStorage::RowHandle Row, const UScriptStruct* ColumnType)
-{
-	checkf(ColumnType->IsChildOf(FMassFragment::StaticStruct()),
-		TEXT("Trying to get the column '%s' which isn't a data column."), *ColumnType->GetName());
+	FCommandBuffer::FCommandBuffer(FEnvironment& InEnvironment)
+		: Environment(InEnvironment)
+	{
+	}
 
-	void* const* StoredData = PendingColumns.Find(PendingColumnMappingKey(Row, ColumnType));
-	if (StoredData)
+	void* FCommandBuffer::GetQueuedDataColumn(
+		TypedElementDataStorage::RowHandle Row, const UScriptStruct* ColumnType)
 	{
-		if (*StoredData != nullptr)
+		checkf(ColumnType->IsChildOf(FMassFragment::StaticStruct()),
+			TEXT("Trying to get the column '%s' which isn't a data column."), *ColumnType->GetName());
+
+		void* const* StoredData = PendingColumns.Find(PendingColumnMappingKey(Row, ColumnType));
+		if (StoredData)
 		{
-			return *StoredData;
+			if (*StoredData != nullptr)
+			{
+				return *StoredData;
+			}
+			else
+			{
+				// If the column was created but had no data assigned to it, create data for it now. If this code path triggers
+				// a lot there may be a large number of AddColumn followed by GetColumn calls. These can be more efficiently done
+				// with an AddorGetColumn call.
+				void* Result = Queue_AddDataColumnCommandUnitialized(Row, ColumnType,
+					[](const UScriptStruct& ColumnType, void* Destination, void* Source)
+					{
+						ColumnType.CopyScriptStruct(Destination, Source);
+					});
+				return Result;
+			}
 		}
 		else
 		{
-			// If the column was created but had no data assigned to it, create data for it now. If this code path triggers
-			// a lot there may be a large number of AddColumn followed by GetColumn calls. These can be more efficiently done
-			// with an AddorGetColumn call.
-			void* Result = Queue_AddDataColumnCommandUnitialized(Row, ColumnType,
-				[](const UScriptStruct& ColumnType, void* Destination, void* Source)
-				{
-					ColumnType.CopyScriptStruct(Destination, Source);
-				});
-			return Result;
+			return nullptr;
 		}
 	}
-	else
+
+	bool FCommandBuffer::HasColumn(TypedElementDataStorage::RowHandle Row, const UScriptStruct* ColumnType) const
 	{
-		return nullptr;
+		return PendingColumns.Contains(PendingColumnMappingKey(Row, ColumnType));
 	}
-}
-
-bool FTypedElementDatabaseCommandBuffer::HasColumn(TypedElementDataStorage::RowHandle Row, const UScriptStruct* ColumnType) const
-{
-	return PendingColumns.Contains(PendingColumnMappingKey(Row, ColumnType));
-}
 
-void FTypedElementDatabaseCommandBuffer::Clear(TypedElementDataStorage::RowHandle Row)
-{
-	for (TMap<PendingColumnMappingKey, void*>::TIterator It = PendingColumns.CreateIterator(); It; ++It)
+	void FCommandBuffer::Clear(TypedElementDataStorage::RowHandle Row)
 	{
-		if (It.Key().Key == Row)
+		for (TMap<PendingColumnMappingKey, void*>::TIterator It = PendingColumns.CreateIterator(); It; ++It)
 		{
-			It.RemoveCurrent();
+			if (It.Key().Key == Row)
+			{
+				It.RemoveCurrent();
+			}
+		}
+
+		for (TArray<FCommand>::TIterator It = Commands.CreateIterator(); It; ++It)
+		{
+			if (It->Row == Row)
+			{
+				It.RemoveCurrent(); // Don't swap as the order of commands is important.
+			}
 		}
 	}
 
-	for (TArray<FCommand>::TIterator It = Commands.CreateIterator(); It; ++It)
+	FCommandBuffer::FAddDataColumnCommand::~FAddDataColumnCommand()
 	{
-		if (It->Row == Row)
+		if (ColumnType.IsValid() && (ColumnType->StructFlags & (STRUCT_IsPlainOldData | STRUCT_NoDestructor)) == 0)
 		{
-			It.RemoveCurrent(); // Don't swap as the order of commands is important.
+			ColumnType->DestroyStruct(Data);
 		}
 	}
-}
 
-FTypedElementDatabaseCommandBuffer::FAddDataColumnCommand::~FAddDataColumnCommand()
-{
-	if (ColumnType.IsValid() && (ColumnType->StructFlags & (STRUCT_IsPlainOldData | STRUCT_NoDestructor)) == 0)
+	//
+	// Queue section
+	//
+
+	void FCommandBuffer::Queue_AddColumnCommand(TypedElementDataStorage::RowHandle Row, const UScriptStruct* ColumnType)
 	{
-		ColumnType->DestroyStruct(Data);
+		AddCommand(Row, FAddColumnCommand
+			{
+				.ColumnType = ColumnType
+			});
+		PendingColumns.FindOrAdd(PendingColumnMappingKey(Row, ColumnType), nullptr);
 	}
-}
 
-//
-// Queue section
-//
+	void* FCommandBuffer::Queue_AddDataColumnCommandUnitialized(
+		TypedElementDataStorage::RowHandle Row, const UScriptStruct* ColumnType, TypedElementDataStorage::ColumnCopyOrMoveCallback Relocator)
+	{
+		checkf(ColumnType->IsChildOf(FMassFragment::StaticStruct()),
+			TEXT("Trying to queue a data column creation for '%s' which isn't a data column."), *ColumnType->GetName());
 
-void FTypedElementDatabaseCommandBuffer::Queue_AddColumnCommand(TypedElementDataStorage::RowHandle Row, const UScriptStruct* ColumnType)
-{
-	AddCommand(Row, FAddColumnCommand
-		{ 
-			.ColumnType = ColumnType 
-		});
-	PendingColumns.FindOrAdd(PendingColumnMappingKey(Row, ColumnType), nullptr);
-}
-
-void* FTypedElementDatabaseCommandBuffer::Queue_AddDataColumnCommandUnitialized(
-	TypedElementDataStorage::RowHandle Row, const UScriptStruct* ColumnType, TypedElementDataStorage::ColumnCopyOrMoveCallback Relocator)
-{
-	checkf(ColumnType->IsChildOf(FMassFragment::StaticStruct()),
-		TEXT("Trying to queue a data column creation for '%s' which isn't a data column."), *ColumnType->GetName());
+		PendingColumnMappingKey Key(Row, ColumnType);
+		if (void** StoredData = PendingColumns.Find(Key); StoredData == nullptr || *StoredData == nullptr)
+		{
+			void* Data = Environment.GetScratchBuffer().Allocate(ColumnType->GetStructureSize(), ColumnType->GetMinAlignment());
+			// Initialize to zero to replicate the default from Mass.
+			FMemory::Memzero(Data, ColumnType->GetStructureSize());
+			PendingColumns.Add(PendingColumnMappingKey(Row, ColumnType), Data);
+			AddCommand(Row, FAddDataColumnCommand
+				{
+					.ColumnType = ColumnType,
+					.Relocator = Relocator,
+					.Data = Data
+				});
+			return Data;
+		}
+		else
+		{
+			return *StoredData;
+		}
+	}
 
-	PendingColumnMappingKey Key(Row, ColumnType);
-	if (void** StoredData = PendingColumns.Find(Key); StoredData == nullptr || *StoredData == nullptr)
+	void FCommandBuffer::Queue_AddColumnsCommand(TypedElementDataStorage::RowHandle Row,
+		FMassFragmentBitSet FragmentsToAdd, FMassTagBitSet TagsToAdd)
 	{
-		void* Data = Environment.GetScratchBuffer().Allocate(ColumnType->GetStructureSize(), ColumnType->GetMinAlignment());
-		// Initialize to zero to replicate the default from Mass.
-		FMemory::Memzero(Data, ColumnType->GetStructureSize());
-		PendingColumns.Add(PendingColumnMappingKey(Row, ColumnType), Data);
-		AddCommand(Row, FAddDataColumnCommand
+		auto AddColumn = [this, Row](const UScriptStruct* ColumnType)
+		{
+			PendingColumns.FindOrAdd(PendingColumnMappingKey(Row, ColumnType), nullptr);
+			return true;
+		};
+		FragmentsToAdd.ExportTypes(AddColumn);
+		TagsToAdd.ExportTypes(AddColumn);
+
+		AddCommand(Row, FAddColumnsCommand
 			{
-				.ColumnType = ColumnType,
-				.Relocator = Relocator,
-				.Data = Data
+				.FragmentsToAdd = MoveTemp(FragmentsToAdd),
+				.TagsToAdd = MoveTemp(TagsToAdd)
 			});
-		return Data;
 	}
-	else
+
+	void FCommandBuffer::Queue_RemoveColumnCommand(TypedElementDataStorage::RowHandle Row, const UScriptStruct* ColumnType)
 	{
-		return *StoredData;
+		AddCommand(Row, FRemoveColumnCommand
+			{
+				.ColumnType = ColumnType
+			});
+		PendingColumns.Remove(PendingColumnMappingKey(Row, ColumnType));
 	}
-}
 
-void FTypedElementDatabaseCommandBuffer::Queue_AddColumnsCommand(TypedElementDataStorage::RowHandle Row, 
-	FMassFragmentBitSet FragmentsToAdd, FMassTagBitSet TagsToAdd)
-{
-	auto AddColumn = [this, Row](const UScriptStruct* ColumnType)
-	{
-		PendingColumns.FindOrAdd(PendingColumnMappingKey(Row, ColumnType), nullptr);
-		return true;
-	};
-	FragmentsToAdd.ExportTypes(AddColumn);
-	TagsToAdd.ExportTypes(AddColumn);
-
-	AddCommand(Row, FAddColumnsCommand
-		{ 
-			.FragmentsToAdd = MoveTemp(FragmentsToAdd), 
-			.TagsToAdd = MoveTemp(TagsToAdd) 
-		});
-}
-
-void FTypedElementDatabaseCommandBuffer::Queue_RemoveColumnCommand(TypedElementDataStorage::RowHandle Row, const UScriptStruct* ColumnType)
-{
-	AddCommand(Row, FRemoveColumnCommand
-		{ 
-			.ColumnType = ColumnType 
-		});
-	PendingColumns.Remove(PendingColumnMappingKey(Row, ColumnType));
-}
-
-void FTypedElementDatabaseCommandBuffer::Queue_RemoveColumnsCommand(TypedElementDataStorage::RowHandle Row,
-	FMassFragmentBitSet FragmentsToRemove, FMassTagBitSet TagsToRemove)
-{
-	auto RemoveColumn = [this, Row](const UScriptStruct* ColumnType)
+	void FCommandBuffer::Queue_RemoveColumnsCommand(TypedElementDataStorage::RowHandle Row,
+		FMassFragmentBitSet FragmentsToRemove, FMassTagBitSet TagsToRemove)
 	{
-		PendingColumns.Remove(PendingColumnMappingKey(Row, ColumnType));
-		return true;
-	};
-	FragmentsToRemove.ExportTypes(RemoveColumn);
-	TagsToRemove.ExportTypes(RemoveColumn);
-	
-	AddCommand(Row, FRemoveColumnsCommand
-		{ 
-			.FragmentsToRemove = MoveTemp(FragmentsToRemove), 
-			.TagsToRemove = MoveTemp(TagsToRemove) 
-		});
-}
-
-//
-// Execute section
-//
-
-bool FTypedElementDatabaseCommandBuffer::Execute_IsRowAvailable(const FMassEntityManager& MassEntityManager, TypedElementDataStorage::RowHandle Row)
-{
-	if (Row == TypedElementInvalidRowHandle)
+		auto RemoveColumn = [this, Row](const UScriptStruct* ColumnType)
+		{
+			PendingColumns.Remove(PendingColumnMappingKey(Row, ColumnType));
+			return true;
+		};
+		FragmentsToRemove.ExportTypes(RemoveColumn);
+		TagsToRemove.ExportTypes(RemoveColumn);
+
+		AddCommand(Row, FRemoveColumnsCommand
+			{
+				.FragmentsToRemove = MoveTemp(FragmentsToRemove),
+				.TagsToRemove = MoveTemp(TagsToRemove)
+			});
+	}
+
+	//
+	// Execute section
+	//
+
+	bool FCommandBuffer::Execute_IsRowAvailable(const FMassEntityManager& MassEntityManager, TypedElementDataStorage::RowHandle Row)
 	{
-		return false;
+		if (Row == TypedElementInvalidRowHandle)
+		{
+			return false;
+		}
+		return MassEntityManager.IsEntityValid(FMassEntityHandle::FromNumber(Row));
 	}
-	return MassEntityManager.IsEntityValid(FMassEntityHandle::FromNumber(Row));
-}
 
-bool FTypedElementDatabaseCommandBuffer::Execute_IsRowAssigned(const FMassEntityManager& MassEntityManager, TypedElementDataStorage::RowHandle Row)
-{
-	if (Row == TypedElementInvalidRowHandle)
+	bool FCommandBuffer::Execute_IsRowAssigned(const FMassEntityManager& MassEntityManager, TypedElementDataStorage::RowHandle Row)
 	{
-		return false;
+		if (Row == TypedElementInvalidRowHandle)
+		{
+			return false;
+		}
+		return MassEntityManager.IsEntityActive(FMassEntityHandle::FromNumber(Row));
 	}
-	return MassEntityManager.IsEntityActive(FMassEntityHandle::FromNumber(Row));
-}
 
-void FTypedElementDatabaseCommandBuffer::Execute_AddColumnCommand(
-	FMassEntityManager& MassEntityManager,
-	TypedElementDataStorage::RowHandle Row,
-	const UScriptStruct* ColumnType)
-{
-	if (ColumnType)
+	void FCommandBuffer::Execute_AddColumnCommand(
+		FMassEntityManager& MassEntityManager,
+		TypedElementDataStorage::RowHandle Row,
+		const UScriptStruct* ColumnType)
 	{
-		FMassEntityHandle Entity = FMassEntityHandle::FromNumber(Row);
-		if (ColumnType->IsChildOf(FMassTag::StaticStruct()))
+		if (ColumnType)
 		{
-			MassEntityManager.AddTagToEntity(Entity, ColumnType);
+			FMassEntityHandle Entity = FMassEntityHandle::FromNumber(Row);
+			if (ColumnType->IsChildOf(FMassTag::StaticStruct()))
+			{
+				MassEntityManager.AddTagToEntity(Entity, ColumnType);
+			}
+			else if (ColumnType->IsChildOf(FMassFragment::StaticStruct()))
+			{
+				FStructView Column = MassEntityManager.GetFragmentDataStruct(Entity, ColumnType);
+				// Only add if not already added to avoid asserts from Mass.
+				if (!Column.IsValid())
+				{
+					MassEntityManager.AddFragmentToEntity(Entity, ColumnType);
+				}
+			}
 		}
-		else if (ColumnType->IsChildOf(FMassFragment::StaticStruct()))
+	}
+
+	void FCommandBuffer::Execute_AddDataColumnCommand(
+		FMassEntityManager& MassEntityManager, TypedElementDataStorage::RowHandle Row, const UScriptStruct* ColumnType,
+		void* Data, TypedElementDataStorage::ColumnCopyOrMoveCallback Relocator)
+	{
+		if (ColumnType)
 		{
+			checkf(ColumnType->IsChildOf(FMassFragment::StaticStruct()),
+				TEXT("Trying to create a data column for '%s' from a deferred command that isn't a data column."), *ColumnType->GetName());
+
+			FMassEntityHandle Entity = FMassEntityHandle::FromNumber(Row);
 			FStructView Column = MassEntityManager.GetFragmentDataStruct(Entity, ColumnType);
 			// Only add if not already added to avoid asserts from Mass.
 			if (!Column.IsValid())
 			{
-				MassEntityManager.AddFragmentToEntity(Entity, ColumnType);
+				MassEntityManager.AddFragmentToEntity(Entity, ColumnType,
+					[Data, Relocator](void* Fragment, const UScriptStruct& FragmentType)
+					{
+						Relocator(FragmentType, Fragment, Data);
+					});
+			}
+			else
+			{
+				Relocator(*ColumnType, Column.GetMemory(), Data);
 			}
 		}
 	}
-}
 
-void FTypedElementDatabaseCommandBuffer::Execute_AddDataColumnCommand(
-	FMassEntityManager& MassEntityManager, TypedElementDataStorage::RowHandle Row, const UScriptStruct* ColumnType,
-	void* Data, TypedElementDataStorage::ColumnCopyOrMoveCallback Relocator)
-{
-	if (ColumnType)
+	void FCommandBuffer::Execute_AddSharedColumnCommand(FMassEntityManager& MassEntityManager, TypedElementDataStorage::RowHandle Row, const FConstSharedStruct& SharedColumn)
 	{
-		checkf(ColumnType->IsChildOf(FMassFragment::StaticStruct()),
-			TEXT("Trying to create a data column for '%s' from a deferred command that isn't a data column."), *ColumnType->GetName());
-
-		FMassEntityHandle Entity = FMassEntityHandle::FromNumber(Row); 
-		FStructView Column = MassEntityManager.GetFragmentDataStruct(Entity, ColumnType);
-		// Only add if not already added to avoid asserts from Mass.
-		if (!Column.IsValid())
+		if (SharedColumn.IsValid())
 		{
-			MassEntityManager.AddFragmentToEntity(Entity, ColumnType,
-				[Data, Relocator](void* Fragment, const UScriptStruct& FragmentType)
-				{
-					Relocator(FragmentType, Fragment, Data);
-				});
-		}
-		else
-		{
-			Relocator(*ColumnType, Column.GetMemory(), Data);
+			FMassEntityHandle Entity = FMassEntityHandle::FromNumber(Row);
+			MassEntityManager.AddConstSharedFragmentToEntity(Entity, SharedColumn);
 		}
 	}
-}
 
-void FTypedElementDatabaseCommandBuffer::Execute_AddSharedColumnCommand(FMassEntityManager& MassEntityManager, TypedElementDataStorage::RowHandle Row, const FConstSharedStruct& SharedColumn)
-{
-	if (SharedColumn.IsValid())
+	void FCommandBuffer::Execute_RemoveSharedColumnCommand(FMassEntityManager& MassEntityManager, TypedElementDataStorage::RowHandle Row, const UScriptStruct& ColumnType)
 	{
-		FMassEntityHandle Entity = FMassEntityHandle::FromNumber(Row);
-		MassEntityManager.AddConstSharedFragmentToEntity(Entity, SharedColumn);
+		if (ColumnType.IsChildOf(FTedsSharedColumn::StaticStruct()))
+		{
+			FMassEntityHandle Entity = FMassEntityHandle::FromNumber(Row);
+			MassEntityManager.RemoveConstSharedFragmentFromEntity(Entity, ColumnType);
+		}
 	}
-}
 
-void FTypedElementDatabaseCommandBuffer::Execute_RemoveSharedColumnCommand(FMassEntityManager& MassEntityManager, TypedElementDataStorage::RowHandle Row, const UScriptStruct& ColumnType)
-{
-	if (ColumnType.IsChildOf(FTedsSharedColumn::StaticStruct()))
+	void FCommandBuffer::Execute_AddColumnsCommand(
+		FMassEntityManager& MassEntityManager,
+		TypedElementDataStorage::RowHandle Row,
+		FMassFragmentBitSet FragmentsToAdd,
+		FMassTagBitSet TagsToAdd)
 	{
-		FMassEntityHandle Entity = FMassEntityHandle::FromNumber(Row);
-		MassEntityManager.RemoveConstSharedFragmentFromEntity(Entity, ColumnType);
+		FMassArchetypeCompositionDescriptor AddComposition(
+			MoveTemp(FragmentsToAdd), MoveTemp(TagsToAdd), FMassChunkFragmentBitSet(), FMassSharedFragmentBitSet(), FMassConstSharedFragmentBitSet());
+		MassEntityManager.AddCompositionToEntity_GetDelta(FMassEntityHandle::FromNumber(Row), AddComposition);
 	}
-}
 
-void FTypedElementDatabaseCommandBuffer::Execute_AddColumnsCommand(
-	FMassEntityManager& MassEntityManager,
-	TypedElementDataStorage::RowHandle Row,
-	FMassFragmentBitSet FragmentsToAdd,
-	FMassTagBitSet TagsToAdd)
-{
-	FMassArchetypeCompositionDescriptor AddComposition(
-		MoveTemp(FragmentsToAdd), MoveTemp(TagsToAdd), FMassChunkFragmentBitSet(), FMassSharedFragmentBitSet(), FMassConstSharedFragmentBitSet());
-	MassEntityManager.AddCompositionToEntity_GetDelta(FMassEntityHandle::FromNumber(Row), AddComposition);
-}
-
-void FTypedElementDatabaseCommandBuffer::Execute_RemoveColumnCommand(
-	FMassEntityManager& MassEntityManager,
-	TypedElementDataStorage::RowHandle Row,
-	const UScriptStruct* ColumnType)
-{
-	if (ColumnType)
+	void FCommandBuffer::Execute_RemoveColumnCommand(
+		FMassEntityManager& MassEntityManager,
+		TypedElementDataStorage::RowHandle Row,
+		const UScriptStruct* ColumnType)
 	{
-		FMassEntityHandle Entity = FMassEntityHandle::FromNumber(Row);
-		if (ColumnType->IsChildOf(FMassTag::StaticStruct()))
-		{
-			MassEntityManager.RemoveTagFromEntity(Entity, ColumnType);
-		}
-		else if (ColumnType->IsChildOf(FMassFragment::StaticStruct()))
+		if (ColumnType)
 		{
-			MassEntityManager.RemoveFragmentFromEntity(Entity, ColumnType);
+			FMassEntityHandle Entity = FMassEntityHandle::FromNumber(Row);
+			if (ColumnType->IsChildOf(FMassTag::StaticStruct()))
+			{
+				MassEntityManager.RemoveTagFromEntity(Entity, ColumnType);
+			}
+			else if (ColumnType->IsChildOf(FMassFragment::StaticStruct()))
+			{
+				MassEntityManager.RemoveFragmentFromEntity(Entity, ColumnType);
+			}
 		}
 	}
-}
 
-void FTypedElementDatabaseCommandBuffer::Execute_RemoveColumnsCommand(
-	FMassEntityManager& MassEntityManager,
-	TypedElementDataStorage::RowHandle Row,
-	FMassFragmentBitSet FragmentsToRemove,
-	FMassTagBitSet TagsToRemove)
-{
-	FMassArchetypeCompositionDescriptor RemoveComposition(
-		MoveTemp(FragmentsToRemove), MoveTemp(TagsToRemove), FMassChunkFragmentBitSet(), FMassSharedFragmentBitSet(), FMassConstSharedFragmentBitSet());
-	MassEntityManager.RemoveCompositionFromEntity(FMassEntityHandle::FromNumber(Row), RemoveComposition);
-}
+	void FCommandBuffer::Execute_RemoveColumnsCommand(
+		FMassEntityManager& MassEntityManager,
+		TypedElementDataStorage::RowHandle Row,
+		FMassFragmentBitSet FragmentsToRemove,
+		FMassTagBitSet TagsToRemove)
+	{
+		FMassArchetypeCompositionDescriptor RemoveComposition(
+			MoveTemp(FragmentsToRemove), MoveTemp(TagsToRemove), FMassChunkFragmentBitSet(), FMassSharedFragmentBitSet(), FMassConstSharedFragmentBitSet());
+		MassEntityManager.RemoveCompositionFromEntity(FMassEntityHandle::FromNumber(Row), RemoveComposition);
+	}
 
-void FTypedElementDatabaseCommandBuffer::ProcessCommands()
-{
-	using namespace TypedElementDataStorage;
+	void FCommandBuffer::ProcessCommands()
+	{
+		using namespace TypedElementDataStorage;
 
-	Commands.StableSort(
-		[](const FCommand& Lhs, const FCommand& Rhs)
-		{
-			return Lhs.Row < Rhs.Row;
-		});
+		Commands.StableSort(
+			[](const FCommand& Lhs, const FCommand& Rhs)
+			{
+				return Lhs.Row < Rhs.Row;
+			});
 
-	struct FProcessor
-	{
-		FMassEntityManager& EntityManager;
-		RowHandle Row;
-		void operator()(const FAddColumnCommand& Command) { Execute_AddColumnCommand(EntityManager, Row, Command.ColumnType.Get()); }
-		void operator()(const FAddDataColumnCommand& Command) { Execute_AddDataColumnCommand(EntityManager, Row, Command.ColumnType.Get(), Command.Data, Command.Relocator); }
-		void operator()(const FAddColumnsCommand& Command) { Execute_AddColumnsCommand(EntityManager, Row, Command.FragmentsToAdd, Command.TagsToAdd); }
-		void operator()(const FRemoveColumnCommand& Command) { Execute_RemoveColumnCommand(EntityManager, Row, Command.ColumnType.Get()); }
-		void operator()(const FRemoveColumnsCommand& Command) { Execute_RemoveColumnsCommand(EntityManager, Row, Command.FragmentsToRemove, Command.TagsToRemove); }
-	};
-	
-	FProcessor Processor{ .EntityManager = Environment.GetMassEntityManager() };
-	for (FCommand& Command : Commands)
-	{
-		if (Execute_IsRowAssigned(Environment.GetMassEntityManager(), Command.Row))
+		struct FProcessor
+		{
+			FMassEntityManager& EntityManager;
+			RowHandle Row;
+			void operator()(const FAddColumnCommand& Command) { Execute_AddColumnCommand(EntityManager, Row, Command.ColumnType.Get()); }
+			void operator()(const FAddDataColumnCommand& Command) { Execute_AddDataColumnCommand(EntityManager, Row, Command.ColumnType.Get(), Command.Data, Command.Relocator); }
+			void operator()(const FAddColumnsCommand& Command) { Execute_AddColumnsCommand(EntityManager, Row, Command.FragmentsToAdd, Command.TagsToAdd); }
+			void operator()(const FRemoveColumnCommand& Command) { Execute_RemoveColumnCommand(EntityManager, Row, Command.ColumnType.Get()); }
+			void operator()(const FRemoveColumnsCommand& Command) { Execute_RemoveColumnsCommand(EntityManager, Row, Command.FragmentsToRemove, Command.TagsToRemove); }
+		};
+
+		FProcessor Processor{ .EntityManager = Environment.GetMassEntityManager() };
+		for (FCommand& Command : Commands)
 		{
-			Processor.Row = Command.Row;
-			Visit(Processor, Command.Data);
+			if (Execute_IsRowAssigned(Environment.GetMassEntityManager(), Command.Row))
+			{
+				Processor.Row = Command.Row;
+				Visit(Processor, Command.Data);
+			}
 		}
-	}
 
-	PendingColumns.Reset();
-	Commands.Reset();
-}
+		PendingColumns.Reset();
+		Commands.Reset();
+	}
 
-// 
-// misc
-//
+	// 
+	// misc
+	//
 
-template<typename T>
-void FTypedElementDatabaseCommandBuffer::AddCommand(TypedElementDataStorage::RowHandle Row, T&& Args)
-{
-	FCommand Command;
-	Command.Row = Row;
-	Command.Data.Emplace<T>(Forward<T>(Args));
-	Commands.Add(MoveTemp(Command));
-}
+	template<typename T>
+	void FCommandBuffer::AddCommand(TypedElementDataStorage::RowHandle Row, T&& Args)
+	{
+		FCommand Command;
+		Command.Row = Row;
+		Command.Data.Emplace<T>(Forward<T>(Args));
+		Commands.Add(MoveTemp(Command));
+	}
+} // namespace UE::Editor::DataStorage::Legacy
diff --git a/Engine/Plugins/Experimental/EditorDataStorage/Source/TedsCore/Private/TypedElementDatabaseCommandBuffer.h b/Engine/Plugins/Experimental/EditorDataStorage/Source/TedsCore/Private/TypedElementDatabaseCommandBuffer.h
--- a/Engine/Plugins/Experimental/EditorDataStorage/Source/TedsCore/Private/TypedElementDatabaseCommandBuffer.h
+++ b/Engine/Plugins/Experimental/EditorDataStorage/Source/TedsCore/Private/TypedElementDatabaseCommandBuffer.h
@@ -10,106 +10,113 @@
 #include "UObject/Class.h"
 #include "UObject/WeakObjectPtrTemplates.h"
 
-class FTypedElementDatabaseEnvironment;
 struct FMassEntityManager;
 struct FMassProcessingPhaseManager;
 
-class FTypedElementDatabaseCommandBuffer final
+namespace UE::Editor::DataStorage
 {
-public:
-	explicit FTypedElementDatabaseCommandBuffer(FTypedElementDatabaseEnvironment& InEnvironment);
+	class FEnvironment;
 
-	/** Returns a pointer to the queued data column, if it exists and hasn't been processed yet. Otherwise null is returned. */
-	void* GetQueuedDataColumn(TypedElementDataStorage::RowHandle Row, const UScriptStruct* ColumnType);
-	/** Returns if the column on the provided row is pending processing. */
-	bool HasColumn(TypedElementDataStorage::RowHandle Row, const UScriptStruct* ColumnType) const;
-	/** Clears all entries in the internal cache for the provided row. */
-	void Clear(TypedElementDataStorage::RowHandle Row);
+	namespace Legacy
+	{
+		class FCommandBuffer final
+		{
+		public:
+			explicit FCommandBuffer(FEnvironment& InEnvironment);
 	
-	/**
-	 * @section Queue_* functions create a command that can be stored for later execution in the provided command buffer.
-	 * Typically the Execute_* counter part will be called for execution.
-	 */
-	void Queue_AddColumnCommand(TypedElementDataStorage::RowHandle Row, const UScriptStruct* ColumnType);
-	void* Queue_AddDataColumnCommandUnitialized(TypedElementDataStorage::RowHandle Row, const UScriptStruct* ColumnType, 
-		TypedElementDataStorage::ColumnCopyOrMoveCallback Relocator);
-	void Queue_AddColumnsCommand(TypedElementDataStorage::RowHandle Row, FMassFragmentBitSet FragmentsToAdd, FMassTagBitSet TagsToAdd);
-
-	void Queue_RemoveColumnCommand(TypedElementDataStorage::RowHandle Row, const UScriptStruct* ColumnType);
-	void Queue_RemoveColumnsCommand(TypedElementDataStorage::RowHandle Row, FMassFragmentBitSet FragmentsToRemove, FMassTagBitSet TagsToRemove);
-
-	/**
-	 * @section Execute_* functions directly execute a command with limited validation checks.
-	 */
+			/** Returns a pointer to the queued data column, if it exists and hasn't been processed yet. Otherwise null is returned. */
+			void* GetQueuedDataColumn(TypedElementDataStorage::RowHandle Row, const UScriptStruct* ColumnType);
+			/** Returns if the column on the provided row is pending processing. */
+			bool HasColumn(TypedElementDataStorage::RowHandle Row, const UScriptStruct* ColumnType) const;
+			/** Clears all entries in the internal cache for the provided row. */
+			void Clear(TypedElementDataStorage::RowHandle Row);
 	
-	static bool Execute_IsRowAvailable(const FMassEntityManager& MassEntityManager, TypedElementDataStorage::RowHandle Row);
-	static bool Execute_IsRowAssigned(const FMassEntityManager& MassEntityManager, TypedElementDataStorage::RowHandle Row);
-
-	static void Execute_AddColumnCommand(
-		FMassEntityManager& MassEntityManager, TypedElementDataStorage::RowHandle Row, const UScriptStruct* ColumnType);
-	static void Execute_AddDataColumnCommand(
-		FMassEntityManager& MassEntityManager, TypedElementDataStorage::RowHandle Row, const UScriptStruct* ColumnType,
-		void* Data, TypedElementDataStorage::ColumnCopyOrMoveCallback Relocator);
-	static void Execute_AddSharedColumnCommand(
-		FMassEntityManager& MassEntityManager, TypedElementDataStorage::RowHandle Row, const FConstSharedStruct& SharedColumn);
-	static void Execute_RemoveSharedColumnCommand(
-		FMassEntityManager& MassEntityManager, TypedElementDataStorage::RowHandle Row, const UScriptStruct& ColumnType);
-	static void Execute_AddColumnsCommand(
-		FMassEntityManager& MassEntityManager, TypedElementDataStorage::RowHandle Row,
-		FMassFragmentBitSet FragmentsToAdd, FMassTagBitSet TagsToAdd);
-	static void Execute_RemoveColumnCommand(
-		FMassEntityManager& MassEntityManager, TypedElementDataStorage::RowHandle Row, const UScriptStruct* ColumnType);
-	static void Execute_RemoveColumnsCommand(
-		FMassEntityManager& MassEntityManager, TypedElementDataStorage::RowHandle Row,
-		FMassFragmentBitSet FragmentsToRemove, FMassTagBitSet TagsToRemove);
-
-	void ProcessCommands();
-
-private:
-	struct FAddColumnCommand
-	{
-		TWeakObjectPtr<const UScriptStruct> ColumnType;
-	};
-	struct FAddDataColumnCommand
-	{
-		TWeakObjectPtr<const UScriptStruct> ColumnType;
-		TypedElementDataStorage::ColumnCopyOrMoveCallback Relocator;
-		void* Data;
-
-		~FAddDataColumnCommand();
-	};
-	struct FAddColumnsCommand
-	{
-		FMassFragmentBitSet FragmentsToAdd;
-		FMassTagBitSet TagsToAdd;
-	};
-	struct FRemoveColumnCommand
-	{
-		TWeakObjectPtr<const UScriptStruct> ColumnType;
-	};
-	struct FRemoveColumnsCommand
-	{
-		FMassFragmentBitSet FragmentsToRemove;
-		FMassTagBitSet TagsToRemove;
-	};
-	using CommandData = TVariant<
-		FAddColumnCommand,
-		FAddDataColumnCommand,
-		FAddColumnsCommand,
-		FRemoveColumnCommand,
-		FRemoveColumnsCommand>;
-
-	struct FCommand
-	{
-		TypedElementDataStorage::RowHandle Row;
-		CommandData Data;
-	};
-
-	template<typename T>
-	void AddCommand(TypedElementDataStorage::RowHandle Row, T&& Args);
-
-	using PendingColumnMappingKey = TPair<TypedElementDataStorage::RowHandle, TWeakObjectPtr<const UScriptStruct>>;
-	TMap<PendingColumnMappingKey, void*> PendingColumns;
-	TArray<FCommand> Commands;
-	FTypedElementDatabaseEnvironment& Environment;
-};
+			/**
+			 * @section Queue_* functions create a command that can be stored for later execution in the provided command buffer.
+			 * Typically the Execute_* counter part will be called for execution.
+			 */
+			void Queue_AddColumnCommand(TypedElementDataStorage::RowHandle Row, const UScriptStruct* ColumnType);
+			void* Queue_AddDataColumnCommandUnitialized(TypedElementDataStorage::RowHandle Row, const UScriptStruct* ColumnType,
+				TypedElementDataStorage::ColumnCopyOrMoveCallback Relocator);
+			void Queue_AddColumnsCommand(TypedElementDataStorage::RowHandle Row, FMassFragmentBitSet FragmentsToAdd, FMassTagBitSet TagsToAdd);
+	
+			void Queue_RemoveColumnCommand(TypedElementDataStorage::RowHandle Row, const UScriptStruct* ColumnType);
+			void Queue_RemoveColumnsCommand(TypedElementDataStorage::RowHandle Row, FMassFragmentBitSet FragmentsToRemove, FMassTagBitSet TagsToRemove);
+	
+			/**
+			 * @section Execute_* functions directly execute a command with limited validation checks.
+			 */
+	
+			static bool Execute_IsRowAvailable(const FMassEntityManager& MassEntityManager, TypedElementDataStorage::RowHandle Row);
+			static bool Execute_IsRowAssigned(const FMassEntityManager& MassEntityManager, TypedElementDataStorage::RowHandle Row);
+	
+			static void Execute_AddColumnCommand(
+				FMassEntityManager& MassEntityManager, TypedElementDataStorage::RowHandle Row, const UScriptStruct* ColumnType);
+			static void Execute_AddDataColumnCommand(
+				FMassEntityManager& MassEntityManager, TypedElementDataStorage::RowHandle Row, const UScriptStruct* ColumnType,
+				void* Data, TypedElementDataStorage::ColumnCopyOrMoveCallback Relocator);
+			static void Execute_AddSharedColumnCommand(
+				FMassEntityManager& MassEntityManager, TypedElementDataStorage::RowHandle Row, const FConstSharedStruct& SharedColumn);
+			static void Execute_RemoveSharedColumnCommand(
+				FMassEntityManager& MassEntityManager, TypedElementDataStorage::RowHandle Row, const UScriptStruct& ColumnType);
+			static void Execute_AddColumnsCommand(
+				FMassEntityManager& MassEntityManager, TypedElementDataStorage::RowHandle Row,
+				FMassFragmentBitSet FragmentsToAdd, FMassTagBitSet TagsToAdd);
+			static void Execute_RemoveColumnCommand(
+				FMassEntityManager& MassEntityManager, TypedElementDataStorage::RowHandle Row, const UScriptStruct* ColumnType);
+			static void Execute_RemoveColumnsCommand(
+				FMassEntityManager& MassEntityManager, TypedElementDataStorage::RowHandle Row,
+				FMassFragmentBitSet FragmentsToRemove, FMassTagBitSet TagsToRemove);
+	
+			void ProcessCommands();
+	
+		private:
+			struct FAddColumnCommand
+			{
+				TWeakObjectPtr<const UScriptStruct> ColumnType;
+			};
+			struct FAddDataColumnCommand
+			{
+				TWeakObjectPtr<const UScriptStruct> ColumnType;
+				TypedElementDataStorage::ColumnCopyOrMoveCallback Relocator;
+				void* Data;
+	
+				~FAddDataColumnCommand();
+			};
+			struct FAddColumnsCommand
+			{
+				FMassFragmentBitSet FragmentsToAdd;
+				FMassTagBitSet TagsToAdd;
+			};
+			struct FRemoveColumnCommand
+			{
+				TWeakObjectPtr<const UScriptStruct> ColumnType;
+			};
+			struct FRemoveColumnsCommand
+			{
+				FMassFragmentBitSet FragmentsToRemove;
+				FMassTagBitSet TagsToRemove;
+			};
+			using CommandData = TVariant<
+				FAddColumnCommand,
+				FAddDataColumnCommand,
+				FAddColumnsCommand,
+				FRemoveColumnCommand,
+				FRemoveColumnsCommand>;
+	
+			struct FCommand
+			{
+				TypedElementDataStorage::RowHandle Row;
+				CommandData Data;
+			};
+	
+			template<typename T>
+			void AddCommand(TypedElementDataStorage::RowHandle Row, T&& Args);
+	
+			using PendingColumnMappingKey = TPair<TypedElementDataStorage::RowHandle, TWeakObjectPtr<const UScriptStruct>>;
+			TMap<PendingColumnMappingKey, void*> PendingColumns;
+			TArray<FCommand> Commands;
+			FEnvironment& Environment;
+		};
+	} // namespace Legacy
+} // namespace UE::Editor::DataStorage
diff --git a/Engine/Plugins/Experimental/EditorDataStorage/Source/TedsCore/Private/TypedElementDatabaseCompatibility.cpp b/Engine/Plugins/Experimental/EditorDataStorage/Source/TedsCore/Private/TypedElementDatabaseCompatibility.cpp
--- a/Engine/Plugins/Experimental/EditorDataStorage/Source/TedsCore/Private/TypedElementDatabaseCompatibility.cpp
+++ b/Engine/Plugins/Experimental/EditorDataStorage/Source/TedsCore/Private/TypedElementDatabaseCompatibility.cpp
@@ -67,7 +67,7 @@ static const FName CompatibilityUsesCommandBufferExtensionName(TEXT("Compatiblit
 void UTypedElementDatabaseCompatibility::Initialize(UTypedElementDatabase* InStorage)
 {
 	using namespace TypedElementDataStorage;
-	using namespace UE::EditorDataStorage;
+	using namespace UE::Editor::DataStorage;
 
 	checkf(InStorage, TEXT("TEDS Compatibility is being initialized with an invalid storage target."));
 	
@@ -95,7 +95,7 @@ void UTypedElementDatabaseCompatibility::Initialize(UTypedElementDatabase* InSto
 void UTypedElementDatabaseCompatibility::Deinitialize()
 {
 	using namespace TypedElementDataStorage;
-	using namespace UE::EditorDataStorage;
+	using namespace UE::Editor::DataStorage;
 
 	FScopedExclusiveLock Lock(EGlobalLockScope::Public);
 	
@@ -118,7 +118,7 @@ void UTypedElementDatabaseCompatibility::Deinitialize()
 
 void UTypedElementDatabaseCompatibility::RegisterRegistrationFilter(ObjectRegistrationFilter Filter)
 {
-	using namespace UE::EditorDataStorage;
+	using namespace UE::Editor::DataStorage;
 
 	FScopedExclusiveLock Lock(EGlobalLockScope::Public);
 	ObjectRegistrationFilters.Add(MoveTemp(Filter));
@@ -126,7 +126,7 @@ void UTypedElementDatabaseCompatibility::RegisterRegistrationFilter(ObjectRegist
 
 void UTypedElementDatabaseCompatibility::RegisterDealiaserCallback(ObjectToRowDealiaser Dealiaser)
 {
-	using namespace UE::EditorDataStorage;
+	using namespace UE::Editor::DataStorage;
 
 	FScopedExclusiveLock Lock(EGlobalLockScope::Public);
 	ObjectToRowDialiasers.Add(MoveTemp(Dealiaser));
@@ -135,11 +135,11 @@ void UTypedElementDatabaseCompatibility::RegisterDealiaserCallback(ObjectToRowDe
 void UTypedElementDatabaseCompatibility::RegisterTypeTableAssociation(
 	TObjectPtr<UStruct> TypeInfo, TypedElementDataStorage::TableHandle Table)
 {
-	using namespace UE::EditorDataStorage;
+	using namespace UE::Editor::DataStorage;
 	
 	if (TypedElementDataStorage::bUseCommandBuffer)
 	{
-		QueuedCommands.AddCommand(UE::EditorDataStorage::FRegisterTypeTableAssociation{ .TypeInfo = TypeInfo, .Table = Table });
+		QueuedCommands.AddCommand(UE::Editor::DataStorage::FRegisterTypeTableAssociation{ .TypeInfo = TypeInfo, .Table = Table });
 	}
 	else
 	{
@@ -148,14 +148,14 @@ void UTypedElementDatabaseCompatibility::RegisterTypeTableAssociation(
 	}
 }
 
-FDelegateHandle UTypedElementDatabaseCompatibility::RegisterObjectAddedCallback(UE::EditorDataStorage::ObjectAddedCallback&& OnObjectAdded)
+FDelegateHandle UTypedElementDatabaseCompatibility::RegisterObjectAddedCallback(UE::Editor::DataStorage::ObjectAddedCallback&& OnObjectAdded)
 {
-	using namespace UE::EditorDataStorage;
+	using namespace UE::Editor::DataStorage;
 
 	FDelegateHandle Handle(FDelegateHandle::GenerateNewHandle);
 	if (TypedElementDataStorage::bUseCommandBuffer)
 	{
-		QueuedCommands.AddCommand(UE::EditorDataStorage::FRegisterObjectAddedCallback{ .Callback = MoveTemp(OnObjectAdded), .Handle = Handle });
+		QueuedCommands.AddCommand(UE::Editor::DataStorage::FRegisterObjectAddedCallback{ .Callback = MoveTemp(OnObjectAdded), .Handle = Handle });
 	}
 	else
 	{
@@ -167,26 +167,26 @@ FDelegateHandle UTypedElementDatabaseCompatibility::RegisterObjectAddedCallback(
 
 void UTypedElementDatabaseCompatibility::UnregisterObjectAddedCallback(FDelegateHandle Handle)
 {
-	using namespace UE::EditorDataStorage;
+	using namespace UE::Editor::DataStorage;
 
 	if (TypedElementDataStorage::bUseCommandBuffer)
 	{
-		QueuedCommands.AddCommand(UE::EditorDataStorage::FUnregisterObjectAddedCallback{ .Handle = Handle });
+		QueuedCommands.AddCommand(UE::Editor::DataStorage::FUnregisterObjectAddedCallback{ .Handle = Handle });
 	}
 	else
 	{
 		FScopedExclusiveLock Lock(EGlobalLockScope::Public);
 		ObjectAddedCallbackList.RemoveAll(
-			[Handle](const TPair<UE::EditorDataStorage::ObjectAddedCallback, FDelegateHandle>& Element)->bool
+			[Handle](const TPair<UE::Editor::DataStorage::ObjectAddedCallback, FDelegateHandle>& Element)->bool
 			{
 				return Element.Value == Handle;
 			});
 	}
 }
 
-FDelegateHandle UTypedElementDatabaseCompatibility::RegisterObjectRemovedCallback(UE::EditorDataStorage::ObjectRemovedCallback&& OnObjectAdded)
+FDelegateHandle UTypedElementDatabaseCompatibility::RegisterObjectRemovedCallback(UE::Editor::DataStorage::ObjectRemovedCallback&& OnObjectAdded)
 {
-	using namespace UE::EditorDataStorage;
+	using namespace UE::Editor::DataStorage;
 
 	FScopedExclusiveLock Lock(EGlobalLockScope::Public);
 
@@ -198,12 +198,12 @@ FDelegateHandle UTypedElementDatabaseCompatibility::RegisterObjectRemovedCallbac
 
 void UTypedElementDatabaseCompatibility::UnregisterObjectRemovedCallback(FDelegateHandle Handle)
 {
-	using namespace UE::EditorDataStorage;
+	using namespace UE::Editor::DataStorage;
 
 	// Since removing object has be immediately executed in some situation, adding the callback can not be delayed through the command buffer.
 	FScopedExclusiveLock Lock(EGlobalLockScope::Public);
 
-	PreObjectRemovedCallbackList.RemoveAll([Handle](const TPair<UE::EditorDataStorage::ObjectRemovedCallback, FDelegateHandle>& Element)->bool
+	PreObjectRemovedCallbackList.RemoveAll([Handle](const TPair<UE::Editor::DataStorage::ObjectRemovedCallback, FDelegateHandle>& Element)->bool
 	{
 		return Element.Value == Handle;
 	});
@@ -222,7 +222,7 @@ TypedElementRowHandle UTypedElementDatabaseCompatibility::AddCompatibleObjectExp
 TypedElementRowHandle UTypedElementDatabaseCompatibility::AddCompatibleObjectExplicit(void* Object, TWeakObjectPtr<UScriptStruct> TypeInfo)
 {
 	using namespace TypedElementDataStorage;
-	using namespace UE::EditorDataStorage;
+	using namespace UE::Editor::DataStorage;
 	
 	checkf(Storage, TEXT("Trying to add an object to Typed Element's Data Storage before the storage is available."));
 	FScopedExclusiveLock Lock(EGlobalLockScope::Public);
@@ -252,7 +252,7 @@ void UTypedElementDatabaseCompatibility::RemoveCompatibleObjectExplicit(UObject*
 void UTypedElementDatabaseCompatibility::RemoveCompatibleObjectExplicit(void* Object)
 {
 	using namespace TypedElementDataStorage;
-	using namespace UE::EditorDataStorage;
+	using namespace UE::Editor::DataStorage;
 
 	checkf(Storage, TEXT("Removing compatible objects is not supported before Typed Element's Database compatibility manager has been initialized."));
 
@@ -281,7 +281,7 @@ void UTypedElementDatabaseCompatibility::RemoveCompatibleObjectExplicit(void* Ob
 TypedElementRowHandle UTypedElementDatabaseCompatibility::FindRowWithCompatibleObjectExplicit(const UObject* Object) const
 {
 	using namespace TypedElementDataStorage;
-	using namespace UE::EditorDataStorage;
+	using namespace UE::Editor::DataStorage;
 
 	if (Object && Storage && Storage->IsAvailable())
 	{
@@ -402,7 +402,7 @@ void UTypedElementDatabaseCompatibility::RegisterTypeInformationQueries()
 bool UTypedElementDatabaseCompatibility::ShouldAddObject(const UObject* Object) const
 {
 	using namespace TypedElementDataStorage;
-	using namespace UE::EditorDataStorage;
+	using namespace UE::Editor::DataStorage;
 	
 	FScopedSharedLock Lock(EGlobalLockScope::Public);
 
@@ -422,7 +422,7 @@ bool UTypedElementDatabaseCompatibility::ShouldAddObject(const UObject* Object)
 TypedElementDataStorage::TableHandle UTypedElementDatabaseCompatibility::FindBestMatchingTable(const UStruct* TypeInfo) const
 {
 	using namespace TypedElementDataStorage;
-	using namespace UE::EditorDataStorage;
+	using namespace UE::Editor::DataStorage;
 
 	FScopedSharedLock Lock(EGlobalLockScope::Public);
 
@@ -442,7 +442,7 @@ template<bool bEnableTransactions>
 TypedElementRowHandle UTypedElementDatabaseCompatibility::AddCompatibleObjectExplicitTransactionable(UObject* Object)
 {
 	using namespace TypedElementDataStorage;
-	using namespace UE::EditorDataStorage;
+	using namespace UE::Editor::DataStorage;
 
 	FScopedExclusiveLock Lock(EGlobalLockScope::Public);
 
@@ -475,7 +475,7 @@ template<bool bEnableTransactions>
 void UTypedElementDatabaseCompatibility::RemoveCompatibleObjectExplicitTransactionable(const UObject* Object)
 {
 	using namespace TypedElementDataStorage;
-	using namespace UE::EditorDataStorage;
+	using namespace UE::Editor::DataStorage;
 
 	checkf(Storage,
 		TEXT("Removing compatible objects is not supported before Typed Element's Database compatibility manager has been initialized."));
@@ -508,7 +508,7 @@ void UTypedElementDatabaseCompatibility::RemoveCompatibleObjectExplicitTransacti
 	const UObject* Object, TypedElementDataStorage::RowHandle ObjectRow)
 {
 	using namespace TypedElementDataStorage;
-	using namespace UE::EditorDataStorage;
+	using namespace UE::Editor::DataStorage;
 
 	checkf(Storage,
 		TEXT("Removing compatible objects is not supported before Typed Element's Database compatibility manager has been initialized."));
@@ -563,7 +563,7 @@ TypedElementRowHandle UTypedElementDatabaseCompatibility::DealiasObject(const UO
 
 void UTypedElementDatabaseCompatibility::Tick()
 {
-	using namespace UE::EditorDataStorage;
+	using namespace UE::Editor::DataStorage;
 	
 	TEDS_EVENT_SCOPE(TEXT("Compatibility Tick"))
 	
@@ -615,7 +615,7 @@ void UTypedElementDatabaseCompatibility::FPendingTypeInformationUpdate::Process(
 {
 	using namespace TypedElementDataStorage;
 	using namespace TypedElementQueryBuilder;
-	using namespace UE::EditorDataStorage;
+	using namespace UE::Editor::DataStorage;
 
 	if (bHasPendingUpdate)
 	{
@@ -856,7 +856,7 @@ void UTypedElementDatabaseCompatibility::PendingRegistration<AddressType>::Reset
 
 void UTypedElementDatabaseCompatibility::TickPendingCommands()
 {
-	using namespace UE::EditorDataStorage;
+	using namespace UE::Editor::DataStorage;
 	using namespace TypedElementDataStorage;
 
 	// Thread safe because it's only called from functions that already lock.
@@ -998,7 +998,7 @@ void UTypedElementDatabaseCompatibility::TickObjectSync()
 
 void UTypedElementDatabaseCompatibility::OnPrePropertyChanged(UObject* Object, const FEditPropertyChain& PropertyChain)
 {
-	using namespace UE::EditorDataStorage;
+	using namespace UE::Editor::DataStorage;
 
 	if (TypedElementDataStorage::bUseCommandBuffer)
 	{
@@ -1016,7 +1016,7 @@ void UTypedElementDatabaseCompatibility::OnPostEditChangeProperty(
 	UObject* Object,
 	FPropertyChangedEvent& PropertyChangedEvent)
 {
-	using namespace UE::EditorDataStorage;
+	using namespace UE::Editor::DataStorage;
 
 	if (TypedElementDataStorage::bUseCommandBuffer)
 	{
@@ -1043,7 +1043,7 @@ void UTypedElementDatabaseCompatibility::OnPostEditChangeProperty(
 
 void UTypedElementDatabaseCompatibility::OnObjectModified(UObject* Object)
 {
-	using namespace UE::EditorDataStorage;
+	using namespace UE::Editor::DataStorage;
 
 	if (TypedElementDataStorage::bUseCommandBuffer)
 	{
@@ -1062,9 +1062,9 @@ void UTypedElementDatabaseCompatibility::OnObjectModified(UObject* Object)
 }
 
 void UTypedElementDatabaseCompatibility::TriggerOnObjectAdded(
-	const void* Object, UE::EditorDataStorage::FObjectTypeInfo TypeInfo, TypedElementDataStorage::RowHandle Row) const
+	const void* Object, UE::Editor::DataStorage::FObjectTypeInfo TypeInfo, TypedElementDataStorage::RowHandle Row) const
 {
-	using namespace UE::EditorDataStorage;
+	using namespace UE::Editor::DataStorage;
 
 	// Thread safe because it's only called from functions that already lock.
 
@@ -1076,9 +1076,9 @@ void UTypedElementDatabaseCompatibility::TriggerOnObjectAdded(
 }
 
 void UTypedElementDatabaseCompatibility::TriggerOnPreObjectRemoved(
-	const void* Object, UE::EditorDataStorage::FObjectTypeInfo TypeInfo, TypedElementDataStorage::RowHandle Row) const
+	const void* Object, UE::Editor::DataStorage::FObjectTypeInfo TypeInfo, TypedElementDataStorage::RowHandle Row) const
 {
-	using namespace UE::EditorDataStorage;
+	using namespace UE::Editor::DataStorage;
 
 	// Thread safe because it's only called from functions that already lock.
 
@@ -1091,7 +1091,7 @@ void UTypedElementDatabaseCompatibility::TriggerOnPreObjectRemoved(
 
 void UTypedElementDatabaseCompatibility::OnObjectReinstanced(const FCoreUObjectDelegates::FReplacementObjectMap& ReplacedObjects)
 {
-	using namespace UE::EditorDataStorage;
+	using namespace UE::Editor::DataStorage;
 	
 	if (TypedElementDataStorage::bUseCommandBuffer)
 	{
@@ -1118,7 +1118,7 @@ void UTypedElementDatabaseCompatibility::OnPostGcUnreachableAnalysis()
 {
 	using namespace TypedElementDataStorage;
 	using namespace TypedElementQueryBuilder;
-	using namespace UE::EditorDataStorage;
+	using namespace UE::Editor::DataStorage;
 
 	if (bIntegrateWithGC)
 	{
@@ -1179,7 +1179,7 @@ void UTypedElementDatabaseCompatibility::OnPostGcUnreachableAnalysis()
 
 void UTypedElementDatabaseCompatibility::OnPostWorldInitialization(UWorld* World, const UWorld::InitializationValues InitializationValues)
 {
-	using namespace UE::EditorDataStorage;
+	using namespace UE::Editor::DataStorage;
 
 	FScopedExclusiveLock Lock(EGlobalLockScope::Public);
 
@@ -1190,7 +1190,7 @@ void UTypedElementDatabaseCompatibility::OnPostWorldInitialization(UWorld* World
 
 void UTypedElementDatabaseCompatibility::OnPreWorldFinishDestroy(UWorld* World)
 {
-	using namespace UE::EditorDataStorage;
+	using namespace UE::Editor::DataStorage;
 
 	FScopedExclusiveLock Lock(EGlobalLockScope::Public);
 
@@ -1234,7 +1234,7 @@ UTypedElementDatabaseCompatibility::FRegistrationCommandChange::~FRegistrationCo
 	{
 		if (TypedElementDataStorage::bUseCommandBuffer)
 		{
-			DataStorageCompat->QueuedCommands.AddCommand(UE::EditorDataStorage::FDestroyMemento{ .MementoRow = MementoRow});
+			DataStorageCompat->QueuedCommands.AddCommand(UE::Editor::DataStorage::FDestroyMemento{ .MementoRow = MementoRow});
 		}
 		else
 		{
@@ -1246,7 +1246,7 @@ UTypedElementDatabaseCompatibility::FRegistrationCommandChange::~FRegistrationCo
 void UTypedElementDatabaseCompatibility::FRegistrationCommandChange::Apply(UObject* Object)
 {
 	using namespace TypedElementDataStorage;
-	using namespace UE::EditorDataStorage;
+	using namespace UE::Editor::DataStorage;
 
 	checkf(Owner.IsValid() && Owner.Get() == Object, 
 		TEXT("Applying registration transaction command within TEDS Compat was called after TEDS is not longer available."));
@@ -1271,7 +1271,7 @@ void UTypedElementDatabaseCompatibility::FRegistrationCommandChange::Apply(UObje
 void UTypedElementDatabaseCompatibility::FRegistrationCommandChange::Revert(UObject* Object)
 {
 	using namespace TypedElementDataStorage;
-	using namespace UE::EditorDataStorage;
+	using namespace UE::Editor::DataStorage;
 
 	checkf(Owner.IsValid() && Owner.Get() == Object,
 		TEXT("Reverting registration transaction command within TEDS Compat was called after TEDS is not longer available."));
@@ -1315,7 +1315,7 @@ UTypedElementDatabaseCompatibility::FDeregistrationCommandChange::FDeregistratio
 	, TargetObject(InTargetObject)
 {
 	using namespace TypedElementDataStorage;
-	using namespace UE::EditorDataStorage;
+	using namespace UE::Editor::DataStorage;
 
 	ITypedElementDataStorageInterface* DataStorage = InOwner->Storage;
 
@@ -1336,7 +1336,7 @@ UTypedElementDatabaseCompatibility::FDeregistrationCommandChange::FDeregistratio
 
 UTypedElementDatabaseCompatibility::FDeregistrationCommandChange::~FDeregistrationCommandChange()
 {
-	using namespace UE::EditorDataStorage;
+	using namespace UE::Editor::DataStorage;
 
 	// There's no memento row if target object was never registered with TEDS Compat.
 	if (UTypedElementDatabaseCompatibility* DataStorageCompat = Owner.Get();
@@ -1371,7 +1371,7 @@ void UTypedElementDatabaseCompatibility::FDeregistrationCommandChange::Apply(UOb
 void UTypedElementDatabaseCompatibility::FDeregistrationCommandChange::Revert(UObject* Object)
 {
 	using namespace TypedElementDataStorage;
-	using namespace UE::EditorDataStorage;
+	using namespace UE::Editor::DataStorage;
 
 	checkf(Owner.IsValid() && Owner.Get() == Object,
 		TEXT("Reverting deregistration transaction command within TEDS Compat was called after TEDS is not longer available."));
diff --git a/Engine/Plugins/Experimental/EditorDataStorage/Source/TedsCore/Private/TypedElementDatabaseCompatibility.h b/Engine/Plugins/Experimental/EditorDataStorage/Source/TedsCore/Private/TypedElementDatabaseCompatibility.h
--- a/Engine/Plugins/Experimental/EditorDataStorage/Source/TedsCore/Private/TypedElementDatabaseCompatibility.h
+++ b/Engine/Plugins/Experimental/EditorDataStorage/Source/TedsCore/Private/TypedElementDatabaseCompatibility.h
@@ -21,11 +21,15 @@
 
 class ITypedElementDataStorageInterface;
 struct FMassActorManager;
-class FTypedElementDatabaseEnvironment;
 class UTypedElementMementoSystem;
 
 enum class ETypedElementDatabaseCompatibilityObjectType : uint8;
 
+namespace UE::Editor::DataStorage
+{
+	class FEnvironment;
+} // namespace UE::Editor::DataStorage
+
 UCLASS()
 class TEDSCORE_API UTypedElementDatabaseCompatibility
 	: public UObject
@@ -33,9 +37,9 @@ class TEDSCORE_API UTypedElementDatabaseCompatibility
 {
 	GENERATED_BODY()
 
-	friend struct UE::EditorDataStorage::FCommandProcessor;
-	friend struct UE::EditorDataStorage::FPatchData;
-	friend struct UE::EditorDataStorage::FPrepareCommands;
+	friend struct UE::Editor::DataStorage::FCommandProcessor;
+	friend struct UE::Editor::DataStorage::FPatchData;
+	friend struct UE::Editor::DataStorage::FPrepareCommands;
 public:
 	~UTypedElementDatabaseCompatibility() override = default;
 
@@ -45,9 +49,9 @@ public:
 	void RegisterRegistrationFilter(ObjectRegistrationFilter Filter) override;
 	void RegisterDealiaserCallback(ObjectToRowDealiaser Dealiaser) override;
 	void RegisterTypeTableAssociation(TObjectPtr<UStruct> TypeInfo, TypedElementDataStorage::TableHandle Table) override;
-	FDelegateHandle RegisterObjectAddedCallback(UE::EditorDataStorage::ObjectAddedCallback&& OnObjectAdded);
+	FDelegateHandle RegisterObjectAddedCallback(UE::Editor::DataStorage::ObjectAddedCallback&& OnObjectAdded);
 	void UnregisterObjectAddedCallback(FDelegateHandle Handle);
-	FDelegateHandle RegisterObjectRemovedCallback(UE::EditorDataStorage::ObjectRemovedCallback&& OnObjectRemoved);
+	FDelegateHandle RegisterObjectRemovedCallback(UE::Editor::DataStorage::ObjectRemovedCallback&& OnObjectRemoved);
 	void UnregisterObjectRemovedCallback(FDelegateHandle Handle);
 	
 	TypedElementRowHandle AddCompatibleObjectExplicit(UObject* Object) override;
@@ -121,8 +125,8 @@ private:
 	void OnPrePropertyChanged(UObject* Object, const FEditPropertyChain& PropertyChain);
 	void OnPostEditChangeProperty(UObject* Object, FPropertyChangedEvent& PropertyChangedEvent);
 	void OnObjectModified(UObject* Object);
-	void TriggerOnObjectAdded(const void* Object, UE::EditorDataStorage::FObjectTypeInfo TypeInfo, TypedElementDataStorage::RowHandle Row) const;
-	void TriggerOnPreObjectRemoved(const void* Object, UE::EditorDataStorage::FObjectTypeInfo TypeInfo, TypedElementDataStorage::RowHandle Row) const;
+	void TriggerOnObjectAdded(const void* Object, UE::Editor::DataStorage::FObjectTypeInfo TypeInfo, TypedElementDataStorage::RowHandle Row) const;
+	void TriggerOnPreObjectRemoved(const void* Object, UE::Editor::DataStorage::FObjectTypeInfo TypeInfo, TypedElementDataStorage::RowHandle Row) const;
 	void OnObjectReinstanced(const FCoreUObjectDelegates::FReplacementObjectMap& ReplacedObjects);
 
 	void OnPostGcUnreachableAnalysis();
@@ -187,8 +191,8 @@ private:
 		void Reset();
 	};
 
-	UE::EditorDataStorage::CompatibilityCommandBuffer QueuedCommands;
-	UE::EditorDataStorage::CompatibilityCommandBuffer::FCollection PendingCommands;
+	UE::Editor::DataStorage::CompatibilityCommandBuffer QueuedCommands;
+	UE::Editor::DataStorage::CompatibilityCommandBuffer::FCollection PendingCommands;
 	PendingRegistration<TWeakObjectPtr<UObject>> UObjectsPendingRegistration;
 	PendingRegistration<ExternalObjectRegistration> ExternalObjectsPendingRegistration;
 	TArray<TypedElementDataStorage::RowHandle> RowScratchBuffer;
@@ -197,8 +201,8 @@ private:
 	TArray<ObjectToRowDealiaser> ObjectToRowDialiasers;
 	using TypeToTableMapType = TMap<TWeakObjectPtr<UStruct>, TypedElementDataStorage::TableHandle>;
 	TypeToTableMapType TypeToTableMap;
-	TArray<TPair<UE::EditorDataStorage::ObjectAddedCallback, FDelegateHandle>> ObjectAddedCallbackList;
-	TArray<TPair<UE::EditorDataStorage::ObjectRemovedCallback, FDelegateHandle>> PreObjectRemovedCallbackList;
+	TArray<TPair<UE::Editor::DataStorage::ObjectAddedCallback, FDelegateHandle>> ObjectAddedCallbackList;
+	TArray<TPair<UE::Editor::DataStorage::ObjectRemovedCallback, FDelegateHandle>> PreObjectRemovedCallbackList;
 
 	TypedElementDataStorage::TableHandle StandardActorTable{ TypedElementDataStorage::InvalidTableHandle };
 	TypedElementDataStorage::TableHandle StandardActorWithTransformTable{ TypedElementDataStorage::InvalidTableHandle };
@@ -234,7 +238,7 @@ private:
 	FDelegateHandle ObjectReinstancedDelegateHandle;
 	FDelegateHandle PostGcUnreachableAnalysisHandle;
 	
-	TSharedPtr<FTypedElementDatabaseEnvironment> Environment;
+	TSharedPtr<UE::Editor::DataStorage::FEnvironment> Environment;
 	TypedElementDataStorage::QueryHandle ClassTypeInfoQuery;
 	TypedElementDataStorage::QueryHandle ScriptStructTypeInfoQuery;
 	TypedElementDataStorage::QueryHandle UObjectQuery;
