[Teds] Rename native classes and move into UE::EditorDataStorage namespaces

diff --git a/Engine/Plugins/Experimental/EditorDataStorage/Source/TedsCore/Private/TypedElementDatabaseEnvironment.cpp b/Engine/Plugins/Experimental/EditorDataStorage/Source/TedsCore/Private/TypedElementDatabaseEnvironment.cpp
--- a/Engine/Plugins/Experimental/EditorDataStorage/Source/TedsCore/Private/TypedElementDatabaseEnvironment.cpp
+++ b/Engine/Plugins/Experimental/EditorDataStorage/Source/TedsCore/Private/TypedElementDatabaseEnvironment.cpp
@@ -3,111 +3,113 @@
 #include "TypedElementDatabaseEnvironment.h"
 #include "TypedElementDatabase.h"
 
-FTypedElementDatabaseEnvironment::FTypedElementDatabaseEnvironment(UTypedElementDatabase& InDataStorage,
-	FMassEntityManager& InMassEntityManager, FMassProcessingPhaseManager& InMassPhaseManager)
-	: DataStorage(InDataStorage)
-	, DirectDeferredCommands(*this)
-	, MementoSystem(InDataStorage)
-	, DynamicTagManager(DynamicColumnGenerator)
-	, MassEntityManager(InMassEntityManager)
-	, MassPhaseManager(InMassPhaseManager)
-{
-}
-
-FTypedElementDatabaseCommandBuffer& FTypedElementDatabaseEnvironment::GetDirectDeferredCommands()
-{
-	return DirectDeferredCommands;
-}
-
-const FTypedElementDatabaseCommandBuffer& FTypedElementDatabaseEnvironment::GetDirectDeferredCommands() const
-{
-	return DirectDeferredCommands;
-}
-
-FTypedElementDatabaseIndexTable& FTypedElementDatabaseEnvironment::GetIndexTable()
-{
-	return IndexTable;
-}
-
-const FTypedElementDatabaseIndexTable& FTypedElementDatabaseEnvironment::GetIndexTable() const
-{
-	return IndexTable;
-}
-
-FTypedElementDatabaseScratchBuffer& FTypedElementDatabaseEnvironment::GetScratchBuffer()
-{
-	return ScratchBuffer;
-}
-
-const FTypedElementDatabaseScratchBuffer& FTypedElementDatabaseEnvironment::GetScratchBuffer() const
-{
-	return ScratchBuffer;
-}
-
-FTypedElementExtendedQueryStore& FTypedElementDatabaseEnvironment::GetQueryStore()
-{
-	return Queries;
-}
-
-const FTypedElementExtendedQueryStore& FTypedElementDatabaseEnvironment::GetQueryStore() const
-{
-	return Queries;
-}
-
-UTypedElementMementoSystem& FTypedElementDatabaseEnvironment::GetMementoSystem()
-{
-	return MementoSystem;
-}
-
-const UTypedElementMementoSystem& FTypedElementDatabaseEnvironment::GetMementoSystem() const
-{
-	return MementoSystem;
-}
-
-FMassEntityManager& FTypedElementDatabaseEnvironment::GetMassEntityManager()
-{
-	return MassEntityManager;
-}
-
-const FMassEntityManager& FTypedElementDatabaseEnvironment::GetMassEntityManager() const
-{
-	return MassEntityManager;
-}
-
-FMassArchetypeHandle FTypedElementDatabaseEnvironment::LookupMassArchetype(TypedElementDataStorage::TableHandle TableHandle) const
-{
-	return DataStorage.LookupArchetype(TableHandle);
-}
-
-FMassProcessingPhaseManager& FTypedElementDatabaseEnvironment::GetMassPhaseManager()
-{
-	return MassPhaseManager;
-}
-
-const FMassProcessingPhaseManager& FTypedElementDatabaseEnvironment::GetMassPhaseManager() const
-{
-	return MassPhaseManager;
-}
-FConstSharedStruct GenerateDynamicTag(const UE::EditorDataStorage::FDynamicTag& Tag, const FName& Value);
-const UScriptStruct* GenerateColumnType(const UE::EditorDataStorage::FDynamicTag& Tag);
-FConstSharedStruct FTypedElementDatabaseEnvironment::GenerateDynamicTag(const UE::EditorDataStorage::FDynamicTag& Tag, const FName& Value)
-{
-	return DynamicTagManager.GenerateDynamicTag(Tag, Value);
-}
-
-const UScriptStruct* FTypedElementDatabaseEnvironment::GenerateColumnType(const UE::EditorDataStorage::FDynamicTag& Tag)
-{
-	return DynamicTagManager.GenerateColumnType(Tag);
-}
-
-void FTypedElementDatabaseEnvironment::NextUpdateCycle()
-{
-	Queries.UpdateActivatableQueries();
-	ScratchBuffer.BatchDelete();
-	UpdateCycleId++;
-}
-
-uint64 FTypedElementDatabaseEnvironment::GetUpdateCycleId() const
-{
-	return UpdateCycleId;
-}
+namespace UE::Editor::DataStorage
+{
+	FEnvironment::FEnvironment(UTypedElementDatabase& InDataStorage,
+		FMassEntityManager& InMassEntityManager, FMassProcessingPhaseManager& InMassPhaseManager)
+		: DataStorage(InDataStorage)
+		, DirectDeferredCommands(*this)
+		, MementoSystem(InDataStorage)
+		, DynamicTagManager(DynamicColumnGenerator)
+		, MassEntityManager(InMassEntityManager)
+		, MassPhaseManager(InMassPhaseManager)
+	{
+	}
+
+	Legacy::FCommandBuffer& FEnvironment::GetDirectDeferredCommands()
+	{
+		return DirectDeferredCommands;
+	}
+
+	const Legacy::FCommandBuffer& FEnvironment::GetDirectDeferredCommands() const
+	{
+		return DirectDeferredCommands;
+	}
+
+	FIndexTable& FEnvironment::GetIndexTable()
+	{
+		return IndexTable;
+	}
+
+	const FIndexTable& FEnvironment::GetIndexTable() const
+	{
+		return IndexTable;
+	}
+
+	FScratchBuffer& FEnvironment::GetScratchBuffer()
+	{
+		return ScratchBuffer;
+	}
+
+	const FScratchBuffer& FEnvironment::GetScratchBuffer() const
+	{
+		return ScratchBuffer;
+	}
+
+	FExtendedQueryStore& FEnvironment::GetQueryStore()
+	{
+		return Queries;
+	}
+
+	const FExtendedQueryStore& FEnvironment::GetQueryStore() const
+	{
+		return Queries;
+	}
+
+	UTypedElementMementoSystem& FEnvironment::GetMementoSystem()
+	{
+		return MementoSystem;
+	}
+
+	const UTypedElementMementoSystem& FEnvironment::GetMementoSystem() const
+	{
+		return MementoSystem;
+	}
+
+	FMassEntityManager& FEnvironment::GetMassEntityManager()
+	{
+		return MassEntityManager;
+	}
+
+	const FMassEntityManager& FEnvironment::GetMassEntityManager() const
+	{
+		return MassEntityManager;
+	}
+
+	FMassArchetypeHandle FEnvironment::LookupMassArchetype(TypedElementDataStorage::TableHandle TableHandle) const
+	{
+		return DataStorage.LookupArchetype(TableHandle);
+	}
+
+	FMassProcessingPhaseManager& FEnvironment::GetMassPhaseManager()
+	{
+		return MassPhaseManager;
+	}
+
+	const FMassProcessingPhaseManager& FEnvironment::GetMassPhaseManager() const
+	{
+		return MassPhaseManager;
+	}
+
+	FConstSharedStruct FEnvironment::GenerateDynamicTag(const FDynamicTag& Tag, const FName& Value)
+	{
+		return DynamicTagManager.GenerateDynamicTag(Tag, Value);
+	}
+
+	const UScriptStruct* FEnvironment::GenerateColumnType(const FDynamicTag& Tag)
+	{
+		return DynamicTagManager.GenerateColumnType(Tag);
+	}
+
+	void FEnvironment::NextUpdateCycle()
+	{
+		Queries.UpdateActivatableQueries();
+		ScratchBuffer.BatchDelete();
+		UpdateCycleId++;
+	}
+
+	uint64 FEnvironment::GetUpdateCycleId() const
+	{
+		return UpdateCycleId;
+	}
+} // namespace UE::Editor::DataStorage
diff --git a/Engine/Plugins/Experimental/EditorDataStorage/Source/TedsCore/Private/TypedElementDatabaseEnvironment.h b/Engine/Plugins/Experimental/EditorDataStorage/Source/TedsCore/Private/TypedElementDatabaseEnvironment.h
--- a/Engine/Plugins/Experimental/EditorDataStorage/Source/TedsCore/Private/TypedElementDatabaseEnvironment.h
+++ b/Engine/Plugins/Experimental/EditorDataStorage/Source/TedsCore/Private/TypedElementDatabaseEnvironment.h
@@ -13,53 +13,56 @@
 
 class UTypedElementDatabase;
 
-class FTypedElementDatabaseEnvironment final
+namespace UE::Editor::DataStorage
 {
-public:
-	FTypedElementDatabaseEnvironment(UTypedElementDatabase& InDataStorage, 
-		FMassEntityManager& InMassEntityManager, FMassProcessingPhaseManager& InMassPhaseManager);
-
-	FTypedElementDatabaseCommandBuffer& GetDirectDeferredCommands();
-	const FTypedElementDatabaseCommandBuffer& GetDirectDeferredCommands() const;
-	
-	FTypedElementDatabaseIndexTable& GetIndexTable();
-	const FTypedElementDatabaseIndexTable& GetIndexTable() const;
-
-	FTypedElementDatabaseScratchBuffer& GetScratchBuffer();
-	const FTypedElementDatabaseScratchBuffer& GetScratchBuffer() const;
-
-	FTypedElementExtendedQueryStore& GetQueryStore();
-	const FTypedElementExtendedQueryStore& GetQueryStore() const;
-
-	UTypedElementMementoSystem& GetMementoSystem();
-	const UTypedElementMementoSystem& GetMementoSystem() const;
-
-	FMassEntityManager& GetMassEntityManager();
-	const FMassEntityManager& GetMassEntityManager() const;
-
-	FMassArchetypeHandle LookupMassArchetype(TypedElementDataStorage::TableHandle TableHandle) const;
-	
-	FMassProcessingPhaseManager& GetMassPhaseManager();
-	const FMassProcessingPhaseManager& GetMassPhaseManager() const;
-
-	FConstSharedStruct GenerateDynamicTag(const UE::EditorDataStorage::FDynamicTag& Tag, const FName& Value);
-	const UScriptStruct* GenerateColumnType(const UE::EditorDataStorage::FDynamicTag& Tag);
-
-	void NextUpdateCycle();
-	uint64 GetUpdateCycleId() const;
-
-private:
-	UTypedElementDatabase& DataStorage;
-	FTypedElementDatabaseCommandBuffer DirectDeferredCommands;
-	FTypedElementDatabaseIndexTable IndexTable;
-	FTypedElementDatabaseScratchBuffer ScratchBuffer;
-	FTypedElementExtendedQueryStore Queries;
-	UTypedElementMementoSystem MementoSystem;
-	UE::EditorDataStorage::FDynamicColumnGenerator DynamicColumnGenerator;
-	UE::EditorDataStorage::FDynamicTagManager DynamicTagManager;
-
-	FMassEntityManager& MassEntityManager;
-	FMassProcessingPhaseManager& MassPhaseManager;
-
-	uint64 UpdateCycleId = 0;
-};
+	class FEnvironment final
+	{
+	public:
+		FEnvironment(UTypedElementDatabase& InDataStorage,
+			FMassEntityManager& InMassEntityManager, FMassProcessingPhaseManager& InMassPhaseManager);
+
+		Legacy::FCommandBuffer& GetDirectDeferredCommands();
+		const Legacy::FCommandBuffer& GetDirectDeferredCommands() const;
+
+		FIndexTable& GetIndexTable();
+		const FIndexTable& GetIndexTable() const;
+
+		FScratchBuffer& GetScratchBuffer();
+		const FScratchBuffer& GetScratchBuffer() const;
+
+		FExtendedQueryStore& GetQueryStore();
+		const FExtendedQueryStore& GetQueryStore() const;
+
+		UTypedElementMementoSystem& GetMementoSystem();
+		const UTypedElementMementoSystem& GetMementoSystem() const;
+
+		FMassEntityManager& GetMassEntityManager();
+		const FMassEntityManager& GetMassEntityManager() const;
+
+		FMassArchetypeHandle LookupMassArchetype(TypedElementDataStorage::TableHandle TableHandle) const;
+
+		FMassProcessingPhaseManager& GetMassPhaseManager();
+		const FMassProcessingPhaseManager& GetMassPhaseManager() const;
+
+		FConstSharedStruct GenerateDynamicTag(const FDynamicTag& Tag, const FName& Value);
+		const UScriptStruct* GenerateColumnType(const FDynamicTag& Tag);
+
+		void NextUpdateCycle();
+		uint64 GetUpdateCycleId() const;
+
+	private:
+		UTypedElementDatabase& DataStorage;
+		Legacy::FCommandBuffer DirectDeferredCommands;
+		FIndexTable IndexTable;
+		FScratchBuffer ScratchBuffer;
+		FExtendedQueryStore Queries;
+		UTypedElementMementoSystem MementoSystem;
+		FDynamicColumnGenerator DynamicColumnGenerator;
+		FDynamicTagManager DynamicTagManager;
+
+		FMassEntityManager& MassEntityManager;
+		FMassProcessingPhaseManager& MassPhaseManager;
+
+		uint64 UpdateCycleId = 0;
+	};
+} // namespace UE::Editor::DataStorage
diff --git a/Engine/Plugins/Experimental/EditorDataStorage/Source/TedsCore/Private/TypedElementDatabaseIndexTable.cpp b/Engine/Plugins/Experimental/EditorDataStorage/Source/TedsCore/Private/TypedElementDatabaseIndexTable.cpp
--- a/Engine/Plugins/Experimental/EditorDataStorage/Source/TedsCore/Private/TypedElementDatabaseIndexTable.cpp
+++ b/Engine/Plugins/Experimental/EditorDataStorage/Source/TedsCore/Private/TypedElementDatabaseIndexTable.cpp
@@ -2,96 +2,91 @@
 
 #include "TypedElementDatabaseIndexTable.h"
 
-TypedElementDataStorage::RowHandle FTypedElementDatabaseIndexTable::FindIndexedRow(
-	UE::EditorDataStorage::EGlobalLockScope LockScope, TypedElementDataStorage::IndexHash Index) const
+namespace UE::Editor::DataStorage
 {
-	using namespace TypedElementDataStorage;
-	using namespace UE::EditorDataStorage;
+	TypedElementDataStorage::RowHandle FIndexTable::FindIndexedRow(
+		EGlobalLockScope LockScope, TypedElementDataStorage::IndexHash Index) const
+	{
+		using namespace TypedElementDataStorage;
 	
-	FScopedSharedLock Lock(LockScope);
+		FScopedSharedLock Lock(LockScope);
 	
-	const RowHandle* Result = IndexLookupMap.Find(Index);
-	return Result ? *Result : InvalidRowHandle;
-}
+		const RowHandle* Result = IndexLookupMap.Find(Index);
+		return Result ? *Result : InvalidRowHandle;
+	}
 
-void FTypedElementDatabaseIndexTable::BatchIndexRows(UE::EditorDataStorage::EGlobalLockScope LockScope,
-	TConstArrayView<TPair<TypedElementDataStorage::IndexHash, TypedElementDataStorage::RowHandle>> IndexRowPairs)
-{
-	using namespace TypedElementDataStorage;
-	using namespace UE::EditorDataStorage;
+	void FIndexTable::BatchIndexRows(EGlobalLockScope LockScope,
+		TConstArrayView<TPair<TypedElementDataStorage::IndexHash, TypedElementDataStorage::RowHandle>> IndexRowPairs)
+	{
+		using namespace TypedElementDataStorage;
 
-	FScopedExclusiveLock Lock(LockScope);
+		FScopedExclusiveLock Lock(LockScope);
 	
-	IndexLookupMap.Reserve(IndexLookupMap.Num() + IndexRowPairs.Num());
-	ReverseIndexLookupMap.Reserve(ReverseIndexLookupMap.Num() + IndexRowPairs.Num());
+		IndexLookupMap.Reserve(IndexLookupMap.Num() + IndexRowPairs.Num());
+		ReverseIndexLookupMap.Reserve(ReverseIndexLookupMap.Num() + IndexRowPairs.Num());
 
-	for (const TPair<IndexHash, RowHandle>& IndexAndRow : IndexRowPairs)
-	{
-		IndexRowUnguarded(IndexAndRow.Key, IndexAndRow.Value);
+		for (const TPair<IndexHash, RowHandle>& IndexAndRow : IndexRowPairs)
+		{
+			IndexRowUnguarded(IndexAndRow.Key, IndexAndRow.Value);
+		}
 	}
-}
-
-void FTypedElementDatabaseIndexTable::IndexRow(UE::EditorDataStorage::EGlobalLockScope LockScope,
-	TypedElementDataStorage::IndexHash Index, TypedElementDataStorage::RowHandle Row)
-{
-	using namespace UE::EditorDataStorage;
 
-	FScopedExclusiveLock Lock(LockScope);
-	IndexRowUnguarded(Index, Row);
-}
+	void FIndexTable::IndexRow(EGlobalLockScope LockScope,
+		TypedElementDataStorage::IndexHash Index, TypedElementDataStorage::RowHandle Row)
+	{
 
-void FTypedElementDatabaseIndexTable::ReindexRow(UE::EditorDataStorage::EGlobalLockScope LockScope,
-	TypedElementDataStorage::IndexHash OriginalIndex, TypedElementDataStorage::IndexHash NewIndex, 
-	TypedElementDataStorage::RowHandle Row)
-{
-	using namespace UE::EditorDataStorage;
+		FScopedExclusiveLock Lock(LockScope);
+		IndexRowUnguarded(Index, Row);
+	}
 
-	FScopedExclusiveLock Lock(LockScope);
+	void FIndexTable::ReindexRow(EGlobalLockScope LockScope,
+		TypedElementDataStorage::IndexHash OriginalIndex, TypedElementDataStorage::IndexHash NewIndex, 
+		TypedElementDataStorage::RowHandle Row)
+	{
+		FScopedExclusiveLock Lock(LockScope);
 	
-	RemoveIndexUnguarded(OriginalIndex);
-	IndexRowUnguarded(NewIndex, Row);
-}
-
-void FTypedElementDatabaseIndexTable::RemoveIndex(UE::EditorDataStorage::EGlobalLockScope LockScope, TypedElementDataStorage::IndexHash Index)
-{
-	using namespace UE::EditorDataStorage;
+		RemoveIndexUnguarded(OriginalIndex);
+		IndexRowUnguarded(NewIndex, Row);
+	}
 
-	FScopedExclusiveLock Lock(LockScope);
-	RemoveIndexUnguarded(Index);
-}
+	void FIndexTable::RemoveIndex(EGlobalLockScope LockScope, TypedElementDataStorage::IndexHash Index)
+	{
+		FScopedExclusiveLock Lock(LockScope);
+		RemoveIndexUnguarded(Index);
+	}
 
-void FTypedElementDatabaseIndexTable::RemoveRow(UE::EditorDataStorage::EGlobalLockScope LockScope, TypedElementDataStorage::RowHandle Row)
-{
-	using namespace TypedElementDataStorage;
-	using namespace UE::EditorDataStorage;
+	void FIndexTable::RemoveRow(EGlobalLockScope LockScope, TypedElementDataStorage::RowHandle Row)
+	{
+		using namespace TypedElementDataStorage;
 
-	FScopedExclusiveLock Lock(LockScope);
+		FScopedExclusiveLock Lock(LockScope);
 	
-	if (TMultiMap<RowHandle, IndexHash>::TKeyIterator It = ReverseIndexLookupMap.CreateKeyIterator(Row); It)
-	{
-		do
+		if (TMultiMap<RowHandle, IndexHash>::TKeyIterator It = ReverseIndexLookupMap.CreateKeyIterator(Row); It)
 		{
-			IndexLookupMap.Remove(It.Value());
-			++It;
-		} while (It);
-		ReverseIndexLookupMap.Remove(Row);
+			do
+			{
+				IndexLookupMap.Remove(It.Value());
+				++It;
+			} while (It);
+			ReverseIndexLookupMap.Remove(Row);
+		}
 	}
-}
-
-void FTypedElementDatabaseIndexTable::IndexRowUnguarded(TypedElementDataStorage::IndexHash Index, TypedElementDataStorage::RowHandle Row)
-{
-	IndexLookupMap.Add(Index, Row);
-	ReverseIndexLookupMap.Add(Row, Index);
-}
 
+	void FIndexTable::IndexRowUnguarded(TypedElementDataStorage::IndexHash Index, TypedElementDataStorage::RowHandle Row)
+	{
+		IndexLookupMap.Add(Index, Row);
+		ReverseIndexLookupMap.Add(Row, Index);
+	}
 
-void FTypedElementDatabaseIndexTable::RemoveIndexUnguarded(TypedElementDataStorage::IndexHash Index)
-{
-	using namespace TypedElementDataStorage;
 
-	if (const RowHandle* Row = IndexLookupMap.Find(Index))
+	void FIndexTable::RemoveIndexUnguarded(TypedElementDataStorage::IndexHash Index)
 	{
-		IndexLookupMap.Remove(Index);
-		ReverseIndexLookupMap.Remove(*Row, Index);
+		using namespace TypedElementDataStorage;
+
+		if (const RowHandle* Row = IndexLookupMap.Find(Index))
+		{
+			IndexLookupMap.Remove(Index);
+			ReverseIndexLookupMap.Remove(*Row, Index);
+		}
 	}
-}
\ No newline at end of file
+} // namespace UE::Editor::DataStorage
diff --git a/Engine/Plugins/Experimental/EditorDataStorage/Source/TedsCore/Private/TypedElementDatabaseIndexTable.h b/Engine/Plugins/Experimental/EditorDataStorage/Source/TedsCore/Private/TypedElementDatabaseIndexTable.h
--- a/Engine/Plugins/Experimental/EditorDataStorage/Source/TedsCore/Private/TypedElementDatabaseIndexTable.h
+++ b/Engine/Plugins/Experimental/EditorDataStorage/Source/TedsCore/Private/TypedElementDatabaseIndexTable.h
@@ -7,41 +7,45 @@
 #include "Elements/Interfaces/TypedElementQueryStorageInterfaces.h"
 #include "GlobalLock.h"
 
-/**
- * Storage for an index to row mapping.
- * Access to the index table is thread safe and guarded by the global lock.
- */
 
-class FTypedElementDatabaseIndexTable final
+namespace UE::Editor::DataStorage
 {
-public:
-	TypedElementDataStorage::RowHandle FindIndexedRow(
-		UE::EditorDataStorage::EGlobalLockScope LockScope,
-		TypedElementDataStorage::IndexHash Index) const;
-	void IndexRow(
-		UE::EditorDataStorage::EGlobalLockScope LockScope,
-		TypedElementDataStorage::IndexHash Index, 
-		TypedElementDataStorage::RowHandle Row);
-	void BatchIndexRows(
-		UE::EditorDataStorage::EGlobalLockScope LockScope,
-		TConstArrayView<TPair<TypedElementDataStorage::IndexHash, 
-		TypedElementDataStorage::RowHandle>> IndexRowPairs);
-	void ReindexRow(
-		UE::EditorDataStorage::EGlobalLockScope LockScope,
-		TypedElementDataStorage::IndexHash OriginalIndex,
-		TypedElementDataStorage::IndexHash NewIndex,
-		TypedElementDataStorage::RowHandle Row);
-	void RemoveIndex(
-		UE::EditorDataStorage::EGlobalLockScope LockScope,
-		TypedElementDataStorage::IndexHash Index);
-	void RemoveRow(
-		UE::EditorDataStorage::EGlobalLockScope LockScope,
-		TypedElementDataStorage::RowHandle Row);
+	/**
+	 * Storage for an index to row mapping.
+	 * Access to the index table is thread safe and guarded by the global lock.
+	 */
 
-private:
-	TMap<TypedElementDataStorage::IndexHash, TypedElementDataStorage::RowHandle> IndexLookupMap;
-	TMultiMap<TypedElementDataStorage::RowHandle, TypedElementDataStorage::IndexHash> ReverseIndexLookupMap;
+	class FIndexTable final
+	{
+	public:
+		TypedElementDataStorage::RowHandle FindIndexedRow(
+			EGlobalLockScope LockScope,
+			TypedElementDataStorage::IndexHash Index) const;
+		void IndexRow(
+			EGlobalLockScope LockScope,
+			TypedElementDataStorage::IndexHash Index, 
+			TypedElementDataStorage::RowHandle Row);
+		void BatchIndexRows(
+			EGlobalLockScope LockScope,
+			TConstArrayView<TPair<TypedElementDataStorage::IndexHash, 
+			TypedElementDataStorage::RowHandle>> IndexRowPairs);
+		void ReindexRow(
+			EGlobalLockScope LockScope,
+			TypedElementDataStorage::IndexHash OriginalIndex,
+			TypedElementDataStorage::IndexHash NewIndex,
+			TypedElementDataStorage::RowHandle Row);
+		void RemoveIndex(
+			EGlobalLockScope LockScope,
+			TypedElementDataStorage::IndexHash Index);
+		void RemoveRow(
+			EGlobalLockScope LockScope,
+			TypedElementDataStorage::RowHandle Row);
+
+	private:
+		TMap<TypedElementDataStorage::IndexHash, TypedElementDataStorage::RowHandle> IndexLookupMap;
+		TMultiMap<TypedElementDataStorage::RowHandle, TypedElementDataStorage::IndexHash> ReverseIndexLookupMap;
 	
-	void IndexRowUnguarded(TypedElementDataStorage::IndexHash Index, TypedElementDataStorage::RowHandle Row);
-	void RemoveIndexUnguarded(TypedElementDataStorage::IndexHash Index);
-};
+		void IndexRowUnguarded(TypedElementDataStorage::IndexHash Index, TypedElementDataStorage::RowHandle Row);
+		void RemoveIndexUnguarded(TypedElementDataStorage::IndexHash Index);
+	};
+} // namespace UE::Editor::DataStorage
diff --git a/Engine/Plugins/Experimental/EditorDataStorage/Source/TedsCore/Private/TypedElementDatabaseScratchBuffer.cpp b/Engine/Plugins/Experimental/EditorDataStorage/Source/TedsCore/Private/TypedElementDatabaseScratchBuffer.cpp
--- a/Engine/Plugins/Experimental/EditorDataStorage/Source/TedsCore/Private/TypedElementDatabaseScratchBuffer.cpp
+++ b/Engine/Plugins/Experimental/EditorDataStorage/Source/TedsCore/Private/TypedElementDatabaseScratchBuffer.cpp
@@ -2,22 +2,25 @@
 
 #include "TypedElementDatabaseScratchBuffer.h"
 
-FTypedElementDatabaseScratchBuffer::FTypedElementDatabaseScratchBuffer()
-	: CurrentAllocator(&Allocators[0])
-	, PreviousAllocator(&Allocators[1])
-	, LeastRecentAllocator(&Allocators[2])
+namespace UE::Editor::DataStorage
 {
-}
+	FScratchBuffer::FScratchBuffer()
+		: CurrentAllocator(&Allocators[0])
+		, PreviousAllocator(&Allocators[1])
+		, LeastRecentAllocator(&Allocators[2])
+	{
+	}
 
-void* FTypedElementDatabaseScratchBuffer::Allocate(size_t Size, size_t Alignment)
-{
-	return CurrentAllocator.load()->Malloc(Size, Alignment);
-}
+	void* FScratchBuffer::Allocate(size_t Size, size_t Alignment)
+	{
+		return CurrentAllocator.load()->Malloc(Size, Alignment);
+	}
 
-void FTypedElementDatabaseScratchBuffer::BatchDelete()
-{
-	MemoryAllocator* Current = CurrentAllocator.exchange(LeastRecentAllocator);
-	LeastRecentAllocator = PreviousAllocator;
-	PreviousAllocator = Current;
-	LeastRecentAllocator->BulkDelete();
-}
+	void FScratchBuffer::BatchDelete()
+	{
+		MemoryAllocator* Current = CurrentAllocator.exchange(LeastRecentAllocator);
+		LeastRecentAllocator = PreviousAllocator;
+		PreviousAllocator = Current;
+		LeastRecentAllocator->BulkDelete();
+	}
+} // namespace UE::Editor::DataStorage
diff --git a/Engine/Plugins/Experimental/EditorDataStorage/Source/TedsCore/Private/TypedElementDatabaseScratchBuffer.h b/Engine/Plugins/Experimental/EditorDataStorage/Source/TedsCore/Private/TypedElementDatabaseScratchBuffer.h
--- a/Engine/Plugins/Experimental/EditorDataStorage/Source/TedsCore/Private/TypedElementDatabaseScratchBuffer.h
+++ b/Engine/Plugins/Experimental/EditorDataStorage/Source/TedsCore/Private/TypedElementDatabaseScratchBuffer.h
@@ -7,61 +7,64 @@
 
 // A thread-safe memory allocator that uses linear allocation. This provides a fast and lightweight way to allocate temporary memory 
 // for intermediate values that will live at best until the end of the frame.
-class FTypedElementDatabaseScratchBuffer
+namespace UE::Editor::DataStorage
 {
-private:
-	struct FScratchBufferAllocationTag
+	class FScratchBuffer
 	{
-		static constexpr uint32 BlockSize = 64 * 1024;		// Blocksize used to allocate from
-		static constexpr bool AllowOversizedBlocks = true;  // The allocator supports oversized Blocks and will store them in a separate Block with counter 1
-		static constexpr bool RequiresAccurateSize = true;  // GetAllocationSize returning the accurate size of the allocation otherwise it could be relaxed to return the size to the end of the Block
-		static constexpr bool InlineBlockAllocation = false;  // Inline or Noinline the BlockAllocation which can have an impact on Performance
-		static constexpr const char* TagName = "TEDS_ScratchBuffer";
-
-		using Allocator = TBlockAllocationLockFreeCache<BlockSize, FOsAllocator>;
+	private:
+		struct FAllocationTag
+		{
+			static constexpr uint32 BlockSize = 64 * 1024;		// Blocksize used to allocate from
+			static constexpr bool AllowOversizedBlocks = true;  // The allocator supports oversized Blocks and will store them in a separate Block with counter 1
+			static constexpr bool RequiresAccurateSize = true;  // GetAllocationSize returning the accurate size of the allocation otherwise it could be relaxed to return the size to the end of the Block
+			static constexpr bool InlineBlockAllocation = false;  // Inline or Noinline the BlockAllocation which can have an impact on Performance
+			static constexpr const char* TagName = "TedsScratchBuffer";
+	
+			using Allocator = TBlockAllocationLockFreeCache<BlockSize, FOsAllocator>;
+		};
+		using MemoryAllocator = TConcurrentLinearBulkObjectAllocator<FAllocationTag>;
+	
+	public:
+		FScratchBuffer();
+	
+		void* Allocate(size_t Size, size_t Alignment);
+		template<typename T, typename... ArgTypes>
+		T* Emplace(ArgTypes... Args);
+		template<typename T, typename... ArgTypes>
+		T* EmplaceArray(int32 Count, const ArgTypes&... Args);
+	
+		// Activates a new allocator and deletes all commands and objects in the least recently touched scratch buffer.
+		void BatchDelete();
+	constexpr static int32 MaxAllocationSize();
+	
+	private:
+		// Using a triple buffered approach because the direct API in TEDS (those calls that can be made directly to the API and don't go
+		// through a context) are not required to be atomic. As such it's possible that data for a command is stored in allocator A while
+		// the command is in allocator B if those calls are issued while TEDS is closing it's processing cycle. With double buffering this
+		// would result in allocator A being flushed thus clearing out the data for the command. Using a triple buffered approach will
+		// cause the clearing to be delayed by a frame, avoid this problem. This however does assume that all data and command issuing happens
+		// within a single tick, though for the direct API this should always be true.
+	
+		MemoryAllocator Allocators[3];
+		std::atomic<MemoryAllocator*> CurrentAllocator;
+		MemoryAllocator* PreviousAllocator;
+		MemoryAllocator* LeastRecentAllocator;
 	};
-	using MemoryAllocator = TConcurrentLinearBulkObjectAllocator<FScratchBufferAllocationTag>;
-
-public:
-	FTypedElementDatabaseScratchBuffer();
-
-	void* Allocate(size_t Size, size_t Alignment);
+	
 	template<typename T, typename... ArgTypes>
-	T* Emplace(ArgTypes... Args);
+	T* FScratchBuffer::Emplace(ArgTypes... Args)
+	{
+		return CurrentAllocator.load()->Create<T>(Forward<ArgTypes>(Args)...);
+	}
+	
 	template<typename T, typename... ArgTypes>
-	T* EmplaceArray(int32 Count, const ArgTypes&... Args);
-
-	// Activates a new allocator and deletes all commands and objects in the least recently touched scratch buffer.
-	void BatchDelete();
-constexpr static int32 MaxAllocationSize();
-
-private:
-	// Using a triple buffered approach because the direct API in TEDS (those calls that can be made directly to the API and don't go
-	// through a context) are not required to be atomic. As such it's possible that data for a command is stored in allocator A while
-	// the command is in allocator B if those calls are issued while TEDS is closing it's processing cycle. With double buffering this
-	// would result in allocator A being flushed thus clearing out the data for the command. Using a triple buffered approach will
-	// cause the clearing to be delayed by a frame, avoid this problem. This however does assume that all data and command issuing happens
-	// within a single tick, though for the direct API this should always be true.
-
-	MemoryAllocator Allocators[3];
-	std::atomic<MemoryAllocator*> CurrentAllocator;
-	MemoryAllocator* PreviousAllocator;
-	MemoryAllocator* LeastRecentAllocator;
-};
-
-template<typename T, typename... ArgTypes>
-T* FTypedElementDatabaseScratchBuffer::Emplace(ArgTypes... Args)
-{
-	return CurrentAllocator.load()->Create<T>(Forward<ArgTypes>(Args)...);
-}
-
-template<typename T, typename... ArgTypes>
-T* FTypedElementDatabaseScratchBuffer::EmplaceArray(int32 Count, const ArgTypes&... Args)
-{
-	return CurrentAllocator.load()->CreateArray<T>(Count, Args...);
-}
-
-constexpr int32 FTypedElementDatabaseScratchBuffer::MaxAllocationSize()
-{
-	return FScratchBufferAllocationTag::BlockSize;
-}
+	T* FScratchBuffer::EmplaceArray(int32 Count, const ArgTypes&... Args)
+	{
+		return CurrentAllocator.load()->CreateArray<T>(Count, Args...);
+	}
+	
+	constexpr int32 FScratchBuffer::MaxAllocationSize()
+	{
+		return FAllocationTag::BlockSize;
+	}
+} // namespace UE::Editor::DataStorage
diff --git a/Engine/Plugins/Experimental/EditorDataStorage/Source/TedsCore/Private/TypedElementHandleStore.h b/Engine/Plugins/Experimental/EditorDataStorage/Source/TedsCore/Private/TypedElementHandleStore.h
--- a/Engine/Plugins/Experimental/EditorDataStorage/Source/TedsCore/Private/TypedElementHandleStore.h
+++ b/Engine/Plugins/Experimental/EditorDataStorage/Source/TedsCore/Private/TypedElementHandleStore.h
@@ -8,222 +8,225 @@
 #include "Templates/Function.h"
 #include "Templates/MemoryOps.h"
 
-template<typename DataType, uint32 ReservationSize = 64>
-class TTypedElementHandleStore
+namespace UE::Editor::DataStorage
 {
-public:
-	struct FGeneration
-	{
-		uint32 Generation : 31;
-		uint32 bIsAlive : 1;
-	};
-	
-	class Handle
+	template<typename DataType, uint32 ReservationSize = 64>
+	class THandleStore
 	{
 	public:
-		Handle() = default;
-		explicit Handle(uint64 InPackedData)
-			: PackedData(InPackedData)
-		{};
-		
-		Handle(uint32 Index, uint32 Generation)
+		struct FGeneration
 		{
-			PackedData = (static_cast<uint64_t>(Index) << 32) | Generation;
-		}
+			uint32 Generation : 31;
+			uint32 bIsAlive : 1;
+		};
 
-		uint64 Packed() const
+		class Handle
 		{
-			return PackedData;
-		}
-		
-		uint32_t Generation() const
-		{
-			return static_cast<uint32_t>(PackedData); // Unpack
-		}
+		public:
+			Handle() = default;
+			explicit Handle(uint64 InPackedData)
+				: PackedData(InPackedData)
+			{};
 
-		uint32_t Index() const
-		{
-			return static_cast<uint32_t>(PackedData >> 32);  // Unpack
-		}
+			Handle(uint32 Index, uint32 Generation)
+			{
+				PackedData = (static_cast<uint64_t>(Index) << 32) | Generation;
+			}
 
-		friend auto operator<=>(const Handle&, const Handle&) = default;
-	private:
-		uint64 PackedData = TNumericLimits<uint64>::Max();
-	};
+			uint64 Packed() const
+			{
+				return PackedData;
+			}
 
-	using ListAliveEntriesConstCallback = TFunctionRef<void(Handle, const DataType&)>;
-	using ListAliveEntriesCallback = TFunctionRef<void(Handle, DataType&)>;
-	
-	TTypedElementHandleStore() = default;
-	~TTypedElementHandleStore();
-	TTypedElementHandleStore(const TTypedElementHandleStore&) = delete;
-	TTypedElementHandleStore(TTypedElementHandleStore&&) = delete;
-	TTypedElementHandleStore& operator=(const TTypedElementHandleStore&) = delete;
-	TTypedElementHandleStore& operator=(TTypedElementHandleStore&&) = delete;
-	
-	template<typename... Args>
-	Handle Emplace(Args... Arguments);
-	
-	DataType& Get(Handle Entry);
-	DataType& GetMutable(Handle Entry);
-	const DataType& Get(Handle Entry) const;
+			uint32_t Generation() const
+			{
+				return static_cast<uint32_t>(PackedData); // Unpack
+			}
 
-	/** Removes the entry at the given handle if alive. */
-	void Remove(Handle Entry);
+			uint32_t Index() const
+			{
+				return static_cast<uint32_t>(PackedData >> 32);  // Unpack
+			}
 
-	bool IsAlive(Handle Entry) const;
-	void ListAliveEntries(const ListAliveEntriesConstCallback& Callback) const;
-	void ListAliveEntries(const ListAliveEntriesConstCallback& Callback);
-	void ListAliveEntries(const ListAliveEntriesCallback& Callback);
+			friend auto operator<=>(const Handle&, const Handle&) = default;
+		private:
+			uint64 PackedData = TNumericLimits<uint64>::Max();
+		};
 
-private:
-	TArray<DataType> Data;
-	TArray<FGeneration> Generations;
-	TDeque<uint32> RecycleBin; // Using a deque to avoid the same slot being constantly reused.
-};
+		using ListAliveEntriesConstCallback = TFunctionRef<void(Handle, const DataType&)>;
+		using ListAliveEntriesCallback = TFunctionRef<void(Handle, DataType&)>;
 
+		THandleStore() = default;
+		~THandleStore();
+		THandleStore(const THandleStore&) = delete;
+		THandleStore(THandleStore&&) = delete;
+		THandleStore& operator=(const THandleStore&) = delete;
+		THandleStore& operator=(THandleStore&&) = delete;
 
+		template<typename... Args>
+		Handle Emplace(Args... Arguments);
 
-// Implementations
-template<typename DataType, uint32 ReservationSize>
-TTypedElementHandleStore<DataType, ReservationSize>::~TTypedElementHandleStore()
-{
-	if constexpr (std::is_destructible_v<DataType> && !std::is_trivially_destructible_v<DataType>)
+		DataType& Get(Handle Entry);
+		DataType& GetMutable(Handle Entry);
+		const DataType& Get(Handle Entry) const;
+
+		/** Removes the entry at the given handle if alive. */
+		void Remove(Handle Entry);
+
+		bool IsAlive(Handle Entry) const;
+		void ListAliveEntries(const ListAliveEntriesConstCallback& Callback) const;
+		void ListAliveEntries(const ListAliveEntriesConstCallback& Callback);
+		void ListAliveEntries(const ListAliveEntriesCallback& Callback);
+
+	private:
+		TArray<DataType> Data;
+		TArray<FGeneration> Generations;
+		TDeque<uint32> RecycleBin; // Using a deque to avoid the same slot being constantly reused.
+	};
+
+
+
+	// Implementations
+	template<typename DataType, uint32 ReservationSize>
+	THandleStore<DataType, ReservationSize>::~THandleStore()
 	{
-		int32 Count = Data.Num();
-		const DataType* EntryIt = Data.GetData();
-		const FGeneration* GenerationIt = Generations.GetData();
-		
-		for (int32 Index = 0; Index < Count; ++Index)
+		if constexpr (std::is_destructible_v<DataType> && !std::is_trivially_destructible_v<DataType>)
 		{
-			if (GenerationIt->bIsAlive)
+			int32 Count = Data.Num();
+			const DataType* EntryIt = Data.GetData();
+			const FGeneration* GenerationIt = Generations.GetData();
+
+			for (int32 Index = 0; Index < Count; ++Index)
 			{
-				DestructItem(EntryIt);
-			}
+				if (GenerationIt->bIsAlive)
+				{
+					DestructItem(EntryIt);
+				}
 
-			++EntryIt;
-			++GenerationIt;
+				++EntryIt;
+				++GenerationIt;
+			}
 		}
+		Data.SetNumUnsafeInternal(0); // Make sure array destructor does not run destructor on items again
 	}
-	Data.SetNumUnsafeInternal(0); // Make sure array destructor does not run destructor on items again
-}
 
-template<typename DataType, uint32 ReservationSize>
-template<typename... Args>
-auto TTypedElementHandleStore<DataType, ReservationSize>::Emplace(Args... Arguments) -> Handle
-{
-	if (RecycleBin.IsEmpty())
+	template<typename DataType, uint32 ReservationSize>
+	template<typename... Args>
+	auto THandleStore<DataType, ReservationSize>::Emplace(Args... Arguments) -> Handle
 	{
-		uint32 NewSize = Data.Num() + ReservationSize;
-		Data.Reserve(NewSize);
-		Generations.Reserve(NewSize);
-
-		for (uint32 Index = Data.Num(); Index < NewSize; ++Index)
+		if (RecycleBin.IsEmpty())
 		{
-			Data.AddZeroed();
-			Generations.AddZeroed();
-			RecycleBin.EmplaceLast(Index);
+			uint32 NewSize = Data.Num() + ReservationSize;
+			Data.Reserve(NewSize);
+			Generations.Reserve(NewSize);
+
+			for (uint32 Index = Data.Num(); Index < NewSize; ++Index)
+			{
+				Data.AddZeroed();
+				Generations.AddZeroed();
+				RecycleBin.EmplaceLast(Index);
+			}
 		}
-	}
 
-	uint32 Index = RecycleBin.Last();
-	RecycleBin.PopLast();
-	
-	DataType& Entry = Data[Index];
-	new(&Entry) DataType( Forward<Args>(Arguments)... );
+		uint32 Index = RecycleBin.Last();
+		RecycleBin.PopLast();
 
-	Handle Result(Index, Generations[Index].Generation);
-	Generations[Index].bIsAlive = 1;
-	return Result;
-}
+		DataType& Entry = Data[Index];
+		new(&Entry) DataType(Forward<Args>(Arguments)...);
 
-template<typename DataType, uint32 ReservationSize>
-DataType& TTypedElementHandleStore<DataType, ReservationSize>::Get(Handle Entry)
-{
-	checkf(IsAlive(Entry), TEXT("Attempting to retrieve a dead entry from a Typed Element Handle Store."));
-	return Data[Entry.Index()];
-}
+		Handle Result(Index, Generations[Index].Generation);
+		Generations[Index].bIsAlive = 1;
+		return Result;
+	}
 
-template<typename DataType, uint32 ReservationSize>
-DataType& TTypedElementHandleStore<DataType, ReservationSize>::GetMutable(Handle Entry)
-{
-	checkf(IsAlive(Entry), TEXT("Attempting to retrieve a dead entry from a Typed Element Handle Store."));
-	return Data[Entry.Index()];
-}
+	template<typename DataType, uint32 ReservationSize>
+	DataType& THandleStore<DataType, ReservationSize>::Get(Handle Entry)
+	{
+		checkf(IsAlive(Entry), TEXT("Attempting to retrieve a dead entry from a Typed Element Handle Store."));
+		return Data[Entry.Index()];
+	}
 
-template<typename DataType, uint32 ReservationSize>
-const DataType& TTypedElementHandleStore<DataType, ReservationSize>::Get(Handle Entry) const
-{
-	checkf(IsAlive(Entry), TEXT("Attempting to retrieve a dead entry from a Typed Element Handle Store."));
-	return Data[Entry.Index()];
-}
+	template<typename DataType, uint32 ReservationSize>
+	DataType& THandleStore<DataType, ReservationSize>::GetMutable(Handle Entry)
+	{
+		checkf(IsAlive(Entry), TEXT("Attempting to retrieve a dead entry from a Typed Element Handle Store."));
+		return Data[Entry.Index()];
+	}
 
-template<typename DataType, uint32 ReservationSize>
-void TTypedElementHandleStore<DataType, ReservationSize>::Remove(Handle Entry)
-{
-	if (IsAlive(Entry))
+	template<typename DataType, uint32 ReservationSize>
+	const DataType& THandleStore<DataType, ReservationSize>::Get(Handle Entry) const
 	{
-		FGeneration& Generation = Generations[Entry.Index()];
-		++Generation.Generation;
-		Generation.bIsAlive = 0;
-		if constexpr (std::is_destructible_v<DataType> && !std::is_trivially_destructible_v<DataType>)
+		checkf(IsAlive(Entry), TEXT("Attempting to retrieve a dead entry from a Typed Element Handle Store."));
+		return Data[Entry.Index()];
+	}
+
+	template<typename DataType, uint32 ReservationSize>
+	void THandleStore<DataType, ReservationSize>::Remove(Handle Entry)
+	{
+		if (IsAlive(Entry))
 		{
-			DestructItem(&Data[Entry.Index()]);
+			FGeneration& Generation = Generations[Entry.Index()];
+			++Generation.Generation;
+			Generation.bIsAlive = 0;
+			if constexpr (std::is_destructible_v<DataType> && !std::is_trivially_destructible_v<DataType>)
+			{
+				DestructItem(&Data[Entry.Index()]);
+			}
+			RecycleBin.EmplaceLast(Entry.Index());
 		}
-		RecycleBin.EmplaceLast(Entry.Index());
 	}
-}
 
-template<typename DataType, uint32 ReservationSize>
-bool TTypedElementHandleStore<DataType, ReservationSize>::IsAlive(Handle Entry) const
-{
-	return
-		Entry.Index() < static_cast<uint32>(Generations.Num()) &&
-		Generations[Entry.Index()].Generation == Entry.Generation();
-}
+	template<typename DataType, uint32 ReservationSize>
+	bool THandleStore<DataType, ReservationSize>::IsAlive(Handle Entry) const
+	{
+		return
+			Entry.Index() < static_cast<uint32>(Generations.Num()) &&
+			Generations[Entry.Index()].Generation == Entry.Generation();
+	}
 
-template<typename DataType, uint32 ReservationSize>
-void TTypedElementHandleStore<DataType, ReservationSize>::ListAliveEntries(const ListAliveEntriesConstCallback& Callback) const
-{
-	int32 Count = Data.Num();
-	const DataType* EntryIt = Data.GetData();
-	const FGeneration* GenerationIt = Generations.GetData();
-	
-	for (int32 Index = 0; Index < Count; ++Index)
+	template<typename DataType, uint32 ReservationSize>
+	void THandleStore<DataType, ReservationSize>::ListAliveEntries(const ListAliveEntriesConstCallback& Callback) const
 	{
-		if (GenerationIt->bIsAlive)
+		int32 Count = Data.Num();
+		const DataType* EntryIt = Data.GetData();
+		const FGeneration* GenerationIt = Generations.GetData();
+
+		for (int32 Index = 0; Index < Count; ++Index)
 		{
-			Handle DataHandle(Index, GenerationIt->Generation);
-			Callback(DataHandle, *EntryIt);
-		}
+			if (GenerationIt->bIsAlive)
+			{
+				Handle DataHandle(Index, GenerationIt->Generation);
+				Callback(DataHandle, *EntryIt);
+			}
 
-		++EntryIt;
-		++GenerationIt;
+			++EntryIt;
+			++GenerationIt;
+		}
 	}
-}
-
-template<typename DataType, uint32 ReservationSize>
-void TTypedElementHandleStore<DataType, ReservationSize>::ListAliveEntries(const ListAliveEntriesConstCallback& Callback)
-{
-	const_cast<const TTypedElementHandleStore<DataType, ReservationSize>*>(this)->ListAliveEntries(Callback);
-}
 
-template<typename DataType, uint32 ReservationSize>
-void TTypedElementHandleStore<DataType, ReservationSize>::ListAliveEntries(const ListAliveEntriesCallback& Callback)
-{
-	int32 Count = Data.Num();
-	DataType* EntryIt = Data.GetData();
-	FGeneration* GenerationIt = Generations.GetData();
+	template<typename DataType, uint32 ReservationSize>
+	void THandleStore<DataType, ReservationSize>::ListAliveEntries(const ListAliveEntriesConstCallback& Callback)
+	{
+		const_cast<const THandleStore<DataType, ReservationSize>*>(this)->ListAliveEntries(Callback);
+	}
 
-	for (int32 Index = 0; Index < Count; ++Index)
+	template<typename DataType, uint32 ReservationSize>
+	void THandleStore<DataType, ReservationSize>::ListAliveEntries(const ListAliveEntriesCallback& Callback)
 	{
-		if (GenerationIt->bIsAlive)
+		int32 Count = Data.Num();
+		DataType* EntryIt = Data.GetData();
+		FGeneration* GenerationIt = Generations.GetData();
+
+		for (int32 Index = 0; Index < Count; ++Index)
 		{
-			Handle DataHandle(Index, GenerationIt->Generation);
-			Callback(DataHandle, *EntryIt);
-		}
+			if (GenerationIt->bIsAlive)
+			{
+				Handle DataHandle(Index, GenerationIt->Generation);
+				Callback(DataHandle, *EntryIt);
+			}
 
-		++EntryIt;
-		++GenerationIt;
+			++EntryIt;
+			++GenerationIt;
+		}
 	}
-}
+} // namespace UE::Editor::DataStorage
diff --git a/Engine/Plugins/Experimental/EditorDataStorage/Source/TedsCore/Tests/Commands/EditorDataStorageCommandBufferTests.cpp b/Engine/Plugins/Experimental/EditorDataStorage/Source/TedsCore/Tests/Commands/EditorDataStorageCommandBufferTests.cpp
--- a/Engine/Plugins/Experimental/EditorDataStorage/Source/TedsCore/Tests/Commands/EditorDataStorageCommandBufferTests.cpp
+++ b/Engine/Plugins/Experimental/EditorDataStorage/Source/TedsCore/Tests/Commands/EditorDataStorageCommandBufferTests.cpp
@@ -15,7 +15,7 @@ enum class ETestCommandType
 	Nop, A, B, C
 };
 
-using TestCommandBuffer = UE::EditorDataStorage::FCommandBuffer<FTestCommandA, FTestCommandB, FTestCommandC>;
+using TestCommandBuffer = UE::Editor::DataStorage::FCommandBuffer<FTestCommandA, FTestCommandB, FTestCommandC>;
 
 constexpr const TCHAR* TypeToString(ETestCommandType Type)
 {
@@ -43,7 +43,7 @@ struct FSequenceTestingProcessor
 	TArray<ETestCommandType> CommandTypes;
 	int Index = 0;
 
-	void operator()(const UE::EditorDataStorage::FNopCommand&)
+	void operator()(const UE::Editor::DataStorage::FNopCommand&)
 	{
 		REQUIRE_MESSAGE(TEXT("An additional NOP command was issued when no more commands were expected."), Index < CommandTypes.Num());
 		if (CommandTypes[Index] != ETestCommandType::Nop)
@@ -83,9 +83,9 @@ struct FSequenceTestingProcessor
 
 TEST_CASE_NAMED(EditorDataStorage_FCommandBuffer, "EditorDataStorage::Command Buffer", "[ApplicationContextMask][EngineFilter]")
 {
-	using namespace UE::EditorDataStorage;
+	using namespace UE::Editor::DataStorage;
 	
-	FTypedElementDatabaseScratchBuffer ScratchBuffer;
+	FScratchBuffer ScratchBuffer;
 
 	SECTION("Initialize")
 	{
@@ -131,9 +131,9 @@ TEST_CASE_NAMED(EditorDataStorage_FCommandBuffer, "EditorDataStorage::Command Bu
 
 TEST_CASE_NAMED(EditorDataStorage_FCommandBuffer_FCollection, "EditorDataStorage::Command Buffer::FCollection", "[ApplicationContextMask][EngineFilter]")
 {
-	using namespace UE::EditorDataStorage;
+	using namespace UE::Editor::DataStorage;
 
-	FTypedElementDatabaseScratchBuffer ScratchBuffer;
+	FScratchBuffer ScratchBuffer;
 	TestCommandBuffer CommandBuffer;
 	CommandBuffer.Initialize(ScratchBuffer);
 
@@ -403,9 +403,9 @@ TEST_CASE_NAMED(EditorDataStorage_FCommandBuffer_FCollection, "EditorDataStorage
 
 TEST_CASE_NAMED(EditorDataStorage_FCommandBuffer_FOptimizer, "EditorDataStorage::Command Buffer::FOptimizer", "[ApplicationContextMask][EngineFilter]")
 {
-	using namespace UE::EditorDataStorage;
+	using namespace UE::Editor::DataStorage;
 
-	FTypedElementDatabaseScratchBuffer ScratchBuffer;
+	FScratchBuffer ScratchBuffer;
 	TestCommandBuffer CommandBuffer;
 	CommandBuffer.Initialize(ScratchBuffer);
 
diff --git a/Engine/Plugins/Experimental/EditorDataStorage/Source/TedsCore/Tests/TypedElementDatabaseScratchBufferTests.cpp b/Engine/Plugins/Experimental/EditorDataStorage/Source/TedsCore/Tests/TypedElementDatabaseScratchBufferTests.cpp
--- a/Engine/Plugins/Experimental/EditorDataStorage/Source/TedsCore/Tests/TypedElementDatabaseScratchBufferTests.cpp
+++ b/Engine/Plugins/Experimental/EditorDataStorage/Source/TedsCore/Tests/TypedElementDatabaseScratchBufferTests.cpp
@@ -24,14 +24,15 @@ static void WaitForFinalization(std::condition_variable& WaitVariable, std::mute
 	REQUIRE(CompletedThreadCount == ThreadCount);
 }
 
-TEST_CASE("TypedElementsDataStorage::Scratch Buffer - MT (FTypedElementDatabaseScratchBuffer)", "[ApplicationContextMask][EngineFilter]")
+TEST_CASE("TEDS::Scratch Buffer - MT (FScratchBuffer)", "[ApplicationContextMask][EngineFilter]")
 {
+	using namespace UE::Editor::DataStorage;
 	SECTION("Stress test without recycling.")
 	{
 		using namespace std::chrono_literals;
 		constexpr uint32 ThreadCount = 8;
 
-		FTypedElementDatabaseScratchBuffer Buffer;
+		FScratchBuffer Buffer;
 
 		std::atomic<bool> bKeepRunning = true;
 		std::atomic<uint32> FailedAllocations = 0;
@@ -76,7 +77,7 @@ TEST_CASE("TypedElementsDataStorage::Scratch Buffer - MT (FTypedElementDatabaseS
 
 		constexpr uint32 ThreadCount = 8;
 
-		FTypedElementDatabaseScratchBuffer Buffer;
+		FScratchBuffer Buffer;
 
 		std::atomic<bool> bKeepRunning = true;
 		std::atomic<uint32> FailedAllocations = 0;
@@ -137,30 +138,31 @@ TEST_CASE("TypedElementsDataStorage::Scratch Buffer - MT (FTypedElementDatabaseS
 	}
 }
 
-TEST_CASE("TypedElementsDataStorage::Scratch Buffer (FTypedElementDatabaseScratchBuffer)", "[ApplicationContextMask][EngineFilter]")
+TEST_CASE("TEDS::Scratch Buffer (FScratchBuffer)", "[ApplicationContextMask][EngineFilter]")
 {
+	using namespace UE::Editor::DataStorage;
 	SECTION("Create and destroy buffer")
 	{
-		FTypedElementDatabaseScratchBuffer Buffer;
+		FScratchBuffer Buffer;
 	}
 
 	SECTION("Allocate small block")
 	{
-		FTypedElementDatabaseScratchBuffer Buffer;
+		FScratchBuffer Buffer;
 		void* Data = Buffer.Allocate(4, 4);
 		CHECK(Data != nullptr);
 	}
 
 	SECTION("Allocate over-sized block")
 	{
-		FTypedElementDatabaseScratchBuffer Buffer;
+		FScratchBuffer Buffer;
 		void* Data = Buffer.Allocate(Buffer.MaxAllocationSize() * 4, 4);
 		CHECK(Data != nullptr);
 	}
 
 	SECTION("Alignment respected")
 	{
-		FTypedElementDatabaseScratchBuffer Buffer;
+		FScratchBuffer Buffer;
 		void* Data1 = Buffer.Allocate(1, 1);
 		void* Data2 = Buffer.Allocate(4, 4);
 
@@ -172,7 +174,7 @@ TEST_CASE("TypedElementsDataStorage::Scratch Buffer (FTypedElementDatabaseScratc
 
 	SECTION("Multiple blocks used.")
 	{
-		FTypedElementDatabaseScratchBuffer Buffer;
+		FScratchBuffer Buffer;
 
 		int32 IncrementCount = (Buffer.MaxAllocationSize() * 4 /* Fill 4 blocks*/) / 64 /* With 64 byte allocations */;
 		for (int32 Counter = 0; Counter < IncrementCount; ++Counter)
@@ -184,7 +186,7 @@ TEST_CASE("TypedElementsDataStorage::Scratch Buffer (FTypedElementDatabaseScratc
 
 	SECTION("Recycle full blocks.")
 	{
-		FTypedElementDatabaseScratchBuffer Buffer;
+		FScratchBuffer Buffer;
 
 		for (int32 Iterations = 0; Iterations < 16; ++Iterations)
 		{
diff --git a/Engine/Source/Runtime/TypedElementFramework/Private/Elements/Framework/TypedElementQueryBuilder.cpp b/Engine/Source/Runtime/TypedElementFramework/Private/Elements/Framework/TypedElementQueryBuilder.cpp
--- a/Engine/Source/Runtime/TypedElementFramework/Private/Elements/Framework/TypedElementQueryBuilder.cpp
+++ b/Engine/Source/Runtime/TypedElementFramework/Private/Elements/Framework/TypedElementQueryBuilder.cpp
@@ -139,7 +139,7 @@ namespace TypedElementQueryBuilder
 		return *this;
 	}
 
-	FSimpleQuery& FSimpleQuery::All(const UE::EditorDataStorage::FDynamicTag& Tag, const FName& Value)
+	FSimpleQuery& FSimpleQuery::All(const UE::Editor::DataStorage::FDynamicTag& Tag, const FName& Value)
 	{
 		Query->DynamicTags.Emplace(
 			TypedElementDataStorage::FQueryDescription::FDynamicTagData
@@ -152,7 +152,7 @@ namespace TypedElementQueryBuilder
 
 	FSimpleQuery& FSimpleQuery::All(const UEnum& Enum)
 	{
-		using namespace UE::EditorDataStorage;
+		using namespace UE::Editor::DataStorage;
 		return All(FDynamicTag(Enum.GetFName()));
 	}
 
@@ -164,23 +164,23 @@ namespace TypedElementQueryBuilder
 			UE_LOG(LogTypedElementDataStorage, Warning, TEXT("Invalid value '%lld' for enum '%s'"), Value, *Enum.GetName());
 			return *this;
 		}
-		using namespace UE::EditorDataStorage;
+		using namespace UE::Editor::DataStorage;
 		return All(FDynamicTag(Enum.GetFName()), ValueName);
 	}
 
 	template <>
-	FSimpleQuery& FSimpleQuery::All<UE::EditorDataStorage::FDynamicTag>(const FName& Tag)
+	FSimpleQuery& FSimpleQuery::All<UE::Editor::DataStorage::FDynamicTag>(const FName& Tag)
 	{
-		return All(UE::EditorDataStorage::FDynamicTag(Tag));
+		return All(UE::Editor::DataStorage::FDynamicTag(Tag));
 	}
 
 	template <>
-	FSimpleQuery& FSimpleQuery::All<UE::EditorDataStorage::FDynamicTag>(const FName& Tag, const FName& Value)
+	FSimpleQuery& FSimpleQuery::All<UE::Editor::DataStorage::FDynamicTag>(const FName& Tag, const FName& Value)
 	{
-		return All(UE::EditorDataStorage::FDynamicTag(Tag), Value);
+		return All(UE::Editor::DataStorage::FDynamicTag(Tag), Value);
 	}
 
-	FSimpleQuery& FSimpleQuery::All(const UE::EditorDataStorage::FDynamicTag& Tag)
+	FSimpleQuery& FSimpleQuery::All(const UE::Editor::DataStorage::FDynamicTag& Tag)
 	{
 		return All(Tag, NAME_None);
 	}
diff --git a/Engine/Source/Runtime/TypedElementFramework/Public/Elements/Common/TypedElementCommonTypes.h b/Engine/Source/Runtime/TypedElementFramework/Public/Elements/Common/TypedElementCommonTypes.h
--- a/Engine/Source/Runtime/TypedElementFramework/Public/Elements/Common/TypedElementCommonTypes.h
+++ b/Engine/Source/Runtime/TypedElementFramework/Public/Elements/Common/TypedElementCommonTypes.h
@@ -45,7 +45,7 @@ namespace UE
 
 namespace UE
 {
-	namespace EditorDataStorage
+	namespace Editor::DataStorage
 	{
 		/**
 		 * Defines a dynamic type for a dynamic tag
diff --git a/Engine/Source/Runtime/TypedElementFramework/Public/Elements/Common/TypedElementQueryDescription.h b/Engine/Source/Runtime/TypedElementFramework/Public/Elements/Common/TypedElementQueryDescription.h
--- a/Engine/Source/Runtime/TypedElementFramework/Public/Elements/Common/TypedElementQueryDescription.h
+++ b/Engine/Source/Runtime/TypedElementFramework/Public/Elements/Common/TypedElementQueryDescription.h
@@ -59,7 +59,7 @@ namespace TypedElementDataStorage
 		struct FDynamicTagData
 		{
 			// The Tag maps to a Mass ConstSharedFragment object
-			UE::EditorDataStorage::FDynamicTag Tag;
+			UE::Editor::DataStorage::FDynamicTag Tag;
 
 			// The MatchValue specifies the value that the fragment must have to be matched
 			// If MatchValue is NAME_None, then TEDS will match all values
diff --git a/Engine/Source/Runtime/TypedElementFramework/Public/Elements/Framework/TypedElementQueryBuilder.h b/Engine/Source/Runtime/TypedElementFramework/Public/Elements/Framework/TypedElementQueryBuilder.h
--- a/Engine/Source/Runtime/TypedElementFramework/Public/Elements/Framework/TypedElementQueryBuilder.h
+++ b/Engine/Source/Runtime/TypedElementFramework/Public/Elements/Framework/TypedElementQueryBuilder.h
@@ -169,10 +169,10 @@ namespace TypedElementQueryBuilder
 		// Dynamic Tags
 		// ============
 		// Adds a filter to the query which must match a given DynamicTag.  The value of the DynamicTag will not be checked.
-		TYPEDELEMENTFRAMEWORK_API FSimpleQuery& All(const UE::EditorDataStorage::FDynamicTag& Tag);
+		TYPEDELEMENTFRAMEWORK_API FSimpleQuery& All(const UE::Editor::DataStorage::FDynamicTag& Tag);
 		// Adds a filter to the query which must match a given DynamicTag.  The value of the DynamicTag must also match.
 		// Note: The query can only match a single value.  Multiple value queries are not supported at this time.
-		TYPEDELEMENTFRAMEWORK_API FSimpleQuery& All(const UE::EditorDataStorage::FDynamicTag& Tag, const FName& Value);
+		TYPEDELEMENTFRAMEWORK_API FSimpleQuery& All(const UE::Editor::DataStorage::FDynamicTag& Tag, const FName& Value);
 		
 		TYPEDELEMENTFRAMEWORK_API FSimpleQuery& All(const UEnum& Enum);
 		TYPEDELEMENTFRAMEWORK_API FSimpleQuery& All(const UEnum& Enum, int64 Value);
@@ -186,11 +186,11 @@ namespace TypedElementQueryBuilder
 
 		// Adds a filter to the query which must match a given DynamicTag.  The value of the DynamicTag will not be checked.
 		template<>
-		TYPEDELEMENTFRAMEWORK_API FSimpleQuery& All<UE::EditorDataStorage::FDynamicTag>(const FName& Tag);
+		TYPEDELEMENTFRAMEWORK_API FSimpleQuery& All<UE::Editor::DataStorage::FDynamicTag>(const FName& Tag);
 		// Adds a filter to the query which must match a given DynamicTag.  The value of the DynamicTag must also match.
 		// Note: The query can only match a single value.  Multiple value queries are not supported at this time.
 		template<>
-		TYPEDELEMENTFRAMEWORK_API FSimpleQuery& All<UE::EditorDataStorage::FDynamicTag>(const FName& Tag, const FName& Value);
+		TYPEDELEMENTFRAMEWORK_API FSimpleQuery& All<UE::Editor::DataStorage::FDynamicTag>(const FName& Tag, const FName& Value);
 		
 		template<TypedElementDataStorage::TEnumType EnumT>
 		FSimpleQuery& All();
diff --git a/Engine/Source/Runtime/TypedElementFramework/Public/Elements/Interfaces/TypedElementDataStorageInterface.h b/Engine/Source/Runtime/TypedElementFramework/Public/Elements/Interfaces/TypedElementDataStorageInterface.h
--- a/Engine/Source/Runtime/TypedElementFramework/Public/Elements/Interfaces/TypedElementDataStorageInterface.h
+++ b/Engine/Source/Runtime/TypedElementFramework/Public/Elements/Interfaces/TypedElementDataStorageInterface.h
@@ -195,13 +195,13 @@ public:
 	 *       is added.  This will likely change in the future to transparently replace the tag to have consistent behaviour with other usages
 	 *       of AddColumn
 	 */
-	virtual void AddColumn(TypedElementDataStorage::RowHandle Row, const UE::EditorDataStorage::FDynamicTag& Tag, const FName& Value) = 0;
+	virtual void AddColumn(TypedElementDataStorage::RowHandle Row, const UE::Editor::DataStorage::FDynamicTag& Tag, const FName& Value) = 0;
 
 	template<typename T>
 	void AddColumn(TypedElementDataStorage::RowHandle Row, const FName& Tag, const FName& Value) = delete;
 	
 	template<>
-	void AddColumn<UE::EditorDataStorage::FDynamicTag>(TypedElementDataStorage::RowHandle Row, const FName& Tag, const FName& Value);
+	void AddColumn<UE::Editor::DataStorage::FDynamicTag>(TypedElementDataStorage::RowHandle Row, const FName& Tag, const FName& Value);
 
 	template<TypedElementDataStorage::TEnumType EnumT>
 	void AddColumn(TypedElementDataStorage::RowHandle Row, EnumT Value);
@@ -229,13 +229,13 @@ public:
 	 * Removes a dynamic tag from the given row
 	 * If tag does not exist on row, operation will do nothing.
 	 */
-	virtual void RemoveColumn(TypedElementDataStorage::RowHandle Row, const UE::EditorDataStorage::FDynamicTag& Tag) = 0;
+	virtual void RemoveColumn(TypedElementDataStorage::RowHandle Row, const UE::Editor::DataStorage::FDynamicTag& Tag) = 0;
 
 	template<typename T>
 	void RemoveColumn(TypedElementDataStorage::RowHandle Row, const FName& Tag) = delete;
 	
 	template<>
-	void RemoveColumn<UE::EditorDataStorage::FDynamicTag>(TypedElementDataStorage::RowHandle Row, const FName& Tag);
+	void RemoveColumn<UE::Editor::DataStorage::FDynamicTag>(TypedElementDataStorage::RowHandle Row, const FName& Tag);
 
 	/**
 	 * Removes multiple columns from a row. This is typically more efficient than adding columns one
@@ -444,15 +444,15 @@ void ITypedElementDataStorageInterface::AddColumns(TypedElementRowHandle Row)
 }
 
 template <>
-inline void ITypedElementDataStorageInterface::AddColumn<UE::EditorDataStorage::FDynamicTag>(TypedElementDataStorage::RowHandle Row, const FName& Tag, const FName& Value)
+inline void ITypedElementDataStorageInterface::AddColumn<UE::Editor::DataStorage::FDynamicTag>(TypedElementDataStorage::RowHandle Row, const FName& Tag, const FName& Value)
 {
-	AddColumn(Row, UE::EditorDataStorage::FDynamicTag(Tag), Value);
+	AddColumn(Row, UE::Editor::DataStorage::FDynamicTag(Tag), Value);
 }
 
 template <>
-inline void ITypedElementDataStorageInterface::RemoveColumn<UE::EditorDataStorage::FDynamicTag>(TypedElementDataStorage::RowHandle Row, const FName& Tag)
+inline void ITypedElementDataStorageInterface::RemoveColumn<UE::Editor::DataStorage::FDynamicTag>(TypedElementDataStorage::RowHandle Row, const FName& Tag)
 {
-	using namespace UE::EditorDataStorage;
+	using namespace UE::Editor::DataStorage;
 	RemoveColumn(Row, FDynamicTag(Tag));
 }
 
@@ -463,7 +463,7 @@ void ITypedElementDataStorageInterface::AddColumn(TypedElementDataStorage::RowHa
 	const FName ValueAsFName = *Enum->GetNameStringByValue(static_cast<int64>(Value));
 	if (ValueAsFName != NAME_None)
 	{
-		AddColumn(Row, UE::EditorDataStorage::FDynamicTag(Enum->GetFName()), ValueAsFName);
+		AddColumn(Row, UE::Editor::DataStorage::FDynamicTag(Enum->GetFName()), ValueAsFName);
 	}
 }
 
@@ -477,7 +477,7 @@ template<TypedElementDataStorage::TEnumType EnumT>
 void ITypedElementDataStorageInterface::RemoveColumn(TypedElementDataStorage::RowHandle Row)
 {
 	const UEnum* Enum = StaticEnum<EnumT>();
-	RemoveColumn(Row, UE::EditorDataStorage::FDynamicTag(Enum->GetFName()));
+	RemoveColumn(Row, UE::Editor::DataStorage::FDynamicTag(Enum->GetFName()));
 }
 
 template<TypedElementDataStorage::TColumnType... Columns>
