[Teds] Rename TypedElementDataStorage log to EditorDataStorage log

#rnx

diff --git a/Engine/Plugins/Experimental/EditorDataStorage/Source/TedsCore/Private/Debug/TypedElementDatabaseDebugCommands.cpp b/Engine/Plugins/Experimental/EditorDataStorage/Source/TedsCore/Private/Debug/TypedElementDatabaseDebugCommands.cpp
--- a/Engine/Plugins/Experimental/EditorDataStorage/Source/TedsCore/Private/Debug/TypedElementDatabaseDebugCommands.cpp
+++ b/Engine/Plugins/Experimental/EditorDataStorage/Source/TedsCore/Private/Debug/TypedElementDatabaseDebugCommands.cpp
@@ -221,7 +221,7 @@ static FAutoConsoleCommand CVarCreateRow(
 
 		const TypedElementDataStorage::RowHandle RowHandle = DataStorage->AddRow(Table);
 
-		UE_LOG(LogTypedElementDataStorage, Warning, TEXT("Added Row %llu"), static_cast<uint64>(RowHandle));
+		UE_LOG(LogEditorDataStorage, Warning, TEXT("Added Row %llu"), static_cast<uint64>(RowHandle));
 	}));
 
 static FAutoConsoleCommand CVarAddDynamicTag(
@@ -341,7 +341,7 @@ static FAutoConsoleCommand CVarMatchDynamicTag(
 			}));
 		DataStorage->UnregisterQuery(Query);
 
-		UE_LOG(LogTypedElementDataStorage, Warning, TEXT("Processed %llu rows"), static_cast<uint64>(Count));
+		UE_LOG(LogEditorDataStorage, Warning, TEXT("Processed %llu rows"), static_cast<uint64>(Count));
 	}),
 	ECVF_Default);
 
@@ -415,7 +415,7 @@ static FAutoConsoleCommand CVarMatchDynamicTagFromEnum(
 
 		if (Args.Num() > 1)
 		{
-			UE_LOG(LogTypedElementDataStorage, Error, TEXT("Invalid number of arguments"));
+			UE_LOG(LogEditorDataStorage, Error, TEXT("Invalid number of arguments"));
 			return;
 		}
 
@@ -464,7 +464,7 @@ static FAutoConsoleCommand CVarMatchDynamicTagFromEnum(
 		}();
 		if (Query == TypedElementDataStorage::InvalidQueryHandle)
 		{
-			UE_LOG(LogTypedElementDataStorage, Error, TEXT("Invalid number of arguments"));
+			UE_LOG(LogEditorDataStorage, Error, TEXT("Invalid number of arguments"));
 			return;
 		}
 		
@@ -477,7 +477,7 @@ static FAutoConsoleCommand CVarMatchDynamicTagFromEnum(
 			}));
 		DataStorage->UnregisterQuery(Query);
 
-		UE_LOG(LogTypedElementDataStorage, Warning, TEXT("Processed %llu rows"), static_cast<uint64>(Count));
+		UE_LOG(LogEditorDataStorage, Warning, TEXT("Processed %llu rows"), static_cast<uint64>(Count));
 		
 	}),
 	ECVF_Default);
\ No newline at end of file
diff --git a/Engine/Plugins/Experimental/EditorDataStorage/Source/TedsCore/Private/Memento/TypedElementMementoSystem.cpp b/Engine/Plugins/Experimental/EditorDataStorage/Source/TedsCore/Private/Memento/TypedElementMementoSystem.cpp
--- a/Engine/Plugins/Experimental/EditorDataStorage/Source/TedsCore/Private/Memento/TypedElementMementoSystem.cpp
+++ b/Engine/Plugins/Experimental/EditorDataStorage/Source/TedsCore/Private/Memento/TypedElementMementoSystem.cpp
@@ -82,7 +82,7 @@ void UTypedElementMementoSystem::CreateMementoInternal(TypedElementDataStorage::
 					ColumnType.CopyScriptStruct(Destination, Source);
 				});
 
-			UE_LOG(LogTypedElementDataStorage, VeryVerbose,
+			UE_LOG(LogEditorDataStorage, VeryVerbose,
 				TEXT("Column->Memento: %llu -> %llu"), SourceRow, MementoRow);
 		}
 	}
@@ -111,7 +111,7 @@ void UTypedElementMementoSystem::RestoreMemento(TypedElementDataStorage::RowHand
 					ColumnType.CopyScriptStruct(Destination, Source);
 				});
 
-			UE_LOG(LogTypedElementDataStorage, VeryVerbose, 
+			UE_LOG(LogEditorDataStorage, VeryVerbose, 
 				TEXT("Memento->Column: %llu -> %llu"), MementoRow, TargetRow);
 		}
 	}
diff --git a/Engine/Plugins/Experimental/EditorDataStorage/Source/TedsCore/Private/TypedElementsDataStorage.cpp b/Engine/Plugins/Experimental/EditorDataStorage/Source/TedsCore/Private/TypedElementsDataStorage.cpp
--- a/Engine/Plugins/Experimental/EditorDataStorage/Source/TedsCore/Private/TypedElementsDataStorage.cpp
+++ b/Engine/Plugins/Experimental/EditorDataStorage/Source/TedsCore/Private/TypedElementsDataStorage.cpp
@@ -76,7 +76,7 @@ void FTypedElementsDataStorageModule::StartupModule()
 		{
 			if (!bInitialized)
 			{
-				UE_LOG(LogTypedElementDataStorage, Log, TEXT("Initializing"));
+				UE_LOG(LogEditorDataStorage, Log, TEXT("Initializing"));
 				
 				Database = NewObject<UTypedElementDatabase>();
 				Database->Initialize();
@@ -128,7 +128,7 @@ void FTypedElementsDataStorageModule::StartupModule()
 					Factory->RegisterWidgetConstructors(*Database, *DatabaseUi);
 				}
 				
-				UE_LOG(LogTypedElementDataStorage, Log, TEXT("Initialized"));
+				UE_LOG(LogEditorDataStorage, Log, TEXT("Initialized"));
 
 				bInitialized = true;
 			}
@@ -140,7 +140,7 @@ void FTypedElementsDataStorageModule::ShutdownModule()
 {
 	if (bInitialized)
 	{
-		UE_LOG(LogTypedElementDataStorage, Log, TEXT("Deinitializing"));
+		UE_LOG(LogEditorDataStorage, Log, TEXT("Deinitializing"));
 
 		Database->ResetFactories();
 
diff --git a/Engine/Plugins/Experimental/EditorDataStorageFeatures/Source/TedsDebugger/Private/Widgets/QueryEditor/TedsConditionSelectionComboWidget.cpp b/Engine/Plugins/Experimental/EditorDataStorageFeatures/Source/TedsDebugger/Private/Widgets/QueryEditor/TedsConditionSelectionComboWidget.cpp
--- a/Engine/Plugins/Experimental/EditorDataStorageFeatures/Source/TedsDebugger/Private/Widgets/QueryEditor/TedsConditionSelectionComboWidget.cpp
+++ b/Engine/Plugins/Experimental/EditorDataStorageFeatures/Source/TedsDebugger/Private/Widgets/QueryEditor/TedsConditionSelectionComboWidget.cpp
@@ -94,7 +94,7 @@ namespace UE::EditorDataStorage::Debug::QueryEditor
 			EErrorCode ErrorCode = Model->SetOperatorType(NewSelection->Handle, ConditionType);
 			if (ErrorCode != EErrorCode::Success)
 			{
-				UE_LOG(LogTypedElementDataStorage, Error, TEXT("Could not set model condition: [%d]"), static_cast<int>(ErrorCode));
+				UE_LOG(LogEditorDataStorage, Error, TEXT("Could not set model condition: [%d]"), static_cast<int>(ErrorCode));
 			}
 			ComboBox->ClearSelection();
 		}
diff --git a/Engine/Plugins/Experimental/EditorDataStorageFeatures/Source/TedsTableViewer/Private/TedsTableViewerUtils.cpp b/Engine/Plugins/Experimental/EditorDataStorageFeatures/Source/TedsTableViewer/Private/TedsTableViewerUtils.cpp
--- a/Engine/Plugins/Experimental/EditorDataStorageFeatures/Source/TedsTableViewer/Private/TedsTableViewerUtils.cpp
+++ b/Engine/Plugins/Experimental/EditorDataStorageFeatures/Source/TedsTableViewer/Private/TedsTableViewerUtils.cpp
@@ -79,7 +79,7 @@ namespace UE::EditorDataStorage::TableViewerUtils
 			}
 			else
 			{
-				UE_LOG(LogTypedElementDataStorage, Verbose, TEXT("Invalid column provided to the table viewer"));
+				UE_LOG(LogEditorDataStorage, Verbose, TEXT("Invalid column provided to the table viewer"));
 			}
 		}
 		return VerifiedColumnTypes;
diff --git a/Engine/Source/Runtime/TypedElementFramework/Private/Elements/Common/TypedElementDataStorageLog.cpp b/Engine/Source/Runtime/TypedElementFramework/Private/Elements/Common/TypedElementDataStorageLog.cpp
--- a/Engine/Source/Runtime/TypedElementFramework/Private/Elements/Common/TypedElementDataStorageLog.cpp
+++ b/Engine/Source/Runtime/TypedElementFramework/Private/Elements/Common/TypedElementDataStorageLog.cpp
@@ -1,5 +1,5 @@
-﻿// Copyright Epic Games, Inc. All Rights Reserved.
+// Copyright Epic Games, Inc. All Rights Reserved.
 
 #include "Elements/Common/TypedElementDataStorageLog.h"
 
-DEFINE_LOG_CATEGORY(LogTypedElementDataStorage)
\ No newline at end of file
+DEFINE_LOG_CATEGORY(LogEditorDataStorage)
\ No newline at end of file
diff --git a/Engine/Source/Runtime/TypedElementFramework/Private/Elements/Framework/TypedElementQueryBuilder.cpp b/Engine/Source/Runtime/TypedElementFramework/Private/Elements/Framework/TypedElementQueryBuilder.cpp
--- a/Engine/Source/Runtime/TypedElementFramework/Private/Elements/Framework/TypedElementQueryBuilder.cpp
+++ b/Engine/Source/Runtime/TypedElementFramework/Private/Elements/Framework/TypedElementQueryBuilder.cpp
@@ -161,7 +161,7 @@ namespace TypedElementQueryBuilder
 		const FName ValueName = Enum.GetNameByValue(Value);
 		if (ValueName == NAME_None)
 		{
-			UE_LOG(LogTypedElementDataStorage, Warning, TEXT("Invalid value '%lld' for enum '%s'"), Value, *Enum.GetName());
+			UE_LOG(LogEditorDataStorage, Warning, TEXT("Invalid value '%lld' for enum '%s'"), Value, *Enum.GetName());
 			return *this;
 		}
 		using namespace UE::Editor::DataStorage;
diff --git a/Engine/Source/Runtime/TypedElementFramework/Public/Elements/Common/TypedElementDataStorageLog.h b/Engine/Source/Runtime/TypedElementFramework/Public/Elements/Common/TypedElementDataStorageLog.h
--- a/Engine/Source/Runtime/TypedElementFramework/Public/Elements/Common/TypedElementDataStorageLog.h
+++ b/Engine/Source/Runtime/TypedElementFramework/Public/Elements/Common/TypedElementDataStorageLog.h
@@ -1,7 +1,7 @@
-﻿// Copyright Epic Games, Inc. All Rights Reserved.
+// Copyright Epic Games, Inc. All Rights Reserved.
 
 #pragma once
 
 #include <Logging/LogMacros.h>
 
-TYPEDELEMENTFRAMEWORK_API DECLARE_LOG_CATEGORY_EXTERN(LogTypedElementDataStorage, Log, All);
\ No newline at end of file
+TYPEDELEMENTFRAMEWORK_API DECLARE_LOG_CATEGORY_EXTERN(LogEditorDataStorage, Log, All);
\ No newline at end of file
