TEDS UI: Register widget purposes that were used by SRowDetails

diff --git a/Engine/Plugins/Experimental/EditorDataStorage/Source/TedsUI/Private/Widgets/GeneralWidgetRegistrationFactory.cpp b/Engine/Plugins/Experimental/EditorDataStorage/Source/TedsUI/Private/Widgets/GeneralWidgetRegistrationFactory.cpp
--- a/Engine/Plugins/Experimental/EditorDataStorage/Source/TedsUI/Private/Widgets/GeneralWidgetRegistrationFactory.cpp
+++ b/Engine/Plugins/Experimental/EditorDataStorage/Source/TedsUI/Private/Widgets/GeneralWidgetRegistrationFactory.cpp
@@ -7,6 +7,7 @@
 #define LOCTEXT_NAMESPACE "TypedElementsUI_GeneralRegistration"
 
 const FName UGeneralWidgetRegistrationFactory::CellPurpose(TEXT("General.Cell"));
+const FName UGeneralWidgetRegistrationFactory::LargeCellPurpose(TEXT("General.Cell.Large"));
 const FName UGeneralWidgetRegistrationFactory::HeaderPurpose(TEXT("General.Header"));
 const FName UGeneralWidgetRegistrationFactory::CellDefaultPurpose(TEXT("General.Cell.Default"));
 const FName UGeneralWidgetRegistrationFactory::HeaderDefaultPurpose(TEXT("General.Header.Default"));
@@ -16,6 +17,9 @@ void UGeneralWidgetRegistrationFactory::RegisterWidgetPurposes(ITypedElementData
 	DataStorageUi.RegisterWidgetPurpose(CellPurpose, ITypedElementDataStorageUiInterface::EPurposeType::UniqueByNameAndColumn,
 		LOCTEXT("GeneralCellPurpose", "General purpose widgets that can be used as cells for specific columns or column combinations."));
 
+	DataStorageUi.RegisterWidgetPurpose(LargeCellPurpose, ITypedElementDataStorageUiInterface::EPurposeType::UniqueByNameAndColumn,
+		LOCTEXT("GeneralLargeCellPurpose", "General purpose widgets that are specifically designed to be embedded in a space larger than a single cell."));
+
 	DataStorageUi.RegisterWidgetPurpose(HeaderPurpose, ITypedElementDataStorageUiInterface::EPurposeType::UniqueByNameAndColumn,
 		LOCTEXT("GeneralHeaderPurpose", "General purpose widget that can be used as a header."));
 
diff --git a/Engine/Plugins/Experimental/EditorDataStorage/Source/TedsUI/Public/Widgets/GeneralWidgetRegistrationFactory.h b/Engine/Plugins/Experimental/EditorDataStorage/Source/TedsUI/Public/Widgets/GeneralWidgetRegistrationFactory.h
--- a/Engine/Plugins/Experimental/EditorDataStorage/Source/TedsUI/Public/Widgets/GeneralWidgetRegistrationFactory.h
+++ b/Engine/Plugins/Experimental/EditorDataStorage/Source/TedsUI/Public/Widgets/GeneralWidgetRegistrationFactory.h
@@ -14,6 +14,7 @@ class TEDSUI_API UGeneralWidgetRegistrationFactory : public UTypedElementDataSto
 
 public:
 	static const FName CellPurpose;
+	static const FName LargeCellPurpose;
 	static const FName HeaderPurpose;
 	static const FName CellDefaultPurpose;
 	static const FName HeaderDefaultPurpose;
diff --git a/Engine/Plugins/Experimental/EditorDataStorageFeatures/Source/TedsTableViewer/Private/Widgets/RowHandleWidget.cpp b/Engine/Plugins/Experimental/EditorDataStorageFeatures/Source/TedsTableViewer/Private/Widgets/RowHandleWidget.cpp
--- a/Engine/Plugins/Experimental/EditorDataStorageFeatures/Source/TedsTableViewer/Private/Widgets/RowHandleWidget.cpp
+++ b/Engine/Plugins/Experimental/EditorDataStorageFeatures/Source/TedsTableViewer/Private/Widgets/RowHandleWidget.cpp
@@ -18,6 +18,14 @@ void URowHandleWidgetFactory::RegisterWidgetPurposes(ITypedElementDataStorageUiI
 {
 	DataStorageUi.RegisterWidgetPurpose(FName(TEXT("General.Cell.RowHandle")), ITypedElementDataStorageUiInterface::EPurposeType::UniqueByName,
 	LOCTEXT("RowHandlePurpose", "Specific purpose to request a widget to display row handles."));
+
+	DataStorageUi.RegisterWidgetPurpose(FName(TEXT("RowDetails.Cell")), ITypedElementDataStorageUiInterface::EPurposeType::UniqueByNameAndColumn,
+	LOCTEXT("RowHandlePurpose", "Specific purpose to request a widget to display the details on a row (e.g SRowDetails)."));
+
+	DataStorageUi.RegisterWidgetPurpose(FName(TEXT("RowDetails.Cell.Large")), ITypedElementDataStorageUiInterface::EPurposeType::UniqueByNameAndColumn,
+	LOCTEXT("RowHandlePurpose", "Specific purpose to request a widget that is larger than a single cell to display the details on a row (e.g SRowDetails)"));
+
+
 }
 
 FRowHandleWidgetConstructor::FRowHandleWidgetConstructor()
