Added support for TEDS Query Callbacks to be called per chunk on different threads.

Tables in TEDS are broken up in smaller chunks that get processed by calling Query Callbacks for each of them. This typically happens on the calling thread, meaning all chunks are processed in sequence. This updates adds an new option that allows the chunks to be processed in parallel on multiple threads. In some cases this can improve performance, if the processors don't have any internal dependencies on each other.
This change also cleans up some outdated alias uses.

diff --git a/Engine/Plugins/Experimental/EditorDataStorage/Source/TedsCore/Private/Compatibility/TypedElementActorLabelQueries.cpp b/Engine/Plugins/Experimental/EditorDataStorage/Source/TedsCore/Private/Compatibility/TypedElementActorLabelQueries.cpp
--- a/Engine/Plugins/Experimental/EditorDataStorage/Source/TedsCore/Private/Compatibility/TypedElementActorLabelQueries.cpp
+++ b/Engine/Plugins/Experimental/EditorDataStorage/Source/TedsCore/Private/Compatibility/TypedElementActorLabelQueries.cpp
@@ -22,13 +22,13 @@ void UTypedElementActorLabelFactory::RegisterQueries(ITypedElementDataStorageInt
 void UTypedElementActorLabelFactory::RegisterActorLabelToColumnQuery(ITypedElementDataStorageInterface& DataStorage) const
 {
 	using namespace TypedElementQueryBuilder;
-	using DSI = ITypedElementDataStorageInterface;
-
+	using namespace TypedElementDataStorage;
+	
 	DataStorage.RegisterQuery(
 		Select(
 			TEXT("Sync actor label to column"),
-			FProcessor(DSI::EQueryTickPhase::PrePhysics, DataStorage.GetQueryTickGroupName(DSI::EQueryTickGroups::SyncExternalToDataStorage))
-				.ForceToGameThread(true),
+			FProcessor(EQueryTickPhase::PrePhysics, DataStorage.GetQueryTickGroupName(EQueryTickGroups::SyncExternalToDataStorage))
+				.SetExecutionMode(EExecutionMode::GameThread),
 			[](const FTypedElementUObjectColumn& Actor, FTypedElementLabelColumn& Label, FTypedElementLabelHashColumn& LabelHash)
 			{
 				if (const AActor* ActorInstance = Cast<AActor>(Actor.Object); ActorInstance != nullptr)
@@ -52,13 +52,13 @@ void UTypedElementActorLabelFactory::RegisterActorLabelToColumnQuery(ITypedEleme
 void UTypedElementActorLabelFactory::RegisterLabelColumnToActorQuery(ITypedElementDataStorageInterface& DataStorage) const
 {
 	using namespace TypedElementQueryBuilder;
-	using DSI = ITypedElementDataStorageInterface;
-
+	using namespace TypedElementDataStorage;
+	
 	DataStorage.RegisterQuery(
 		Select(
 			TEXT("Sync label column to actor"),
-			FProcessor(DSI::EQueryTickPhase::FrameEnd, DataStorage.GetQueryTickGroupName(DSI::EQueryTickGroups::SyncDataStorageToExternal))
-				.ForceToGameThread(true),
+			FProcessor(EQueryTickPhase::FrameEnd, DataStorage.GetQueryTickGroupName(EQueryTickGroups::SyncDataStorageToExternal))
+				.SetExecutionMode(EExecutionMode::GameThread),
 			[](FTypedElementUObjectColumn& Actor, const FTypedElementLabelColumn& Label, const FTypedElementLabelHashColumn& LabelHash)
 			{
 				if (AActor* ActorInstance = Cast<AActor>(Actor.Object); ActorInstance != nullptr)
diff --git a/Engine/Plugins/Experimental/EditorDataStorage/Source/TedsCore/Private/Compatibility/TypedElementActorParentQueries.cpp b/Engine/Plugins/Experimental/EditorDataStorage/Source/TedsCore/Private/Compatibility/TypedElementActorParentQueries.cpp
--- a/Engine/Plugins/Experimental/EditorDataStorage/Source/TedsCore/Private/Compatibility/TypedElementActorParentQueries.cpp
+++ b/Engine/Plugins/Experimental/EditorDataStorage/Source/TedsCore/Private/Compatibility/TypedElementActorParentQueries.cpp
@@ -38,7 +38,7 @@ void UTypedElementActorParentFactory::RegisterAddParentColumn(ITypedElementDataS
 		Select(
 			TEXT("Add parent column to actor"),
 			FProcessor(EQueryTickPhase::PrePhysics, DataStorage.GetQueryTickGroupName(EQueryTickGroups::SyncExternalToDataStorage))
-			.ForceToGameThread(true),
+				.SetExecutionMode(EExecutionMode::GameThread),
 			[](IQueryContext& Context, TypedElementRowHandle Row, const FTypedElementUObjectColumn& Actor)
 			{
 				if (const AActor* ActorInstance = Cast<AActor>(Actor.Object))
@@ -75,7 +75,7 @@ void UTypedElementActorParentFactory::RegisterUpdateOrRemoveParentColumn(ITypedE
 		Select(
 			TEXT("Sync actor's parent to column"),
 			FProcessor(EQueryTickPhase::PrePhysics, DataStorage.GetQueryTickGroupName(EQueryTickGroups::SyncExternalToDataStorage))
-				.ForceToGameThread(true),
+				.SetExecutionMode(EExecutionMode::GameThread),
 			[](IQueryContext& Context, TypedElementRowHandle Row, const FTypedElementUObjectColumn& Actor, FTableRowParentColumn& Parent)
 			{
 				if (const AActor* ActorInstance = Cast<AActor>(Actor.Object))
diff --git a/Engine/Plugins/Experimental/EditorDataStorage/Source/TedsCore/Private/Compatibility/TypedElementActorTransformQueries.cpp b/Engine/Plugins/Experimental/EditorDataStorage/Source/TedsCore/Private/Compatibility/TypedElementActorTransformQueries.cpp
--- a/Engine/Plugins/Experimental/EditorDataStorage/Source/TedsCore/Private/Compatibility/TypedElementActorTransformQueries.cpp
+++ b/Engine/Plugins/Experimental/EditorDataStorage/Source/TedsCore/Private/Compatibility/TypedElementActorTransformQueries.cpp
@@ -18,15 +18,15 @@ void UTypedElementActorTransformFactory::RegisterQueries(ITypedElementDataStorag
 void UTypedElementActorTransformFactory::RegisterActorAddTransformColumn(ITypedElementDataStorageInterface& DataStorage) const
 {
 	using namespace TypedElementQueryBuilder;
-	using DSI = ITypedElementDataStorageInterface;
-
+	using namespace TypedElementDataStorage;
+	
 	DataStorage.RegisterQuery(
 		Select(
 			TEXT("Add transform column to actor"),
-			FProcessor(DSI::EQueryTickPhase::PrePhysics,
-				DataStorage.GetQueryTickGroupName(DSI::EQueryTickGroups::SyncExternalToDataStorage))
-			.ForceToGameThread(true),
-			[](DSI::IQueryContext& Context, TypedElementRowHandle Row, const FTypedElementUObjectColumn& Actor)
+			FProcessor(EQueryTickPhase::PrePhysics,
+				DataStorage.GetQueryTickGroupName(EQueryTickGroups::SyncExternalToDataStorage))
+				.SetExecutionMode(EExecutionMode::GameThread),
+			[](IQueryContext& Context, TypedElementRowHandle Row, const FTypedElementUObjectColumn& Actor)
 			{
 				if (const AActor* ActorInstance = Cast<AActor>(Actor.Object); ActorInstance != nullptr && ActorInstance->GetRootComponent())
 				{
@@ -44,14 +44,14 @@ void UTypedElementActorTransformFactory::RegisterActorAddTransformColumn(ITypedE
 void UTypedElementActorTransformFactory::RegisterActorLocalTransformToColumn(ITypedElementDataStorageInterface& DataStorage) const
 {
 	using namespace TypedElementQueryBuilder;
-	using DSI = ITypedElementDataStorageInterface;
-
+	using namespace TypedElementDataStorage;
+	
 	DataStorage.RegisterQuery(
 		Select(
 			TEXT("Sync actor transform to column"),
-			FProcessor(DSI::EQueryTickPhase::PostPhysics, DataStorage.GetQueryTickGroupName(DSI::EQueryTickGroups::SyncExternalToDataStorage))
-				.ForceToGameThread(true),
-			[](DSI::IQueryContext& Context, TypedElementRowHandle Row, const FTypedElementUObjectColumn& Actor, FTypedElementLocalTransformColumn& Transform)
+			FProcessor(EQueryTickPhase::PostPhysics, DataStorage.GetQueryTickGroupName(EQueryTickGroups::SyncExternalToDataStorage))
+				.SetExecutionMode(EExecutionMode::GameThread),
+			[](IQueryContext& Context, TypedElementRowHandle Row, const FTypedElementUObjectColumn& Actor, FTypedElementLocalTransformColumn& Transform)
 			{
 				if (const AActor* ActorInstance = Cast<AActor>(Actor.Object); ActorInstance != nullptr && ActorInstance->GetRootComponent() != nullptr)
 				{
@@ -73,13 +73,13 @@ void UTypedElementActorTransformFactory::RegisterActorLocalTransformToColumn(ITy
 void UTypedElementActorTransformFactory::RegisterLocalTransformColumnToActor(ITypedElementDataStorageInterface& DataStorage) const
 {
 	using namespace TypedElementQueryBuilder;
-	using DSI = ITypedElementDataStorageInterface;
-
+	using namespace TypedElementDataStorage;
+	
 	DataStorage.RegisterQuery(
 		Select(
 			TEXT("Sync transform column to actor"),
-			FProcessor(DSI::EQueryTickPhase::FrameEnd, DataStorage.GetQueryTickGroupName(DSI::EQueryTickGroups::SyncDataStorageToExternal))
-				.ForceToGameThread(true),
+			FProcessor(EQueryTickPhase::FrameEnd, DataStorage.GetQueryTickGroupName(EQueryTickGroups::SyncDataStorageToExternal))
+				.SetExecutionMode(EExecutionMode::GameThread),
 			[](FTypedElementUObjectColumn& Actor, const FTypedElementLocalTransformColumn& Transform)
 			{
 				if (AActor* ActorInstance = Cast<AActor>(Actor.Object); ActorInstance != nullptr)
diff --git a/Engine/Plugins/Experimental/EditorDataStorage/Source/TedsCore/Private/Compatibility/TypedElementActorViewportProcessors.cpp b/Engine/Plugins/Experimental/EditorDataStorage/Source/TedsCore/Private/Compatibility/TypedElementActorViewportProcessors.cpp
--- a/Engine/Plugins/Experimental/EditorDataStorage/Source/TedsCore/Private/Compatibility/TypedElementActorViewportProcessors.cpp
+++ b/Engine/Plugins/Experimental/EditorDataStorage/Source/TedsCore/Private/Compatibility/TypedElementActorViewportProcessors.cpp
@@ -173,13 +173,13 @@ void UTypedElementActorViewportFactory::RegisterQueries(ITypedElementDataStorage
 void UTypedElementActorViewportFactory::RegisterOutlineColorColumnToActor(ITypedElementDataStorageInterface& DataStorage)
 {
 	using namespace TypedElementQueryBuilder;
-	using DSI = ITypedElementDataStorageInterface;
+	using namespace TypedElementDataStorage;
 
 	DataStorage.RegisterQuery(
 		Select(
 			TEXT("Sync viewport outline color column to actor"),
-			FProcessor(DSI::EQueryTickPhase::DuringPhysics, DataStorage.GetQueryTickGroupName(DSI::EQueryTickGroups::SyncDataStorageToExternal))
-				.ForceToGameThread(true),
+			FProcessor(EQueryTickPhase::DuringPhysics, DataStorage.GetQueryTickGroupName(EQueryTickGroups::SyncDataStorageToExternal))
+				.SetExecutionMode(EExecutionMode::GameThread),
 			[](FTypedElementUObjectColumn& Actor, const FTypedElementViewportOutlineColorColumn& ViewportColor)
 			{
 				if (AActor* ActorInstance = Cast<AActor>(Actor.Object); ActorInstance != nullptr)
@@ -201,13 +201,13 @@ void UTypedElementActorViewportFactory::RegisterOutlineColorColumnToActor(ITyped
 void UTypedElementActorViewportFactory::RegisterOverlayColorColumnToActor(ITypedElementDataStorageInterface& DataStorage)
 {
 	using namespace TypedElementQueryBuilder;
-	using DSI = ITypedElementDataStorageInterface;
-
+	using namespace TypedElementDataStorage;
+	
 	DataStorage.RegisterQuery(
 		Select(
 			TEXT("Sync viewport overlay color column to actor"),
 			FObserver::OnAdd<FTypedElementViewportOverlayColorColumn>()
-				.ForceToGameThread(true),
+				.SetExecutionMode(EExecutionMode::GameThread),
 			[](FTypedElementUObjectColumn& Actor, const FTypedElementViewportOverlayColorColumn& ViewportColor)
 			{
 				if (AActor* ActorInstance = Cast<AActor>(Actor.Object); ActorInstance != nullptr)
diff --git a/Engine/Plugins/Experimental/EditorDataStorage/Source/TedsCore/Private/Compatibility/TypedElementUObjectWorldQueries.cpp b/Engine/Plugins/Experimental/EditorDataStorage/Source/TedsCore/Private/Compatibility/TypedElementUObjectWorldQueries.cpp
--- a/Engine/Plugins/Experimental/EditorDataStorage/Source/TedsCore/Private/Compatibility/TypedElementUObjectWorldQueries.cpp
+++ b/Engine/Plugins/Experimental/EditorDataStorage/Source/TedsCore/Private/Compatibility/TypedElementUObjectWorldQueries.cpp
@@ -22,7 +22,7 @@ void UTypedElementUObjectWorldFactory::RegisterAddWorldColumn(ITypedElementDataS
 		Select(
 			TEXT("Add world column to UObject"),
 			FProcessor(EQueryTickPhase::PrePhysics, DataStorage.GetQueryTickGroupName(EQueryTickGroups::SyncExternalToDataStorage))
-			.ForceToGameThread(true),
+				.SetExecutionMode(EExecutionMode::GameThread),
 			[](IQueryContext& Context, TypedElementRowHandle Row, const FTypedElementUObjectColumn& Object)
 			{
 				// Not all objects, in particular actors, are always correctly cleaned up, resulting in dangling
@@ -52,7 +52,7 @@ void UTypedElementUObjectWorldFactory::RegisterUpdateOrRemoveWorldColumn(ITypedE
 		Select(
 			TEXT("Sync UObject's world to column"),
 			FProcessor(EQueryTickPhase::PostPhysics, DataStorage.GetQueryTickGroupName(EQueryTickGroups::SyncExternalToDataStorage))
-			.ForceToGameThread(true),
+				.SetExecutionMode(EExecutionMode::GameThread),
 			[](IQueryContext& Context, TypedElementRowHandle Row, const FTypedElementUObjectColumn& Object, FTypedElementWorldColumn& World)
 			{
 				// Not all objects, in particular actors, are always correctly cleaned up, resulting in dangling
diff --git a/Engine/Plugins/Experimental/EditorDataStorage/Source/TedsCore/Private/Processors/TypedElementProcessorAdaptors.cpp b/Engine/Plugins/Experimental/EditorDataStorage/Source/TedsCore/Private/Processors/TypedElementProcessorAdaptors.cpp
--- a/Engine/Plugins/Experimental/EditorDataStorage/Source/TedsCore/Private/Processors/TypedElementProcessorAdaptors.cpp
+++ b/Engine/Plugins/Experimental/EditorDataStorage/Source/TedsCore/Private/Processors/TypedElementProcessorAdaptors.cpp
@@ -833,7 +833,24 @@ void FTypedElementQueryProcessorData::DebugOutputDescription(FOutputDevice& Ar,
 			Ar.Logf(TEXT("\n%*sMonitored type: %s"), Indent, TEXT(""), *Callback.MonitoredType->GetName());
 		}
 
-		Ar.Logf(TEXT("\n%*sIs forced to GameThread: %s"), Indent, TEXT(""), Callback.bForceToGameThread ? TEXT("True") : TEXT("False"));
+		switch (Callback.ExecutionMode)
+		{
+		case TypedElementDataStorage::EExecutionMode::Default:
+			Ar.Logf(TEXT("\n%*sExecution mode: Default"), Indent, TEXT(""));
+			break;
+		case TypedElementDataStorage::EExecutionMode::GameThread:
+			Ar.Logf(TEXT("\n%*sExecution mode: Game Thread"), Indent, TEXT(""));
+			break;
+		case TypedElementDataStorage::EExecutionMode::Threaded:
+			Ar.Logf(TEXT("\n%*sExecution mode: Threaded"), Indent, TEXT(""));
+			break;
+		case TypedElementDataStorage::EExecutionMode::ThreadedChunks:
+			Ar.Logf(TEXT("\n%*sExecution mode: Threaded Chunks"), Indent, TEXT(""));
+			break;
+		default:
+			Ar.Logf(TEXT("\n%*sExecution mode: <Unknown option>"), Indent, TEXT(""));
+			break;
+		}
 	}
 #endif // WITH_MASSENTITY_DEBUG
 }
@@ -876,24 +893,39 @@ TypedElementDataStorage::FQueryResult FTypedElementQueryProcessorData::Execute(
 	TypedElementDataStorage::FQueryDescription& Description,
 	FMassEntityQuery& NativeQuery, 
 	FMassEntityManager& EntityManager,
-	FEnvironment& Environment)
+	FEnvironment& Environment,
+	TypedElementDataStorage::EDirectQueryExecutionFlags ExecutionFlags)
 {
-	ITypedElementDataStorageInterface::FQueryResult Result;
-	Result.Completed = ITypedElementDataStorageInterface::FQueryResult::ECompletion::Fully;
+	using namespace TypedElementDataStorage;
+
+	FQueryResult Result;
+	Result.Completed = FQueryResult::ECompletion::Fully;
 	
-	if (Description.Callback.ActivationCount > 0)
+	if (EnumHasAnyFlags(ExecutionFlags, EDirectQueryExecutionFlags::AllowBoundQueries) || !Description.Callback.Function)
 	{
-		FMassExecutionContext Context(EntityManager);
-		
-		NativeQuery.ForEachEntityChunk(EntityManager, Context,
-			[&Result, &Callback, &Description](FMassExecutionContext& Context)
+		if (EnumHasAnyFlags(ExecutionFlags, EDirectQueryExecutionFlags::IgnoreActivationCount) || Description.Callback.ActivationCount > 0)
+		{
+			FMassExecutionContext Context(EntityManager);
+			auto ExecuteFunction = [&Result, &Callback, &Description](FMassExecutionContext& Context)
+				{
+					// No need to cache any subsystem dependencies as these are not accessible from a direct query.
+					UE::Editor::DataStorage::Processors::Private::FMassDirectContextForwarder QueryContext(Context);
+					Callback(Description, QueryContext);
+					Result.Count += Context.GetNumEntities();
+				};
+			if (EnumHasAnyFlags(ExecutionFlags, EDirectQueryExecutionFlags::ParallelizeChunks))
 			{
-				// No need to cache any subsystem dependencies as these are not accessible from a direct query.
-				UE::Editor::DataStorage::Processors::Private::FMassDirectContextForwarder QueryContext(Context);
-				Callback(Description, QueryContext);
-				Result.Count += Context.GetNumEntities();
+				NativeQuery.ParallelForEachEntityChunk(EntityManager, Context, ExecuteFunction);
 			}
-		);
+			else
+			{
+				NativeQuery.ForEachEntityChunk(EntityManager, Context, ExecuteFunction);
+			}
+		}
+	}
+	else
+	{
+		Result.Completed = FQueryResult::ECompletion::Unsupported;
 	}
 	return Result;
 }
@@ -906,11 +938,16 @@ TypedElementDataStorage::FQueryResult FTypedElementQueryProcessorData::Execute(
 	FEnvironment& Environment,
 	FMassExecutionContext& ParentContext)
 {
+	using namespace TypedElementDataStorage;
+
 	ITypedElementDataStorageInterface::FQueryResult Result;
 	Result.Completed = ITypedElementDataStorageInterface::FQueryResult::ECompletion::Fully;
 
 	if (Description.Callback.ActivationCount > 0)
 	{
+		checkf(Description.Callback.ExecutionMode != EExecutionMode::ThreadedChunks,
+			TEXT("TEDS Sub-queries do not support parallel chunk processing."));
+		
 		FMassExecutionContext Context(EntityManager);
 		Context.SetDeferredCommandBuffer(ParentContext.GetSharedDeferredCommandBuffer());
 		Context.SetFlushDeferredCommands(false);
@@ -937,12 +974,17 @@ TypedElementDataStorage::FQueryResult FTypedElementQueryProcessorData::Execute(
 	FEnvironment& Environment,
 	FMassExecutionContext& ParentContext)
 {
+	using namespace TypedElementDataStorage;
+
 	ITypedElementDataStorageInterface::FQueryResult Result;
 	Result.Completed = ITypedElementDataStorageInterface::FQueryResult::ECompletion::Fully;
 
 	FMassEntityHandle NativeEntity = FMassEntityHandle::FromNumber(RowHandle);
 	if (Description.Callback.ActivationCount > 0 && EntityManager.IsEntityActive(NativeEntity))
 	{
+		checkf(Description.Callback.ExecutionMode != EExecutionMode::ThreadedChunks,
+			TEXT("TEDS Sub-queries do not support parallel chunk processing."));
+		
 		FMassArchetypeHandle NativeArchetype = EntityManager.GetArchetypeForEntityUnsafe(NativeEntity);
 		FMassExecutionContext Context(EntityManager);
 		Context.SetEntityCollection(FMassArchetypeEntityCollection(NativeArchetype, { NativeEntity }, FMassArchetypeEntityCollection::NoDuplicates));
@@ -965,22 +1007,32 @@ TypedElementDataStorage::FQueryResult FTypedElementQueryProcessorData::Execute(
 
 void FTypedElementQueryProcessorData::Execute(FMassEntityManager& EntityManager, FMassExecutionContext& Context)
 {
+	using namespace TypedElementDataStorage;
+	
 	FExtendedQuery* StoredQuery = QueryStore->GetMutable(ParentQuery);
+
 	checkf(StoredQuery, TEXT("A query callback was registered for execution without an associated query."));
 	
 	ITypedElementDataStorageInterface::FQueryDescription& Description = StoredQuery->Description;
 	if (Description.Callback.ActivationCount > 0)
 	{
-		NativeQuery.ForEachEntityChunk(EntityManager, Context,
-			[this, &Description](FMassExecutionContext& Context)
+		auto ExceteFunction = [this, &Description](FMassExecutionContext& Context)
 			{
 				if (PrepareCachedDependenciesOnQuery(Description, Context))
 				{
 					UE::Editor::DataStorage::Processors::Private::FMassContextForwarder QueryContext(Description, Context, *QueryStore, *Environment);
 					Description.Callback.Function(Description, QueryContext);
 				}
-			}
-		);
+			};
+		
+		if (StoredQuery->Description.Callback.ExecutionMode != EExecutionMode::ThreadedChunks)
+		{
+			NativeQuery.ForEachEntityChunk(EntityManager, Context, ExceteFunction);
+		}
+		else
+		{
+			NativeQuery.ParallelForEachEntityChunk(EntityManager, Context, ExceteFunction);
+		}
 	}
 }
 
@@ -1029,7 +1081,7 @@ bool UTypedElementQueryProcessorCallbackAdapterProcessorBase::ConfigureQueryCall
 {
 	bool Result = Data.CommonQueryConfiguration(*this, Query, QueryHandle, QueryStore, Environment, Subqueries);
 
-	bRequiresGameThreadExecution = Query.Description.Callback.bForceToGameThread;
+	bRequiresGameThreadExecution = Query.Description.Callback.ExecutionMode == TypedElementDataStorage::EExecutionMode::GameThread;
 	ExecutionFlags = static_cast<int32>(EProcessorExecutionFlags::Editor); 
 	ExecutionOrder.ExecuteInGroup = Query.Description.Callback.Group;
 	ExecutionOrder.ExecuteBefore = Query.Description.Callback.BeforeGroups;
@@ -1158,7 +1210,7 @@ bool UTypedElementQueryObserverCallbackAdapterProcessorBase::ConfigureQueryCallb
 {
 	bool Result = Data.CommonQueryConfiguration(*this, Query, QueryHandle, QueryStore, Environment, Subqueries);
 
-	bRequiresGameThreadExecution = Query.Description.Callback.bForceToGameThread;
+	bRequiresGameThreadExecution = Query.Description.Callback.ExecutionMode == TypedElementDataStorage::EExecutionMode::GameThread;
 	ExecutionFlags = static_cast<int32>(EProcessorExecutionFlags::Editor);
 	
 	ObservedType = const_cast<UScriptStruct*>(Query.Description.Callback.MonitoredType);
diff --git a/Engine/Plugins/Experimental/EditorDataStorage/Source/TedsCore/Private/Processors/TypedElementProcessorAdaptors.h b/Engine/Plugins/Experimental/EditorDataStorage/Source/TedsCore/Private/Processors/TypedElementProcessorAdaptors.h
--- a/Engine/Plugins/Experimental/EditorDataStorage/Source/TedsCore/Private/Processors/TypedElementProcessorAdaptors.h
+++ b/Engine/Plugins/Experimental/EditorDataStorage/Source/TedsCore/Private/Processors/TypedElementProcessorAdaptors.h
@@ -62,7 +62,8 @@ struct FTypedElementQueryProcessorData
 		TypedElementDataStorage::FQueryDescription& Description,
 		FMassEntityQuery& NativeQuery, 
 		FMassEntityManager& EntityManager,
-		FEnvironment& Environment);
+		FEnvironment& Environment,
+		TypedElementDataStorage::EDirectQueryExecutionFlags ExecutionFlags);
 	static TypedElementDataStorage::FQueryResult Execute(
 		TypedElementDataStorage::SubqueryCallbackRef& Callback,
 		TypedElementDataStorage::FQueryDescription& Description,
diff --git a/Engine/Plugins/Experimental/EditorDataStorage/Source/TedsCore/Private/Queries/TypedElementExtendedQueryStore.cpp b/Engine/Plugins/Experimental/EditorDataStorage/Source/TedsCore/Private/Queries/TypedElementExtendedQueryStore.cpp
--- a/Engine/Plugins/Experimental/EditorDataStorage/Source/TedsCore/Private/Queries/TypedElementExtendedQueryStore.cpp
+++ b/Engine/Plugins/Experimental/EditorDataStorage/Source/TedsCore/Private/Queries/TypedElementExtendedQueryStore.cpp
@@ -66,7 +66,7 @@ namespace UE::Editor::DataStorage
 	}
 
 	void FExtendedQueryStore::RegisterTickGroup(FName GroupName, ITypedElementDataStorageInterface::EQueryTickPhase Phase,
-		FName BeforeGroup, FName AfterGroup, bool bRequiresMainThread)
+		FName BeforeGroup, FName AfterGroup, TypedElementDataStorage::EExecutionMode ExecutionMode)
 	{
 		FTickGroupDescription& Group = TickGroupDescriptions.FindOrAdd({ GroupName, Phase });
 
@@ -80,10 +80,7 @@ namespace UE::Editor::DataStorage
 			Group.AfterGroups.Add(AfterGroup);
 		}
 
-		if (bRequiresMainThread)
-		{
-			Group.bRequiresMainThread = true;
-		}
+		Group.ExecutionMode = ExecutionMode;
 	}
 
 	void FExtendedQueryStore::UnregisterTickGroup(FName GroupName, ITypedElementDataStorageInterface::EQueryTickPhase Phase)
@@ -219,27 +216,26 @@ namespace UE::Editor::DataStorage
 		FEnvironment& Environment,
 		FMassExecutionContext* ParentContext,
 		Handle Query,
+		TypedElementDataStorage::EDirectQueryExecutionFlags DirectExecutionFlags,
 		CallbackReference Callback)
 	{
-		using ActionType = ITypedElementDataStorageInterface::FQueryDescription::EActionType;
-		using CompletionType = ITypedElementDataStorageInterface::FQueryResult::ECompletion;
-
-		ITypedElementDataStorageInterface::FQueryResult Result;
+		using namespace TypedElementDataStorage;
 
+		FQueryResult Result;
 		if (FExtendedQuery* QueryData = Get(Query))
 		{
 			switch (QueryData->Description.Action)
 			{
-			case ActionType::None:
-				Result.Completed = CompletionType::Fully;
+			case FQueryDescription::EActionType::None:
+				Result.Completed = FQueryResult::ECompletion::Fully;
 				break;
-			case ActionType::Select:
+			case FQueryDescription::EActionType::Select:
 				if (!QueryData->Processor.IsValid())
 				{
-					if constexpr (std::is_same_v<CallbackReference, TypedElementDataStorage::DirectQueryCallbackRef>)
+					if constexpr (std::is_same_v<CallbackReference, DirectQueryCallbackRef>)
 					{
 						Result = FTypedElementQueryProcessorData::Execute(
-							Callback, QueryData->Description, QueryData->NativeQuery, EntityManager, Environment);
+							Callback, QueryData->Description, QueryData->NativeQuery, EntityManager, Environment, DirectExecutionFlags);
 					}
 					else
 					{
@@ -249,38 +245,40 @@ namespace UE::Editor::DataStorage
 				}
 				else
 				{
-					Result.Completed = CompletionType::Unsupported;
+					Result.Completed = FQueryResult::ECompletion::Unsupported;
 				}
 				break;
-			case ActionType::Count:
+			case FQueryDescription::EActionType::Count:
 				// Only the count is requested so no need to trigger the callback.
 				Result.Count = QueryData->NativeQuery.GetNumMatchingEntities(EntityManager);
-				Result.Completed = CompletionType::Fully;
+				Result.Completed = FQueryResult::ECompletion::Fully;
 				break;
 			default:
-				Result.Completed = CompletionType::Unsupported;
+				Result.Completed = FQueryResult::ECompletion::Unsupported;
 				break;
 			}
 		}
 		else
 		{
-			Result.Completed = CompletionType::Unavailable;
+			Result.Completed = FQueryResult::ECompletion::Unavailable;
 		}
 
 		return Result;
 	}
 
 	TypedElementDataStorage::FQueryResult FExtendedQueryStore::RunQuery(FMassEntityManager& EntityManager,
-		FEnvironment& Environment, Handle Query, TypedElementDataStorage::DirectQueryCallbackRef Callback)
+		FEnvironment& Environment, Handle Query, TypedElementDataStorage::EDirectQueryExecutionFlags DirectExecutionFlags,
+		TypedElementDataStorage::DirectQueryCallbackRef Callback)
 	{
-		return RunQueryCallbackCommon(EntityManager, Environment, nullptr, Query, Callback);
+		return RunQueryCallbackCommon(EntityManager, Environment, nullptr, Query, DirectExecutionFlags, Callback);
 	}
-
+	
 	TypedElementDataStorage::FQueryResult FExtendedQueryStore::RunQuery(FMassEntityManager& EntityManager,
 		FEnvironment& Environment, FMassExecutionContext& ParentContext, Handle Query,
 		TypedElementDataStorage::SubqueryCallbackRef Callback)
 	{
-		return RunQueryCallbackCommon(EntityManager, Environment, &ParentContext, Query, Callback);
+		return RunQueryCallbackCommon(EntityManager, Environment, &ParentContext, Query, 
+			TypedElementDataStorage::EDirectQueryExecutionFlags::Default, Callback);
 	}
 
 	TypedElementDataStorage::FQueryResult FExtendedQueryStore::RunQuery(FMassEntityManager& EntityManager,
@@ -704,6 +702,8 @@ namespace UE::Editor::DataStorage
 
 	bool FExtendedQueryStore::SetupTickGroupDefaults(ITypedElementDataStorageInterface::FQueryDescription& Query)
 	{
+	using namespace TypedElementDataStorage;
+
 		const FTickGroupDescription* TickGroup = TickGroupDescriptions.Find({ Query.Callback.Group, Query.Callback.Phase });
 		if (TickGroup)
 		{
@@ -725,7 +725,10 @@ namespace UE::Editor::DataStorage
 			}
 			Query.Callback.AfterGroups.Append(TickGroup->AfterGroups);
 
-			Query.Callback.bForceToGameThread |= TickGroup->bRequiresMainThread;
+			if (Query.Callback.ExecutionMode == EExecutionMode::Default)
+			{
+				Query.Callback.ExecutionMode = TickGroup->ExecutionMode;
+			}
 		}
 		return true;
 	}
diff --git a/Engine/Plugins/Experimental/EditorDataStorage/Source/TedsCore/Private/Queries/TypedElementExtendedQueryStore.h b/Engine/Plugins/Experimental/EditorDataStorage/Source/TedsCore/Private/Queries/TypedElementExtendedQueryStore.h
--- a/Engine/Plugins/Experimental/EditorDataStorage/Source/TedsCore/Private/Queries/TypedElementExtendedQueryStore.h
+++ b/Engine/Plugins/Experimental/EditorDataStorage/Source/TedsCore/Private/Queries/TypedElementExtendedQueryStore.h
@@ -54,7 +54,7 @@ namespace UE::Editor::DataStorage
 
 		/** Register the defaults for a tick group. These will be applied on top of any settings provided with a query registration. */
 		void RegisterTickGroup(FName GroupName, ITypedElementDataStorageInterface::EQueryTickPhase Phase,
-			FName BeforeGroup, FName AfterGroup, bool bRequiresMainThread);
+			FName BeforeGroup, FName AfterGroup, TypedElementDataStorage::EExecutionMode ExecutionMode);
 		/** Removes a previously registered set of tick group defaults. */
 		void UnregisterTickGroup(FName GroupName, ITypedElementDataStorageInterface::EQueryTickPhase Phase);
 
@@ -106,6 +106,7 @@ namespace UE::Editor::DataStorage
 			FMassEntityManager& EntityManager,
 			FEnvironment& Environment,
 			Handle Query,
+			TypedElementDataStorage::EDirectQueryExecutionFlags DirectExecutionFlags,
 			TypedElementDataStorage::DirectQueryCallbackRef Callback);
 		TypedElementDataStorage::FQueryResult RunQuery(
 			FMassEntityManager& EntityManager,
@@ -151,7 +152,7 @@ namespace UE::Editor::DataStorage
 		{
 			TArray<FName> BeforeGroups;
 			TArray<FName> AfterGroups;
-			bool bRequiresMainThread{ false };
+			TypedElementDataStorage::EExecutionMode ExecutionMode = TypedElementDataStorage::EExecutionMode::Default;
 		};
 
 		template<typename CallbackReference>
@@ -160,6 +161,7 @@ namespace UE::Editor::DataStorage
 			FEnvironment& Environment,
 			FMassExecutionContext* ParentContext,
 			Handle Query,
+		TypedElementDataStorage::EDirectQueryExecutionFlags ExecutionFlags,
 			CallbackReference Callback);
 
 		FMassEntityQuery& SetupNativeQuery(ITypedElementDataStorageInterface::FQueryDescription& Query, FExtendedQuery& StoredQuery);
diff --git a/Engine/Plugins/Experimental/EditorDataStorage/Source/TedsCore/Private/TypedElementDatabase.cpp b/Engine/Plugins/Experimental/EditorDataStorage/Source/TedsCore/Private/TypedElementDatabase.cpp
--- a/Engine/Plugins/Experimental/EditorDataStorage/Source/TedsCore/Private/TypedElementDatabase.cpp
+++ b/Engine/Plugins/Experimental/EditorDataStorage/Source/TedsCore/Private/TypedElementDatabase.cpp
@@ -104,6 +104,8 @@ namespace TypedElementDatabasePrivate
 
 void UTypedElementDatabase::Initialize()
 {
+	using namespace TypedElementDataStorage;
+
 	check(GEditor);
 	UMassEntityEditorSubsystem* Mass = GEditor->GetEditorSubsystem<UMassEntityEditorSubsystem>();
 	check(Mass);
@@ -136,27 +138,27 @@ void UTypedElementDatabase::Initialize()
 
 			// Update external source to TEDS at the start of the phase.
 			RegisterTickGroup(GetQueryTickGroupName(EQueryTickGroups::SyncExternalToDataStorage),
-				Phase, {}, {}, false);
+				Phase, {}, {}, EExecutionMode::Threaded);
 			
 			// Default group.
 			RegisterTickGroup(GetQueryTickGroupName(EQueryTickGroups::Default),
-				Phase, {}, GetQueryTickGroupName(EQueryTickGroups::SyncExternalToDataStorage), false);
+				Phase, {}, GetQueryTickGroupName(EQueryTickGroups::SyncExternalToDataStorage), EExecutionMode::Threaded);
 
 			// Order the update groups.
 			RegisterTickGroup(GetQueryTickGroupName(EQueryTickGroups::PreUpdate),
-				Phase, {}, GetQueryTickGroupName(EQueryTickGroups::Default), false);
+				Phase, {}, GetQueryTickGroupName(EQueryTickGroups::Default), EExecutionMode::Threaded);
 			RegisterTickGroup(GetQueryTickGroupName(EQueryTickGroups::Update),
-				Phase, {}, GetQueryTickGroupName(EQueryTickGroups::PreUpdate), false);
+				Phase, {}, GetQueryTickGroupName(EQueryTickGroups::PreUpdate), EExecutionMode::Threaded);
 			RegisterTickGroup(GetQueryTickGroupName(EQueryTickGroups::PostUpdate),
-				Phase, {}, GetQueryTickGroupName(EQueryTickGroups::Update), false);
+				Phase, {}, GetQueryTickGroupName(EQueryTickGroups::Update), EExecutionMode::Threaded);
 
 			// After everything has processed sync the data in TEDS to external sources.
 			RegisterTickGroup(GetQueryTickGroupName(EQueryTickGroups::SyncDataStorageToExternal),
-				Phase, {}, GetQueryTickGroupName(EQueryTickGroups::PostUpdate), false);
+				Phase, {}, GetQueryTickGroupName(EQueryTickGroups::PostUpdate), EExecutionMode::Threaded);
 
 			// Update any widgets with data from TEDS.
 			RegisterTickGroup(GetQueryTickGroupName(EQueryTickGroups::SyncWidgets),
-				Phase, {}, GetQueryTickGroupName(EQueryTickGroups::PostUpdate), true /* Needs main thread*/);
+				Phase, {}, GetQueryTickGroupName(EQueryTickGroups::PostUpdate), EExecutionMode::GameThread /* Needs main thread*/);
 		}
 	}
 }
@@ -955,9 +957,9 @@ bool UTypedElementDatabase::MatchesColumns(TypedElementDataStorage::RowHandle Ro
 }
 
 void UTypedElementDatabase::RegisterTickGroup(
-	FName GroupName, EQueryTickPhase Phase, FName BeforeGroup, FName AfterGroup, bool bRequiresMainThread)
+	FName GroupName, EQueryTickPhase Phase, FName BeforeGroup, FName AfterGroup, TypedElementDataStorage::EExecutionMode ExecutionMode)
 {
-	Environment->GetQueryStore().RegisterTickGroup(GroupName, Phase, BeforeGroup, AfterGroup, bRequiresMainThread);
+	Environment->GetQueryStore().RegisterTickGroup(GroupName, Phase, BeforeGroup, AfterGroup, ExecutionMode);
 }
 
 void UTypedElementDatabase::UnregisterTickGroup(FName GroupName, EQueryTickPhase Phase)
@@ -1034,7 +1036,25 @@ ITypedElementDataStorageInterface::FQueryResult UTypedElementDatabase::RunQuery(
 	if (ActiveEditorEntityManager)
 	{
 		const UE::Editor::DataStorage::FExtendedQueryStore::Handle StorageHandle(Query);
-		return Environment->GetQueryStore().RunQuery(*ActiveEditorEntityManager, *Environment, StorageHandle, Callback);
+		return Environment->GetQueryStore().RunQuery(*ActiveEditorEntityManager, *Environment, StorageHandle, 
+			TypedElementDataStorage::EDirectQueryExecutionFlags::Default, Callback);
+	}
+	else
+	{
+		return FQueryResult();
+	}
+}
+
+ITypedElementDataStorageInterface::FQueryResult UTypedElementDatabase::RunQuery(
+	TypedElementQueryHandle Query, TypedElementDataStorage::EDirectQueryExecutionFlags Flags,
+	ITypedElementDataStorageInterface::DirectQueryCallbackRef Callback)
+{
+	TRACE_CPUPROFILER_EVENT_SCOPE(TEDS.RunQuery);
+
+	if (ActiveEditorEntityManager)
+	{
+		const UE::Editor::DataStorage::FExtendedQueryStore::Handle StorageHandle(Query);
+		return Environment->GetQueryStore().RunQuery(*ActiveEditorEntityManager, *Environment, StorageHandle, Flags, Callback);
 	}
 	else
 	{
diff --git a/Engine/Plugins/Experimental/EditorDataStorage/Source/TedsCore/Private/TypedElementDatabase.h b/Engine/Plugins/Experimental/EditorDataStorage/Source/TedsCore/Private/TypedElementDatabase.h
--- a/Engine/Plugins/Experimental/EditorDataStorage/Source/TedsCore/Private/TypedElementDatabase.h
+++ b/Engine/Plugins/Experimental/EditorDataStorage/Source/TedsCore/Private/TypedElementDatabase.h
@@ -127,7 +127,7 @@ public:
 	virtual void ListColumns(TypedElementDataStorage::RowHandle Row, TypedElementDataStorage::ColumnListWithDataCallbackRef Callback);
 	virtual bool MatchesColumns(TypedElementDataStorage::RowHandle Row, const TypedElementDataStorage::FQueryConditions& Conditions) const override;
 
-	void RegisterTickGroup(FName GroupName, EQueryTickPhase Phase, FName BeforeGroup, FName AfterGroup, bool bRequiresMainThread);
+	void RegisterTickGroup(FName GroupName, EQueryTickPhase Phase, FName BeforeGroup, FName AfterGroup, TypedElementDataStorage::EExecutionMode ExecutionMode);
 	void UnregisterTickGroup(FName GroupName, EQueryTickPhase Phase);
 
 	TypedElementQueryHandle RegisterQuery(FQueryDescription&& Query) override;
@@ -136,6 +136,8 @@ public:
 	virtual FName GetQueryTickGroupName(EQueryTickGroups Group) const override;
 	virtual FQueryResult RunQuery(TypedElementQueryHandle Query) override;
 	virtual FQueryResult RunQuery(TypedElementQueryHandle Query, DirectQueryCallbackRef Callback) override;
+	virtual FQueryResult RunQuery(TypedElementQueryHandle Query, TypedElementDataStorage::EDirectQueryExecutionFlags Flags, 
+		DirectQueryCallbackRef Callback) override;
 	virtual void ActivateQueries(FName ActivationName) override;
 
 	virtual TypedElementDataStorage::RowHandle FindIndexedRow(TypedElementDataStorage::IndexHash Index) const override;
diff --git a/Engine/Plugins/Experimental/EditorDataStorage/Source/TedsUI/Private/Widgets/CounterWidget.cpp b/Engine/Plugins/Experimental/EditorDataStorage/Source/TedsUI/Private/Widgets/CounterWidget.cpp
--- a/Engine/Plugins/Experimental/EditorDataStorage/Source/TedsUI/Private/Widgets/CounterWidget.cpp
+++ b/Engine/Plugins/Experimental/EditorDataStorage/Source/TedsUI/Private/Widgets/CounterWidget.cpp
@@ -46,20 +46,20 @@ UCounterWidgetFactory::UCounterWidgetFactory()
 void UCounterWidgetFactory::RegisterQueries(ITypedElementDataStorageInterface& DataStorage)
 {
 	using namespace TypedElementQueryBuilder;
-	using DSI = ITypedElementDataStorageInterface;
-
+	using namespace TypedElementDataStorage;
+	
 	DataStorage.RegisterQuery(Select(TEXT("Sync counter widgets"), 
-		FProcessor(DSI::EQueryTickPhase::FrameEnd, DataStorage.GetQueryTickGroupName(DSI::EQueryTickGroups::SyncWidgets))
-			.ForceToGameThread(true),
+		FProcessor(EQueryTickPhase::FrameEnd, DataStorage.GetQueryTickGroupName(EQueryTickGroups::SyncWidgets))
+			.SetExecutionMode(EExecutionMode::GameThread),
 		[](
-			DSI::IQueryContext& Context,
+			IQueryContext& Context,
 			FTypedElementSlateWidgetReferenceColumn& Widget,
 			FTypedElementU32IntValueCacheColumn& Comparison, 
 			const FCounterWidgetColumn& Counter
 		)
 		{
-			DSI::FQueryResult Result = Context.RunQuery(Counter.Query);
-			if (Result.Completed == DSI::FQueryResult::ECompletion::Fully && Result.Count != Comparison.Value)
+			FQueryResult Result = Context.RunQuery(Counter.Query);
+			if (Result.Completed == FQueryResult::ECompletion::Fully && Result.Count != Comparison.Value)
 			{
 				TSharedPtr<SWidget> WidgetPointer = Widget.Widget.Pin();
 				checkf(WidgetPointer, TEXT("Referenced widget is not valid. A constructed widget may not have been cleaned up. This can "
diff --git a/Engine/Plugins/Experimental/EditorDataStorage/Source/TedsUI/Private/Widgets/ExportedTextWidget.cpp b/Engine/Plugins/Experimental/EditorDataStorage/Source/TedsUI/Private/Widgets/ExportedTextWidget.cpp
--- a/Engine/Plugins/Experimental/EditorDataStorage/Source/TedsUI/Private/Widgets/ExportedTextWidget.cpp
+++ b/Engine/Plugins/Experimental/EditorDataStorage/Source/TedsUI/Private/Widgets/ExportedTextWidget.cpp
@@ -51,8 +51,7 @@ static void UpdateExportedTextWidget(ITypedElementDataStorageInterface& DataStor
 static TypedElementQueryHandle RegisterUpdateCallback(ITypedElementDataStorageInterface& DataStorage, const UScriptStruct* Target)
 {
 	using namespace TypedElementQueryBuilder;
-	using DSI = ITypedElementDataStorageInterface;
-	namespace DS = TypedElementDataStorage;
+	using namespace TypedElementDataStorage;
 	
 	TypedElementQueryHandle TypeDataQuery = DataStorage.RegisterQuery(
 		Select()
@@ -69,10 +68,10 @@ static TypedElementQueryHandle RegisterUpdateCallback(ITypedElementDataStorageIn
 
 	return DataStorage.RegisterQuery(
 		Select(ProcessorName,
-			FProcessor(DSI::EQueryTickPhase::FrameEnd, DataStorage.GetQueryTickGroupName(DSI::EQueryTickGroups::SyncWidgets))
-				.ForceToGameThread(true),
+			FProcessor(EQueryTickPhase::FrameEnd, DataStorage.GetQueryTickGroupName(EQueryTickGroups::SyncWidgets))
+				.SetExecutionMode(EExecutionMode::GameThread),
 			[Target, ProcessorName](
-				DS::IQueryContext& Context, 
+				IQueryContext& Context, 
 				FTypedElementSlateWidgetReferenceColumn& Widget,
 				const FTypedElementScriptStructTypeInfoColumn& TypeInfo,
 				const FTypedElementRowReferenceColumn& ReferencedRow)
@@ -96,7 +95,7 @@ static TypedElementQueryHandle RegisterUpdateCallback(ITypedElementDataStorageIn
 				if(StructType == Target)
 				{
 					Context.RunSubquery(0, ReferencedRow.Row, 
-					[&Widget, StructType](const DS::FQueryDescription&, DS::ISubqueryContext& SubqueryContext)
+					[&Widget, StructType](const FQueryDescription&, ISubqueryContext& SubqueryContext)
 					{
 						const void* ColumnData = SubqueryContext.GetColumn(StructType);
 						UpdateExportedTextWidget(ColumnData, Widget, StructType);
diff --git a/Engine/Plugins/Experimental/EditorDataStorage/Source/TedsUI/Private/Widgets/LabelWidget.cpp b/Engine/Plugins/Experimental/EditorDataStorage/Source/TedsUI/Private/Widgets/LabelWidget.cpp
--- a/Engine/Plugins/Experimental/EditorDataStorage/Source/TedsUI/Private/Widgets/LabelWidget.cpp
+++ b/Engine/Plugins/Experimental/EditorDataStorage/Source/TedsUI/Private/Widgets/LabelWidget.cpp
@@ -75,8 +75,8 @@ static void SyncColumnsToWidget(
 void ULabelWidgetFactory::RegisterQueries(ITypedElementDataStorageInterface& DataStorage)
 {
 	using namespace TypedElementQueryBuilder;
-	using DSI = ITypedElementDataStorageInterface;
-
+	using namespace TypedElementDataStorage;
+	
 	TypedElementQueryHandle UpdateLabelWidget = DataStorage.RegisterQuery(
 		Select()
 			.ReadOnly<FTypedElementLabelColumn>()
@@ -95,10 +95,10 @@ void ULabelWidgetFactory::RegisterQueries(ITypedElementDataStorageInterface& Dat
 	DataStorage.RegisterQuery(
 		Select(
 			TEXT("Sync label to widget"),
-			FProcessor(DSI::EQueryTickPhase::FrameEnd, DataStorage.GetQueryTickGroupName(DSI::EQueryTickGroups::SyncWidgets))
-				.ForceToGameThread(true),
+			FProcessor(EQueryTickPhase::FrameEnd, DataStorage.GetQueryTickGroupName(EQueryTickGroups::SyncWidgets))
+				.SetExecutionMode(EExecutionMode::GameThread),
 			[](
-				DSI::IQueryContext& Context, 
+				IQueryContext& Context, 
 				FTypedElementSlateWidgetReferenceColumn& Widget,
 				FTypedElementU64IntValueCacheColumn& TextHash,
 				const FLabelWidgetColumn& Config,
diff --git a/Engine/Plugins/Experimental/EditorDataStorageFeatures/Source/TedsDebugger/Private/Widgets/RowReferenceWidget.cpp b/Engine/Plugins/Experimental/EditorDataStorageFeatures/Source/TedsDebugger/Private/Widgets/RowReferenceWidget.cpp
--- a/Engine/Plugins/Experimental/EditorDataStorageFeatures/Source/TedsDebugger/Private/Widgets/RowReferenceWidget.cpp
+++ b/Engine/Plugins/Experimental/EditorDataStorageFeatures/Source/TedsDebugger/Private/Widgets/RowReferenceWidget.cpp
@@ -109,8 +109,8 @@ void URowReferenceWidgetFactory::RegisterWidgetConstructors(ITypedElementDataSto
 void URowReferenceWidgetFactory::RegisterQueries(ITypedElementDataStorageInterface& DataStorage)
 {
 	using namespace TypedElementQueryBuilder;
-	using DSI = ITypedElementDataStorageInterface;
-
+	using namespace TypedElementDataStorage;
+	
 	const TypedElementQueryHandle UpdateRowReferenceWidget = DataStorage.RegisterQuery(
 		Select()
 			.ReadOnly<FTypedElementRowReferenceColumn>()
@@ -121,10 +121,10 @@ void URowReferenceWidgetFactory::RegisterQueries(ITypedElementDataStorageInterfa
 	DataStorage.RegisterQuery(
 		Select(
 			TEXT("Sync row reference to widget"),
-			FProcessor(DSI::EQueryTickPhase::FrameEnd, DataStorage.GetQueryTickGroupName(DSI::EQueryTickGroups::SyncWidgets))
-				.ForceToGameThread(true),
+			FProcessor(EQueryTickPhase::FrameEnd, DataStorage.GetQueryTickGroupName(EQueryTickGroups::SyncWidgets))
+				.SetExecutionMode(EExecutionMode::GameThread),
 			[](
-				DSI::IQueryContext& Context,
+				IQueryContext& Context,
 				TypedElementDataStorage::RowHandle UiRowHandle,
 				FTypedElementSlateWidgetReferenceColumn& Widget,
 				const FTypedElementRowReferenceColumn& Target)
diff --git a/Engine/Plugins/Experimental/EditorDataStorageFeatures/Source/TedsOutliner/Private/TedsOutlinerImpl.cpp b/Engine/Plugins/Experimental/EditorDataStorageFeatures/Source/TedsOutliner/Private/TedsOutlinerImpl.cpp
--- a/Engine/Plugins/Experimental/EditorDataStorageFeatures/Source/TedsOutliner/Private/TedsOutlinerImpl.cpp
+++ b/Engine/Plugins/Experimental/EditorDataStorageFeatures/Source/TedsOutliner/Private/TedsOutlinerImpl.cpp
@@ -575,7 +575,7 @@ void FTedsOutlinerImpl::RecompileQueries()
 	FQueryDescription RowAdditionQueryDescription =
 		Select(
 				TEXT("Add Row to Outliner"),
-				FObserver::OnAdd<FTypedElementLabelColumn>().ForceToGameThread(true),
+				FObserver::OnAdd<FTypedElementLabelColumn>().SetExecutionMode(EExecutionMode::GameThread),
 				[this](IQueryContext& Context, TypedElementRowHandle Row)
 				{
 					OnItemAdded(Row);
@@ -589,7 +589,7 @@ void FTedsOutlinerImpl::RecompileQueries()
 	FQueryDescription RowRemovalQueryDescription =
 		Select(
 				TEXT("Remove Row from Outliner"),
-				FObserver::OnRemove<FTypedElementLabelColumn>().ForceToGameThread(true),
+				FObserver::OnRemove<FTypedElementLabelColumn>().SetExecutionMode(EExecutionMode::GameThread),
 				[this](IQueryContext& Context, TypedElementRowHandle Row)
 				{
 					OnItemRemoved(Row);
@@ -618,7 +618,7 @@ void FTedsOutlinerImpl::RecompileQueries()
 			Select(
 			TEXT("Update item parent"),
 			FProcessor(EQueryTickPhase::DuringPhysics, Storage->GetQueryTickGroupName(EQueryTickGroups::Update))
-				.ForceToGameThread(true),
+				.SetExecutionMode(EExecutionMode::GameThread),
 			[this](IQueryContext& Context, TypedElementDataStorage::RowHandle Row)
 			{
 				TypedElementDataStorage::RowHandle ParentRowHandle = InvalidRowHandle;
@@ -658,7 +658,7 @@ void FTedsOutlinerImpl::RecompileQueries()
 		FQueryDescription SelectionAddedQueryDescription =
 							Select(
 							TEXT("Row selected"),
-							FObserver::OnAdd<FTypedElementSelectionColumn>().ForceToGameThread(true),
+							FObserver::OnAdd<FTypedElementSelectionColumn>().SetExecutionMode(EExecutionMode::GameThread),
 							[this](IQueryContext& Context, TypedElementRowHandle Row)
 							{
 								bSelectionDirty = true;
@@ -672,7 +672,7 @@ void FTedsOutlinerImpl::RecompileQueries()
 		FQueryDescription SelectionRemovedQueryDescription =
 							Select(
 							TEXT("Row deselected"),
-							FObserver::OnRemove<FTypedElementSelectionColumn>().ForceToGameThread(true),
+							FObserver::OnRemove<FTypedElementSelectionColumn>().SetExecutionMode(EExecutionMode::GameThread),
 							[this](IQueryContext& Context, TypedElementRowHandle Row)
 							{
 								bSelectionDirty = true;
diff --git a/Engine/Plugins/Experimental/EditorDataStorageFeatures/Source/TedsOutliner/Private/Widgets/AlertWidget.cpp b/Engine/Plugins/Experimental/EditorDataStorageFeatures/Source/TedsOutliner/Private/Widgets/AlertWidget.cpp
--- a/Engine/Plugins/Experimental/EditorDataStorageFeatures/Source/TedsOutliner/Private/Widgets/AlertWidget.cpp
+++ b/Engine/Plugins/Experimental/EditorDataStorageFeatures/Source/TedsOutliner/Private/Widgets/AlertWidget.cpp
@@ -170,7 +170,6 @@ void UAlertWidgetFactory::RegisterAlertQueries(ITypedElementDataStorageInterface
 {
 	using namespace TypedElementQueryBuilder;
 	using namespace TypedElementDataStorage;
-	using DSI = ITypedElementDataStorageInterface;
 	
 	TypedElementQueryHandle UpdateWidget_OnlyAlert = DataStorage.RegisterQuery(
 		Select()
@@ -197,8 +196,8 @@ void UAlertWidgetFactory::RegisterAlertQueries(ITypedElementDataStorageInterface
 	
 	DataStorage.RegisterQuery(
 		Select(TEXT("Sync Transform column to heads up display"),
-			FProcessor(DSI::EQueryTickPhase::FrameEnd, DataStorage.GetQueryTickGroupName(DSI::EQueryTickGroups::SyncWidgets))
-			.ForceToGameThread(true),
+			FProcessor(EQueryTickPhase::FrameEnd, DataStorage.GetQueryTickGroupName(EQueryTickGroups::SyncWidgets))
+				.SetExecutionMode(EExecutionMode::GameThread),
 			[](IQueryContext& Context,
 				FTypedElementSlateWidgetReferenceColumn& Widget,
 				const FTypedElementRowReferenceColumn& ReferenceColumn)
@@ -253,7 +252,6 @@ void UAlertWidgetFactory::RegisterAlertHeaderQueries(ITypedElementDataStorageInt
 {
 	using namespace TypedElementQueryBuilder;
 	using namespace TypedElementDataStorage;
-	using DSI = ITypedElementDataStorageInterface;
 	
 	TypedElementQueryHandle AlertCount = DataStorage.RegisterQuery(
 		Count()
@@ -263,8 +261,8 @@ void UAlertWidgetFactory::RegisterAlertHeaderQueries(ITypedElementDataStorageInt
 	
 	DataStorage.RegisterQuery(
 		Select(TEXT("Update alert header"),
-			FProcessor(DSI::EQueryTickPhase::FrameEnd, DataStorage.GetQueryTickGroupName(DSI::EQueryTickGroups::SyncWidgets))
-			.ForceToGameThread(true),
+			FProcessor(EQueryTickPhase::FrameEnd, DataStorage.GetQueryTickGroupName(EQueryTickGroups::SyncWidgets))
+				.SetExecutionMode(EExecutionMode::GameThread),
 			[](IQueryContext& Context, RowHandle Row, FTypedElementSlateWidgetReferenceColumn& Widget)
 			{
 				FQueryResult Result = Context.RunSubquery(0);
@@ -287,8 +285,8 @@ void UAlertWidgetFactory::RegisterAlertHeaderQueries(ITypedElementDataStorageInt
 
 	DataStorage.RegisterQuery(
 		Select(TEXT("Update active alert header"),
-			FProcessor(DSI::EQueryTickPhase::FrameEnd, DataStorage.GetQueryTickGroupName(DSI::EQueryTickGroups::SyncWidgets))
-			.ForceToGameThread(true),
+			FProcessor(EQueryTickPhase::FrameEnd, DataStorage.GetQueryTickGroupName(EQueryTickGroups::SyncWidgets))
+				.SetExecutionMode(EExecutionMode::GameThread),
 			[](IQueryContext& Context, RowHandle Row, FTypedElementSlateWidgetReferenceColumn& Widget)
 			{
 				FQueryResult Result = Context.RunSubquery(0);
diff --git a/Engine/Plugins/Experimental/EditorDataStorageFeatures/Source/TedsOutliner/Private/Widgets/OverrideWidget.cpp b/Engine/Plugins/Experimental/EditorDataStorageFeatures/Source/TedsOutliner/Private/Widgets/OverrideWidget.cpp
--- a/Engine/Plugins/Experimental/EditorDataStorageFeatures/Source/TedsOutliner/Private/Widgets/OverrideWidget.cpp
+++ b/Engine/Plugins/Experimental/EditorDataStorageFeatures/Source/TedsOutliner/Private/Widgets/OverrideWidget.cpp
@@ -31,8 +31,8 @@ void UOverrideWidgetFactory::RegisterWidgetConstructors(
 void UOverrideWidgetFactory::RegisterQueries(ITypedElementDataStorageInterface& DataStorage)
 {
 	using namespace TypedElementQueryBuilder;
-	using DSI = ITypedElementDataStorageInterface;
-
+	using namespace TypedElementDataStorage;
+	
 	const TypedElementQueryHandle UpdateWidget = DataStorage.RegisterQuery(
 		Select()
 		.Where()
@@ -43,10 +43,10 @@ void UOverrideWidgetFactory::RegisterQueries(ITypedElementDataStorageInterface&
 	DataStorage.RegisterQuery(
 		Select(
 			TEXT("Sync override status to widget"),
-			FProcessor(DSI::EQueryTickPhase::FrameEnd, DataStorage.GetQueryTickGroupName(DSI::EQueryTickGroups::SyncWidgets))
-				.ForceToGameThread(true),
+			FProcessor(EQueryTickPhase::FrameEnd, DataStorage.GetQueryTickGroupName(EQueryTickGroups::SyncWidgets))
+				.SetExecutionMode(EExecutionMode::GameThread),
 			[](
-				DSI::IQueryContext& Context, 
+				IQueryContext& Context, 
 				const FTypedElementSlateWidgetReferenceColumn& Widget,
 				const FTypedElementRowReferenceColumn& Target)
 				{
diff --git a/Engine/Plugins/Experimental/EditorDataStorageFeatures/Source/TedsOutliner/Private/Widgets/TransformHeadsUpWidget.cpp b/Engine/Plugins/Experimental/EditorDataStorageFeatures/Source/TedsOutliner/Private/Widgets/TransformHeadsUpWidget.cpp
--- a/Engine/Plugins/Experimental/EditorDataStorageFeatures/Source/TedsOutliner/Private/Widgets/TransformHeadsUpWidget.cpp
+++ b/Engine/Plugins/Experimental/EditorDataStorageFeatures/Source/TedsOutliner/Private/Widgets/TransformHeadsUpWidget.cpp
@@ -154,8 +154,7 @@ static void UpdateTransformHeadsUpDisplay(FTypedElementSlateWidgetReferenceColum
 void UTransformHeadsUpWidgetFactory::RegisterQueries(ITypedElementDataStorageInterface& DataStorage)
 {
 	using namespace TypedElementQueryBuilder;
-	using DSI = ITypedElementDataStorageInterface;
-	namespace DS = TypedElementDataStorage;
+	using namespace TypedElementDataStorage;
 		
 	TypedElementQueryHandle UpdateTransformWidget = DataStorage.RegisterQuery(
 		Select()
@@ -166,9 +165,9 @@ void UTransformHeadsUpWidgetFactory::RegisterQueries(ITypedElementDataStorageInt
 
 	DataStorage.RegisterQuery(
 		Select(TEXT("Sync Transform column to heads up display"),
-		FProcessor(DSI::EQueryTickPhase::FrameEnd, DataStorage.GetQueryTickGroupName(DSI::EQueryTickGroups::SyncWidgets))
-			.ForceToGameThread(true),
-			[](DS::IQueryContext& Context,
+		FProcessor(EQueryTickPhase::FrameEnd, DataStorage.GetQueryTickGroupName(EQueryTickGroups::SyncWidgets))
+			.SetExecutionMode(EExecutionMode::GameThread),
+			[](IQueryContext& Context,
 				FTypedElementSlateWidgetReferenceColumn& Widget,
 				const FTypedElementRowReferenceColumn& ReferenceColumn)
 			{
diff --git a/Engine/Plugins/Experimental/EditorDataStorageFeatures/Source/TedsRevisionControl/Private/Processors/RevisionControlProcessors.cpp b/Engine/Plugins/Experimental/EditorDataStorageFeatures/Source/TedsRevisionControl/Private/Processors/RevisionControlProcessors.cpp
--- a/Engine/Plugins/Experimental/EditorDataStorageFeatures/Source/TedsRevisionControl/Private/Processors/RevisionControlProcessors.cpp
+++ b/Engine/Plugins/Experimental/EditorDataStorageFeatures/Source/TedsRevisionControl/Private/Processors/RevisionControlProcessors.cpp
@@ -231,7 +231,7 @@ void URevisionControlDataStorageFactory::RegisterQueries(ITypedElementDataStorag
 void URevisionControlDataStorageFactory::RegisterFetchUpdates(ITypedElementDataStorageInterface& DataStorage)
 {
 	using namespace TypedElementQueryBuilder;
-	using DSI = ITypedElementDataStorageInterface;
+	using namespace TypedElementDataStorage;
 	
 	FSourceControlFileStatusMonitor& FileStatusMonitor = ISourceControlModule::Get().GetSourceControlFileStatusMonitor();
 
@@ -241,8 +241,8 @@ void URevisionControlDataStorageFactory::RegisterFetchUpdates(ITypedElementDataS
 			Select(
 				TEXT("Gather source control statuses for objects with unresolved package paths"),
 				FObserver::OnAdd<FTypedElementPackageUnresolvedReference>()
-				.ForceToGameThread(true),
-				[this, &FileStatusMonitor](DSI::IQueryContext& Context, const FTypedElementPackageUnresolvedReference& UnresolvedReference)
+					.SetExecutionMode(EExecutionMode::GameThread),
+				[this, &FileStatusMonitor](IQueryContext& Context, const FTypedElementPackageUnresolvedReference& UnresolvedReference)
 				{
 					static FSourceControlFileStatusMonitor::FOnSourceControlFileStatus EmptyDelegate{};
 				
@@ -263,8 +263,7 @@ void URevisionControlDataStorageFactory::RegisterApplyOverlays(ITypedElementData
 {
 	using namespace TypedElementDataStorage;
 	using namespace TypedElementQueryBuilder;
-	using DSI = ITypedElementDataStorageInterface;
-
+	
 	if (ApplyOverlaysObjectToSCC == InvalidQueryHandle)
 	{
 		ApplyOverlaysObjectToSCC = DataStorage.RegisterQuery(
@@ -317,9 +316,9 @@ void URevisionControlDataStorageFactory::RegisterApplyOverlays(ITypedElementData
 		FlushPackageUpdates = DataStorage.RegisterQuery(
 			Select(
 				TEXT("Consume collected package updates"),
-				FProcessor(DSI::EQueryTickPhase::PrePhysics, DataStorage.GetQueryTickGroupName(DSI::EQueryTickGroups::Update))
-				.ForceToGameThread(true),
-				[](DSI::IQueryContext& Context, RowHandle Row, const FTypedElementPackageUpdateColumn& Update)
+				FProcessor(EQueryTickPhase::PrePhysics, DataStorage.GetQueryTickGroupName(EQueryTickGroups::Update))
+					.SetExecutionMode(EExecutionMode::GameThread),
+				[](IQueryContext& Context, RowHandle Row, const FTypedElementPackageUpdateColumn& Update)
 				{
 					// Query:
 					// For all actors without an overlay color column AND having a package reference:
@@ -329,7 +328,7 @@ void URevisionControlDataStorageFactory::RegisterApplyOverlays(ITypedElementData
 						[&Context](RowHandle ObjectRow, const FTypedElementUObjectColumn& Actor, const FTypedElementPackageReference& PackageReference)
 						{
 							Context.RunSubquery(EApplyOverlaysObjectToSCC, PackageReference.Row, CreateSubqueryCallbackBinding(
-								[&Context, &ObjectRow, &Actor](DSI::ISubqueryContext& SubQueryContext)
+								[&Context, &ObjectRow, &Actor](ISubqueryContext& SubQueryContext)
 								{
 									FColor Color = DetermineOverlayColor(Context, SubQueryContext, Actor);
 									if (Color.Bits != 0)
@@ -351,7 +350,7 @@ void URevisionControlDataStorageFactory::RegisterApplyOverlays(ITypedElementData
 						[&Context](TypedElementRowHandle ObjectRow, const FTypedElementUObjectColumn& Actor, const FTypedElementPackageReference& PackageReference, const FTypedElementViewportOverlayColorColumn& OverlayColorColumn)
 						{
 							Context.RunSubquery(EApplyOverlaysObjectToSCC, PackageReference.Row, CreateSubqueryCallbackBinding(
-								[&Context, &ObjectRow, &Actor, &OverlayColorColumn](DSI::ISubqueryContext& SubQueryContext)
+								[&Context, &ObjectRow, &Actor, &OverlayColorColumn](ISubqueryContext& SubQueryContext)
 								{
 									FColor Color = DetermineOverlayColor(Context, SubQueryContext, Actor);
 									if (Color.Bits == 0)
@@ -381,8 +380,8 @@ void URevisionControlDataStorageFactory::RegisterApplyOverlays(ITypedElementData
 void URevisionControlDataStorageFactory::RegisterRemoveOverlays(ITypedElementDataStorageInterface& DataStorage)
 {
 	using namespace TypedElementQueryBuilder;
-	using DSI = ITypedElementDataStorageInterface;
-
+	using namespace TypedElementDataStorage;
+	
 	if (RemoveOverlays == TypedElementInvalidQueryHandle)
 	{
 		// Query:
@@ -394,9 +393,9 @@ void URevisionControlDataStorageFactory::RegisterRemoveOverlays(ITypedElementDat
 			Select(
 				TEXT("Remove selection overlay colors"),
 				// This is in PrePhysics because the overlay->actor query is in DuringPhysics and contexts don't flush changes between tick groups
-				FProcessor(DSI::EQueryTickPhase::PrePhysics, DataStorage.GetQueryTickGroupName(DSI::EQueryTickGroups::SyncExternalToDataStorage))
-				.ForceToGameThread(true),
-				[](DSI::IQueryContext& Context, TypedElementRowHandle ObjectRow, FTypedElementUObjectColumn& Actor, const FTypedElementViewportOverlayColorColumn& ViewportColor)
+				FProcessor(EQueryTickPhase::PrePhysics, DataStorage.GetQueryTickGroupName(EQueryTickGroups::SyncExternalToDataStorage))
+					.SetExecutionMode(EExecutionMode::GameThread),
+				[](IQueryContext& Context, TypedElementRowHandle ObjectRow, FTypedElementUObjectColumn& Actor, const FTypedElementViewportOverlayColorColumn& ViewportColor)
 				{
 					Context.RemoveColumns<FTypedElementViewportOverlayColorColumn>(ObjectRow);
 				}
diff --git a/Engine/Plugins/Experimental/EditorDataStorageFeatures/Source/TedsRevisionControl/Private/Queries/ObjectPackagePathToColumnQueries.cpp b/Engine/Plugins/Experimental/EditorDataStorageFeatures/Source/TedsRevisionControl/Private/Queries/ObjectPackagePathToColumnQueries.cpp
--- a/Engine/Plugins/Experimental/EditorDataStorageFeatures/Source/TedsRevisionControl/Private/Queries/ObjectPackagePathToColumnQueries.cpp
+++ b/Engine/Plugins/Experimental/EditorDataStorageFeatures/Source/TedsRevisionControl/Private/Queries/ObjectPackagePathToColumnQueries.cpp
@@ -41,8 +41,8 @@ static void ResolvePackageReference(ITypedElementDataStorageInterface::IQueryCon
 void UTypedElementUObjectPackagePathFactory::RegisterQueries(ITypedElementDataStorageInterface& DataStorage)
 {
 	using namespace TypedElementQueryBuilder;
-	using DSI = ITypedElementDataStorageInterface;
-
+	using namespace TypedElementDataStorage;
+	
 	CVarAutoPopulateState->AsVariable()->OnChangedDelegate().AddLambda(
 		[this, &DataStorage](IConsoleVariable* AutoPopulate)
 		{
@@ -59,8 +59,8 @@ void UTypedElementUObjectPackagePathFactory::RegisterQueries(ITypedElementDataSt
 	DataStorage.RegisterQuery(
 		Select(
 			TEXT("Resolve package references"),
-			FProcessor(DSI::EQueryTickPhase::FrameEnd, DataStorage.GetQueryTickGroupName(DSI::EQueryTickGroups::SyncExternalToDataStorage)),
-			[](DSI::IQueryContext& Context, TypedElementRowHandle Row, const FTypedElementUObjectColumn& Object, const FTypedElementPackageUnresolvedReference& UnresolvedPackageReference)
+			FProcessor(EQueryTickPhase::FrameEnd, DataStorage.GetQueryTickGroupName(EQueryTickGroups::SyncExternalToDataStorage)),
+			[](IQueryContext& Context, TypedElementRowHandle Row, const FTypedElementUObjectColumn& Object, const FTypedElementPackageUnresolvedReference& UnresolvedPackageReference)
 			{
 				TypedElementRowHandle PackageRow = Context.FindIndexedRow(UnresolvedPackageReference.Index);
 				if (!Context.IsRowAvailable(PackageRow))
@@ -90,14 +90,14 @@ void UTypedElementUObjectPackagePathFactory::RegisterQueries(ITypedElementDataSt
 void UTypedElementUObjectPackagePathFactory::RegisterTryAddPackageRef(ITypedElementDataStorageInterface& DataStorage)
 {
 	using namespace TypedElementQueryBuilder;
-	using DSI = ITypedElementDataStorageInterface;
+	using namespace TypedElementDataStorage;
 	
 	TryAddPackageRef = DataStorage.RegisterQuery(
 		Select(
 			TEXT("Sync UObject package info to columns"),
 			FObserver::OnAdd<FTypedElementUObjectColumn>()
-				.ForceToGameThread(true),
-			[](DSI::IQueryContext& Context, TypedElementRowHandle Row, const FTypedElementUObjectColumn& Object)
+				.SetExecutionMode(EExecutionMode::GameThread),
+			[](IQueryContext& Context, TypedElementRowHandle Row, const FTypedElementUObjectColumn& Object)
 			{
 				if (const UObject* ObjectInstance = Object.Object.Get(); ObjectInstance != nullptr)
 				{
diff --git a/Engine/Plugins/Experimental/EditorDataStorageFeatures/Source/TedsTableViewer/Private/TedsTableViewerColumn.cpp b/Engine/Plugins/Experimental/EditorDataStorageFeatures/Source/TedsTableViewer/Private/TedsTableViewerColumn.cpp
--- a/Engine/Plugins/Experimental/EditorDataStorageFeatures/Source/TedsTableViewer/Private/TedsTableViewerColumn.cpp
+++ b/Engine/Plugins/Experimental/EditorDataStorageFeatures/Source/TedsTableViewer/Private/TedsTableViewerColumn.cpp
@@ -166,7 +166,7 @@ namespace UE::EditorDataStorage
 		{
 			const FName ColumnAddObserverName = *(FString::Printf(TEXT("Column Add Monitor for %s Table Viewer Column Column, %s TEDS Column"), *ColumnName.ToString(), *ColumnType->GetName()));
 			FObserver AddObserver(FObserver::EEvent::Add, ColumnType.Get());
-			AddObserver.ForceToGameThread(true);
+			AddObserver.SetExecutionMode(EExecutionMode::GameThread);
 
 			// TEDS-Outliner TODO: Long term if we move this into TypedElementOutlinerMode or similar we can get access to the exact
 			// types the Outliner is looking at and specify them on .Where() to cut down on the things we are observing
@@ -187,7 +187,7 @@ namespace UE::EditorDataStorage
 
 			const FName ColumnRemoveObserverName = *(FString::Printf(TEXT("Column Remove Monitor for %s Table Viewer Column Column, %s TEDS Column"), *ColumnName.ToString(), *ColumnType->GetName()));
 			FObserver RemoveObserver(FObserver::EEvent::Remove, ColumnType.Get());
-			RemoveObserver.ForceToGameThread(true);
+			RemoveObserver.SetExecutionMode(EExecutionMode::GameThread);
 
 			// Table Viewer TODO: We might be able to cut down on the rows we are querying for in the future by getting the rows from the query stack
 			// but we currently have to use a generic query so we can support the TEDS-Outliner as well
diff --git a/Engine/Source/Runtime/TypedElementFramework/Private/Elements/Framework/TypedElementQueryBuilder.cpp b/Engine/Source/Runtime/TypedElementFramework/Private/Elements/Framework/TypedElementQueryBuilder.cpp
--- a/Engine/Source/Runtime/TypedElementFramework/Private/Elements/Framework/TypedElementQueryBuilder.cpp
+++ b/Engine/Source/Runtime/TypedElementFramework/Private/Elements/Framework/TypedElementQueryBuilder.cpp
@@ -91,12 +91,16 @@ namespace TypedElementQueryBuilder
 
 	FDependency& FDependency::SubQuery(TypedElementQueryHandle Handle)
 	{
+		checkf(Query->Callback.ExecutionMode != TypedElementDataStorage::EExecutionMode::ThreadedChunks,
+			TEXT("TEDS sub-queries can not be added to queries with a callback that process chunks in parallel."));
 		Query->Subqueries.Add(Handle);
 		return *this;
 	}
 	
 	FDependency& FDependency::SubQuery(TConstArrayView<TypedElementQueryHandle> Handles)
 	{
+		checkf(Query->Callback.ExecutionMode != TypedElementDataStorage::EExecutionMode::ThreadedChunks,
+			TEXT("TEDS sub-queries can not be added to queries with a callback that process chunks in parallel."));
 		Query->Subqueries.Insert(Handles.GetData(), Handles.Num(), Query->Subqueries.Num());
 		return *this;
 	}
@@ -290,9 +294,9 @@ namespace TypedElementQueryBuilder
 		return *this;
 	}
 
-	FProcessor& FProcessor::ForceToGameThread(bool bForce)
+	FProcessor& FProcessor::SetExecutionMode(TypedElementDataStorage::EExecutionMode Mode)
 	{
-		bForceToGameThread = bForce;
+		ExecutionMode = Mode;
 		return *this;
 	}
 
@@ -324,9 +328,9 @@ namespace TypedElementQueryBuilder
 		return *this;
 	}
 
-	FObserver& FObserver::ForceToGameThread(bool bForce)
+	FObserver& FObserver::SetExecutionMode(TypedElementDataStorage::EExecutionMode Mode)
 	{
-		bForceToGameThread = bForce;
+		ExecutionMode = Mode;
 		return *this;
 	}
 
@@ -358,9 +362,9 @@ namespace TypedElementQueryBuilder
 		return *this;
 	}
 
-	FPhaseAmble& FPhaseAmble::ForceToGameThread(bool bForce)
+	FPhaseAmble& FPhaseAmble::SetExecutionMode(TypedElementDataStorage::EExecutionMode Mode)
 	{
-		bForceToGameThread = bForce;
+		ExecutionMode = Mode;
 		return *this;
 	}
 
diff --git a/Engine/Source/Runtime/TypedElementFramework/Public/Elements/Common/TypedElementQueryDescription.h b/Engine/Source/Runtime/TypedElementFramework/Public/Elements/Common/TypedElementQueryDescription.h
--- a/Engine/Source/Runtime/TypedElementFramework/Public/Elements/Common/TypedElementQueryDescription.h
+++ b/Engine/Source/Runtime/TypedElementFramework/Public/Elements/Common/TypedElementQueryDescription.h
@@ -82,8 +82,8 @@ namespace TypedElementDataStorage
 			 * The number of remaining iterations for a activatable query callback. If this is higher than 0, the query callback will be 
 			 * called. If ActivationName is set, this value will be decremented by one at the end of the update cycle.
 			 */
-			uint8 ActivationCount{ 255 };
-			bool bForceToGameThread{ false };
+			uint8 ActivationCount = 255;
+			EExecutionMode ExecutionMode = EExecutionMode::Default;
 		};
 		FCallbackData Callback;
 
diff --git a/Engine/Source/Runtime/TypedElementFramework/Public/Elements/Common/TypedElementQueryTypes.h b/Engine/Source/Runtime/TypedElementFramework/Public/Elements/Common/TypedElementQueryTypes.h
--- a/Engine/Source/Runtime/TypedElementFramework/Public/Elements/Common/TypedElementQueryTypes.h
+++ b/Engine/Source/Runtime/TypedElementFramework/Public/Elements/Common/TypedElementQueryTypes.h
@@ -75,6 +75,14 @@ namespace TypedElementDataStorage
 		Max //< Value indicating the maximum value in this enum. Not to be used as an enum value.
 	};
 
+	enum class EExecutionMode : uint8
+	{
+		Default, //< Use the default for the query callback type or the group default.
+		GameThread, //< The query callback will be executed sequentially on the game thread.
+		Threaded, //< The entire query callback is called on a separate thread, while respecting dependencies.
+		ThreadedChunks //< The query callback is called on a separate thread for each chunk. This can not be set if there are sub-queries.
+	};
+
 	enum class EQueryAccessType : uint8
 	{
 		/** Provide read-only access. */
diff --git a/Engine/Source/Runtime/TypedElementFramework/Public/Elements/Framework/TypedElementQueryBuilder.h b/Engine/Source/Runtime/TypedElementFramework/Public/Elements/Framework/TypedElementQueryBuilder.h
--- a/Engine/Source/Runtime/TypedElementFramework/Public/Elements/Framework/TypedElementQueryBuilder.h
+++ b/Engine/Source/Runtime/TypedElementFramework/Public/Elements/Framework/TypedElementQueryBuilder.h
@@ -216,7 +216,7 @@ namespace TypedElementQueryBuilder
 		TYPEDELEMENTFRAMEWORK_API FProcessor& SetGroup(FName GroupName);
 		TYPEDELEMENTFRAMEWORK_API FProcessor& SetBeforeGroup(FName GroupName);
 		TYPEDELEMENTFRAMEWORK_API FProcessor& SetAfterGroup(FName GroupName);
-		TYPEDELEMENTFRAMEWORK_API FProcessor& ForceToGameThread(bool bForce);
+		TYPEDELEMENTFRAMEWORK_API FProcessor& SetExecutionMode(TypedElementDataStorage::EExecutionMode Mode);
 		TYPEDELEMENTFRAMEWORK_API FProcessor& MakeActivatable(FName Name);
 		
 		ITypedElementDataStorageInterface::EQueryTickPhase Phase;
@@ -224,7 +224,7 @@ namespace TypedElementQueryBuilder
 		FName BeforeGroup;
 		FName AfterGroup;
 		FName ActivationName;
-		bool bForceToGameThread{ false };
+		TypedElementDataStorage::EExecutionMode ExecutionMode = TypedElementDataStorage::EExecutionMode::Default;
 	};
 
 	struct FObserver final : public FQueryCallbackType
@@ -246,13 +246,13 @@ namespace TypedElementQueryBuilder
 		TYPEDELEMENTFRAMEWORK_API FObserver& SetMonitoredColumn(const UScriptStruct* MonitoredColumn);
 		template<TypedElementDataStorage::TColumnType ColumnType>
 		FObserver& SetMonitoredColumn();
-		TYPEDELEMENTFRAMEWORK_API FObserver& ForceToGameThread(bool bForce);
+		TYPEDELEMENTFRAMEWORK_API FObserver& SetExecutionMode(TypedElementDataStorage::EExecutionMode Mode);
 		TYPEDELEMENTFRAMEWORK_API FObserver& MakeActivatable(FName Name);
 
 		const UScriptStruct* Monitor;
 		EEvent Event;
 		FName ActivationName;
-		bool bForceToGameThread{ false };
+		TypedElementDataStorage::EExecutionMode ExecutionMode = TypedElementDataStorage::EExecutionMode::Default;
 	};
 
 	struct FPhaseAmble final : public FQueryCallbackType
@@ -266,13 +266,13 @@ namespace TypedElementQueryBuilder
 		TYPEDELEMENTFRAMEWORK_API FPhaseAmble(ELocation InLocation, ITypedElementDataStorageInterface::EQueryTickPhase InPhase);
 		TYPEDELEMENTFRAMEWORK_API FPhaseAmble& SetLocation(ELocation NewLocation);
 		TYPEDELEMENTFRAMEWORK_API FPhaseAmble& SetPhase(ITypedElementDataStorageInterface::EQueryTickPhase NewPhase);
-		TYPEDELEMENTFRAMEWORK_API FPhaseAmble& ForceToGameThread(bool bForce);
+		TYPEDELEMENTFRAMEWORK_API FPhaseAmble& SetExecutionMode(TypedElementDataStorage::EExecutionMode Mode);
 		TYPEDELEMENTFRAMEWORK_API FPhaseAmble& MakeActivatable(FName Name);
 
 		ITypedElementDataStorageInterface::EQueryTickPhase Phase;
 		ELocation Location;
 		FName ActivationName;
-		bool bForceToGameThread{ false };
+		TypedElementDataStorage::EExecutionMode ExecutionMode = TypedElementDataStorage::EExecutionMode::Default;
 	};
 
 	// Because this is a thin wrapper called from within a query callback, it's better to inline fully so all
diff --git a/Engine/Source/Runtime/TypedElementFramework/Public/Elements/Framework/TypedElementQueryBuilder.inl b/Engine/Source/Runtime/TypedElementFramework/Public/Elements/Framework/TypedElementQueryBuilder.inl
--- a/Engine/Source/Runtime/TypedElementFramework/Public/Elements/Framework/TypedElementQueryBuilder.inl
+++ b/Engine/Source/Runtime/TypedElementFramework/Public/Elements/Framework/TypedElementQueryBuilder.inl
@@ -1038,7 +1038,7 @@ namespace TypedElementQueryBuilder
 			{
 				Query.Callback.AfterGroups.Add(Processor.AfterGroup);
 			}
-			Query.Callback.bForceToGameThread = Processor.bForceToGameThread;
+			Query.Callback.ExecutionMode = Processor.ExecutionMode;
 		}
 
 		inline void PrepareForQueryBinding(TypedElementDataStorage::FQueryDescription& Query, const FObserver& Observer)
@@ -1055,7 +1055,7 @@ namespace TypedElementQueryBuilder
 			Query.Callback.ActivationName = Observer.ActivationName;
 			Query.Callback.ActivationCount = Observer.ActivationName.IsNone() ? 255 : 0;
 			Query.Callback.MonitoredType = Observer.Monitor;
-			Query.Callback.bForceToGameThread = Observer.bForceToGameThread;
+			Query.Callback.ExecutionMode = Observer.ExecutionMode;
 		}
 
 		inline void PrepareForQueryBinding(TypedElementDataStorage::FQueryDescription& Query, const FPhaseAmble& PhaseAmble)
@@ -1072,7 +1072,7 @@ namespace TypedElementQueryBuilder
 			Query.Callback.ActivationName = PhaseAmble.ActivationName;
 			Query.Callback.ActivationCount = PhaseAmble.ActivationName.IsNone() ? 255 : 0;
 			Query.Callback.Phase = PhaseAmble.Phase;
-			Query.Callback.bForceToGameThread = PhaseAmble.bForceToGameThread;
+			Query.Callback.ExecutionMode = PhaseAmble.ExecutionMode;
 		}
 		
 		template<typename CallbackType, typename Function>
diff --git a/Engine/Source/Runtime/TypedElementFramework/Public/Elements/Interfaces/TypedElementDataStorageInterface.h b/Engine/Source/Runtime/TypedElementFramework/Public/Elements/Interfaces/TypedElementDataStorageInterface.h
--- a/Engine/Source/Runtime/TypedElementFramework/Public/Elements/Interfaces/TypedElementDataStorageInterface.h
+++ b/Engine/Source/Runtime/TypedElementFramework/Public/Elements/Interfaces/TypedElementDataStorageInterface.h
@@ -334,9 +334,15 @@ public:
 	virtual FQueryResult RunQuery(TypedElementQueryHandle Query) = 0;
 	/**
 	 * Directly runs a query. The callback will be called for batches of matching rows. During a single call to RunQuery the callback
-	 * may be called multiple times. If the query handle is invalid or has been deleted nothing happens and the callback won't be called
+	 * may be called multiple times. If the query handle is invalid or has been deleted nothing happens and the callback won't be called.
 	 */
 	virtual FQueryResult RunQuery(TypedElementQueryHandle Query, DirectQueryCallbackRef Callback) = 0;
+	/**
+	 * Directly runs a query. The callback will be called for batches of matching rows. During a single call to RunQuery the callback
+	 * may be called multiple times. If the query handle is invalid or has been deleted nothing happens and the callback won't be called.
+	 */
+	virtual FQueryResult RunQuery(TypedElementQueryHandle Query, TypedElementDataStorage::EDirectQueryExecutionFlags Flags, 
+		DirectQueryCallbackRef Callback) = 0;
 	/**
 	 * Triggers all queries registered under the activation name to run for one update cycle. The activatable queries will be activated at
 	 * start of the cycle and disabled at the end of the cycle and act like regular queries for that cycle. This includes not running
diff --git a/Engine/Source/Runtime/TypedElementFramework/Public/Elements/Interfaces/TypedElementQueryStorageInterfaces.h b/Engine/Source/Runtime/TypedElementFramework/Public/Elements/Interfaces/TypedElementQueryStorageInterfaces.h
--- a/Engine/Source/Runtime/TypedElementFramework/Public/Elements/Interfaces/TypedElementQueryStorageInterfaces.h
+++ b/Engine/Source/Runtime/TypedElementFramework/Public/Elements/Interfaces/TypedElementQueryStorageInterfaces.h
@@ -202,6 +202,15 @@ namespace TypedElementDataStorage
 		virtual void RemoveColumns(TConstArrayView<RowHandle> Rows, TConstArrayView<const UScriptStruct*> ColumnTypes) = 0;
 	};
 
+	enum class EDirectQueryExecutionFlags : uint32
+	{
+		Default = 0, //< No settings, use the default behavior.
+		ParallelizeChunks = 1 << 0, //< If set, each chunk is processed on a separate thread.
+		IgnoreActivationCount = 1 << 1, //< If set, a direct call will not check activatable queries if they're set.
+		AllowBoundQueries = 1 << 2, //< Stops checking if queries are bound when set, otherwise bound queries will assert.
+	};
+	ENUM_CLASS_FLAGS(EDirectQueryExecutionFlags);
+
 	/**
 	 * Interface to be provided to query callbacks that are directly called through RunQuery from outside a query callback.
 	 */
diff --git a/Engine/Source/Runtime/TypedElementFramework/Tests/Elements/Framework/TypedElementDataStoragePerformanceTestCommands.cpp b/Engine/Source/Runtime/TypedElementFramework/Tests/Elements/Framework/TypedElementDataStoragePerformanceTestCommands.cpp
--- a/Engine/Source/Runtime/TypedElementFramework/Tests/Elements/Framework/TypedElementDataStoragePerformanceTestCommands.cpp
+++ b/Engine/Source/Runtime/TypedElementFramework/Tests/Elements/Framework/TypedElementDataStoragePerformanceTestCommands.cpp
@@ -79,16 +79,13 @@ void UTest_PingPongBetweenPhaseFactory::RegisterQueries(ITypedElementDataStorage
 	using namespace TypedElementQueryBuilder;
 	using namespace TypedElementDataStorage;
 
-	using DSI = ITypedElementDataStorageInterface;
-	namespace DS = TypedElementDataStorage;
-	
 	Super::RegisterQueries(DataStorage);
 
 	DataStorage.RegisterQuery(
 		Select(TEXT("PingPong PrePhysics->DurPhysics"),
-			FProcessor(DSI::EQueryTickPhase::PrePhysics, DataStorage.GetQueryTickGroupName(DSI::EQueryTickGroups::Default))
-			.ForceToGameThread(true),
-			[](DS::IQueryContext& Context, const TypedElementRowHandle* RowPtr, FTest_PingPongPrePhys* PingPongAPtr)
+			FProcessor(EQueryTickPhase::PrePhysics, DataStorage.GetQueryTickGroupName(EQueryTickGroups::Default))
+				.SetExecutionMode(EExecutionMode::GameThread),
+			[](IQueryContext& Context, const TypedElementRowHandle* RowPtr, FTest_PingPongPrePhys* PingPongAPtr)
 			{
 				QUICK_SCOPE_CYCLE_COUNTER(PingPong_Pre_During);
 				TArrayView<const TypedElementRowHandle> Rows(RowPtr, Context.GetRowCount());
@@ -108,9 +105,9 @@ void UTest_PingPongBetweenPhaseFactory::RegisterQueries(ITypedElementDataStorage
 
 	DataStorage.RegisterQuery(
 	Select(TEXT("PingPong DurPhysics->PostPhysics"),
-		FProcessor(DSI::EQueryTickPhase::DuringPhysics, DataStorage.GetQueryTickGroupName(DSI::EQueryTickGroups::Default))
-		.ForceToGameThread(true),
-		[](DS::IQueryContext& Context, const TypedElementRowHandle* RowPtr, FTest_PingPongDurPhys* PingPongBPtr)
+		FProcessor(EQueryTickPhase::DuringPhysics, DataStorage.GetQueryTickGroupName(EQueryTickGroups::Default))
+			.SetExecutionMode(EExecutionMode::GameThread),
+		[](IQueryContext& Context, const TypedElementRowHandle* RowPtr, FTest_PingPongDurPhys* PingPongBPtr)
 		{
 			QUICK_SCOPE_CYCLE_COUNTER(PingPong_During_Post);
 			TArrayView<const TypedElementRowHandle> Rows(RowPtr, Context.GetRowCount());
@@ -130,9 +127,9 @@ void UTest_PingPongBetweenPhaseFactory::RegisterQueries(ITypedElementDataStorage
 
 	DataStorage.RegisterQuery(
 	Select(TEXT("PingPong PostPhysics->PrePhysics"),
-	FProcessor(DSI::EQueryTickPhase::PostPhysics, DataStorage.GetQueryTickGroupName(DSI::EQueryTickGroups::Default))
-	.ForceToGameThread(true),
-	[](DS::IQueryContext& Context, const TypedElementRowHandle* RowPtr, FTest_PingPongPostPhys* PingPongBPtr)
+	FProcessor(EQueryTickPhase::PostPhysics, DataStorage.GetQueryTickGroupName(EQueryTickGroups::Default))
+		.SetExecutionMode(EExecutionMode::GameThread),
+	[](IQueryContext& Context, const TypedElementRowHandle* RowPtr, FTest_PingPongPostPhys* PingPongBPtr)
 	{
 		QUICK_SCOPE_CYCLE_COUNTER(PingPong_Post_Pre);
 		TArrayView<const TypedElementRowHandle> Rows(RowPtr, Context.GetRowCount());
