[Teds] Renaming TedsCore/Compatibility classes

diff --git a/Engine/Plugins/Experimental/EditorDataStorage/Source/TedsCore/Private/Compatibility/TedsTypedElementActorHandleFactory.cpp b/Engine/Plugins/Experimental/EditorDataStorage/Source/TedsCore/Private/Compatibility/TedsTypedElementActorHandleFactory.cpp
--- a/Engine/Plugins/Experimental/EditorDataStorage/Source/TedsCore/Private/Compatibility/TedsTypedElementActorHandleFactory.cpp
+++ b/Engine/Plugins/Experimental/EditorDataStorage/Source/TedsCore/Private/Compatibility/TedsTypedElementActorHandleFactory.cpp
@@ -3,6 +3,7 @@
 #include "TedsTypedElementActorHandleFactory.h"
 
 #include "Compatibility/TedsTypedElementBridge.h"
+#include "Compatibility/Columns/TypedElement.h"
 #include "Elements/Columns/TypedElementCompatibilityColumns.h"
 #include "Elements/Framework/EngineElementsLibrary.h"
 #include "Elements/Framework/TypedElementList.h"
@@ -10,24 +11,24 @@
 #include "Elements/Framework/TypedElementRegistry.h"
 #include "GameFramework/Actor.h"
 
-void UTEDSTypedElementActorHandleFactory::PreRegister(ITypedElementDataStorageInterface& DataStorage)
+void UTypedElementActorHandleDataStorageFactory::PreRegister(ITypedElementDataStorageInterface& DataStorage)
 {
 	Super::PreRegister(DataStorage);
 	
-	BridgeEnableDelegateHandle = UTEDSTypedElementBridge::OnEnabled().AddUObject(this, &UTEDSTypedElementActorHandleFactory::HandleBridgeEnabled);
+	BridgeEnableDelegateHandle = UE::Editor::DataStorage::Compatibility::OnTypedElementBridgeEnabled().AddUObject(this, &UTypedElementActorHandleDataStorageFactory::HandleBridgeEnabled);
 }
 
-void UTEDSTypedElementActorHandleFactory::PreShutdown(ITypedElementDataStorageInterface& DataStorage)
+void UTypedElementActorHandleDataStorageFactory::PreShutdown(ITypedElementDataStorageInterface& DataStorage)
 {
-	UTEDSTypedElementBridge::OnEnabled().Remove(BridgeEnableDelegateHandle);
+	UE::Editor::DataStorage::Compatibility::OnTypedElementBridgeEnabled().Remove(BridgeEnableDelegateHandle);
 	BridgeEnableDelegateHandle.Reset();
 }
 
-void UTEDSTypedElementActorHandleFactory::RegisterQueries(ITypedElementDataStorageInterface& DataStorage)
+void UTypedElementActorHandleDataStorageFactory::RegisterQueries(ITypedElementDataStorageInterface& DataStorage)
 {
 	Super::RegisterQueries(DataStorage);
 
-	if (UTEDSTypedElementBridge::IsEnabled())
+	if (UE::Editor::DataStorage::Compatibility::IsTypedElementBridgeEnabled())
 	{
 		RegisterQuery_ActorHandlePopulate(DataStorage);
 	}
@@ -42,7 +43,7 @@ void UTEDSTypedElementActorHandleFactory::RegisterQueries(ITypedElementDataStora
 	.Compile());
 }
 
-void UTEDSTypedElementActorHandleFactory::RegisterQuery_ActorHandlePopulate(ITypedElementDataStorageInterface& DataStorage)
+void UTypedElementActorHandleDataStorageFactory::RegisterQuery_ActorHandlePopulate(ITypedElementDataStorageInterface& DataStorage)
 {
 	using namespace TypedElementQueryBuilder;
 	using namespace TypedElementDataStorage;
@@ -61,7 +62,7 @@ void UTEDSTypedElementActorHandleFactory::RegisterQuery_ActorHandlePopulate(ITyp
 			{
 				checkSlow(Cast<AActor>(Object));
 				FTypedElementHandle Handle = UEngineElementsLibrary::AcquireEditorActorElementHandle(static_cast<AActor*>(Object));
-				Context.AddColumn(Row, FTEDSTypedElementColumn
+				Context.AddColumn(Row, UE::Editor::DataStorage::Compatibility::FTypedElementColumn
 				{
 					.Handle = Handle
 				});
@@ -73,7 +74,7 @@ void UTEDSTypedElementActorHandleFactory::RegisterQuery_ActorHandlePopulate(ITyp
 	.Compile());
 }
 
-void UTEDSTypedElementActorHandleFactory::HandleBridgeEnabled(bool bEnabled)
+void UTypedElementActorHandleDataStorageFactory::HandleBridgeEnabled(bool bEnabled)
 {
 	ITypedElementDataStorageInterface* DataStorage = UTypedElementRegistry::GetInstance()->GetMutableDataStorage();
 	
@@ -108,7 +109,7 @@ void UTEDSTypedElementActorHandleFactory::HandleBridgeEnabled(bool bEnabled)
 			if (const AActor* Actor = Actors[Index].Get())
 			{
 				FTypedElementHandle Handle = UEngineElementsLibrary::AcquireEditorActorElementHandle(Actor);
-				DataStorage->AddColumn(CollatedRowHandles[Index], FTEDSTypedElementColumn
+				DataStorage->AddColumn(CollatedRowHandles[Index], UE::Editor::DataStorage::Compatibility::FTypedElementColumn
 				{
 					.Handle = Handle
 				});
diff --git a/Engine/Plugins/Experimental/EditorDataStorage/Source/TedsCore/Private/Compatibility/TedsTypedElementActorHandleFactory.h b/Engine/Plugins/Experimental/EditorDataStorage/Source/TedsCore/Private/Compatibility/TedsTypedElementActorHandleFactory.h
--- a/Engine/Plugins/Experimental/EditorDataStorage/Source/TedsCore/Private/Compatibility/TedsTypedElementActorHandleFactory.h
+++ b/Engine/Plugins/Experimental/EditorDataStorage/Source/TedsCore/Private/Compatibility/TedsTypedElementActorHandleFactory.h
@@ -1,4 +1,4 @@
-﻿// Copyright Epic Games, Inc. All Rights Reserved.
+// Copyright Epic Games, Inc. All Rights Reserved.
 
 #pragma once
 
@@ -12,11 +12,11 @@
 class UTypedElementRegistry;
 
 /**
- * This class is responsible for acquiring and registering Actor TypedElementHandles
+ * This class is responsible for acquiring and registering Actor Typed Element Handles
  * with TEDS/
  */
 UCLASS(Transient)
-class UTEDSTypedElementActorHandleFactory : public UTypedElementDataStorageFactory
+class UTypedElementActorHandleDataStorageFactory : public UTypedElementDataStorageFactory
 {
 	GENERATED_BODY()
 public:
diff --git a/Engine/Plugins/Experimental/EditorDataStorage/Source/TedsCore/Private/Compatibility/TedsTypedElementBridge.cpp b/Engine/Plugins/Experimental/EditorDataStorage/Source/TedsCore/Private/Compatibility/TedsTypedElementBridge.cpp
--- a/Engine/Plugins/Experimental/EditorDataStorage/Source/TedsCore/Private/Compatibility/TedsTypedElementBridge.cpp
+++ b/Engine/Plugins/Experimental/EditorDataStorage/Source/TedsCore/Private/Compatibility/TedsTypedElementBridge.cpp
@@ -1,114 +1,20 @@
 // Copyright Epic Games, Inc. All Rights Reserved.
 
 #include "Compatibility/TedsTypedElementBridge.h"
+#include "Compatibility/TedsTypedElementBridgeQueries.h"
 
-#include "Elements/Columns/TypedElementCompatibilityColumns.h"
-#include "Elements/Framework/EngineElementsLibrary.h"
-#include "Elements/Framework/TypedElementList.h"
-#include "Elements/Framework/TypedElementQueryBuilder.h"
-#include "Elements/Framework/TypedElementRegistry.h"
 #include "HAL/IConsoleManager.h"
 
-namespace UTEDSTypedElementBridge_Private
+namespace UE::Editor::DataStorage::Compatibility
 {
-	bool bBridgeEnabled = false;
-	FAutoConsoleVariableRef CVarBridgeEnabled(
-		TEXT("TEDS.TypedElementBridge.Enable"),
-		bBridgeEnabled,
-		TEXT("Automatically populated TEDS with TypedElementHandles"));
-}
-
-uint8 UTEDSTypedElementBridge::GetOrder() const
-{
-	return 110;
-}
-
-void UTEDSTypedElementBridge::PreRegister(ITypedElementDataStorageInterface& DataStorage)
-{
-	Super::PreRegister(DataStorage);
-	DebugEnabledDelegateHandle = UTEDSTypedElementBridge_Private::CVarBridgeEnabled->OnChangedDelegate().AddUObject(this, &UTEDSTypedElementBridge::HandleOnEnabled);
-}
-
-void UTEDSTypedElementBridge::PreShutdown(ITypedElementDataStorageInterface& DataStorage)
-{
-	UTEDSTypedElementBridge_Private::CVarBridgeEnabled->OnChangedDelegate().Remove(DebugEnabledDelegateHandle);
-	DebugEnabledDelegateHandle.Reset();
-	CleanupTypedElementColumns(DataStorage);
-
-	Super::PreShutdown(DataStorage);
-}
-
-void UTEDSTypedElementBridge::RegisterQueries(ITypedElementDataStorageInterface& DataStorage)
-{
-	Super::RegisterQueries(DataStorage);
-
-	if (IsEnabled())
+	FOnTypedElementBridgeEnabled& OnTypedElementBridgeEnabled()
 	{
-		RegisterQuery_NewUObject(DataStorage);
+		static FOnTypedElementBridgeEnabled OnBridgeEnabled;
+		return OnBridgeEnabled;
 	}
-}
-
-void UTEDSTypedElementBridge::RegisterQuery_NewUObject(ITypedElementDataStorageInterface& DataStorage)
-{
-	using namespace TypedElementQueryBuilder;
-	using DSI = ITypedElementDataStorageInterface;
-	
-	RemoveTypedElementRowHandleQuery = DataStorage.RegisterQuery(
-		Select()
-			.ReadOnly<FTEDSTypedElementColumn>()
-		.Compile());
-}
-
-void UTEDSTypedElementBridge::UnregisterQuery_NewUObject(ITypedElementDataStorageInterface& DataStorage)
-{
-}
 
-FOnTEDSTypedElementBridgeEnable& UTEDSTypedElementBridge::OnEnabled()
-{
-	static FOnTEDSTypedElementBridgeEnable OnBridgeEnabled;
-	return OnBridgeEnabled;
-}
-
-bool UTEDSTypedElementBridge::IsEnabled()
-{
-	return UTEDSTypedElementBridge_Private::CVarBridgeEnabled->GetBool();
-}
-
-void UTEDSTypedElementBridge::CleanupTypedElementColumns(ITypedElementDataStorageInterface& DataStorage)
-{
-	using namespace TypedElementQueryBuilder;
-	using namespace TypedElementDataStorage;
-
-	// Remove any TEv1 handles
-	{
-		TArray<TypedElementRowHandle> Handles;
-		using namespace TypedElementQueryBuilder;
-		DataStorage.RunQuery(
-			RemoveTypedElementRowHandleQuery,
-			CreateDirectQueryCallbackBinding(
-				[&Handles](ITypedElementDataStorageInterface::IDirectQueryContext& Context)
-				{
-					Handles.Append(Context.GetRowHandles());
-				}));
-		
-		DataStorage.BatchAddRemoveColumns(TConstArrayView<TypedElementRowHandle>(Handles), {}, {FTEDSTypedElementColumn::StaticStruct()});
-	}
-}
-
-void UTEDSTypedElementBridge::HandleOnEnabled(IConsoleVariable* CVar)
-{
-	ITypedElementDataStorageInterface* DataStorage = UTypedElementRegistry::GetInstance()->GetMutableDataStorage();
-	bool bIsEnabled = CVar->GetBool();
-
-	if (bIsEnabled)
-	{
-		RegisterQuery_NewUObject(*DataStorage);
-		UTEDSTypedElementBridge::OnEnabled().Broadcast(bIsEnabled);
-	}
-	else
+	bool IsTypedElementBridgeEnabled()
 	{
-		UTEDSTypedElementBridge::OnEnabled().Broadcast(bIsEnabled);
-		UnregisterQuery_NewUObject(*DataStorage);
-		CleanupTypedElementColumns(*DataStorage);
+		return UTypedElementBridgeDataStorageFactory::IsEnabled();
 	}
-}
+} // namespace UE::Editor::DataStorage::Compatibility
diff --git a/Engine/Plugins/Experimental/EditorDataStorage/Source/TedsCore/Private/Compatibility/TedsTypedElementBridgeQueries.cpp b/Engine/Plugins/Experimental/EditorDataStorage/Source/TedsCore/Private/Compatibility/TedsTypedElementBridgeQueries.cpp
new file mode 100644
--- /dev/null
+++ b/Engine/Plugins/Experimental/EditorDataStorage/Source/TedsCore/Private/Compatibility/TedsTypedElementBridgeQueries.cpp
@@ -0,0 +1,110 @@
+// Copyright Epic Games, Inc. All Rights Reserved.
+
+#include "Compatibility/TedsTypedElementBridgeQueries.h"
+
+#include "Compatibility/TedsTypedElementBridge.h"
+#include "Compatibility/Columns/TypedElement.h"
+#include "Elements/Columns/TypedElementCompatibilityColumns.h"
+#include "Elements/Framework/EngineElementsLibrary.h"
+#include "Elements/Framework/TypedElementList.h"
+#include "Elements/Framework/TypedElementQueryBuilder.h"
+#include "Elements/Framework/TypedElementRegistry.h"
+#include "HAL/IConsoleManager.h"
+
+namespace UE::Editor::DataStorage::Compatibility::Private
+{
+	bool bBridgeEnabled = false;
+	FAutoConsoleVariableRef CVarBridgeEnabled(
+		TEXT("TEDS.TypedElementBridge.Enable"),
+		bBridgeEnabled,
+		TEXT("Automatically populated TEDS with TypedElementHandles"));
+} // namespace UE::Editor::DataStorage::Compatibility::Private
+
+uint8 UTypedElementBridgeDataStorageFactory::GetOrder() const
+{
+	return 110;
+}
+
+void UTypedElementBridgeDataStorageFactory::PreRegister(ITypedElementDataStorageInterface& DataStorage)
+{
+	Super::PreRegister(DataStorage);
+	DebugEnabledDelegateHandle = UE::Editor::DataStorage::Compatibility::Private::CVarBridgeEnabled->OnChangedDelegate().AddUObject(this, &UTypedElementBridgeDataStorageFactory::HandleOnEnabled);
+}
+
+void UTypedElementBridgeDataStorageFactory::PreShutdown(ITypedElementDataStorageInterface& DataStorage)
+{
+	UE::Editor::DataStorage::Compatibility::Private::CVarBridgeEnabled->OnChangedDelegate().Remove(DebugEnabledDelegateHandle);
+	DebugEnabledDelegateHandle.Reset();
+	CleanupTypedElementColumns(DataStorage);
+
+	Super::PreShutdown(DataStorage);
+}
+
+void UTypedElementBridgeDataStorageFactory::RegisterQueries(ITypedElementDataStorageInterface& DataStorage)
+{
+	Super::RegisterQueries(DataStorage);
+
+	if (IsEnabled())
+	{
+		RegisterQuery_NewUObject(DataStorage);
+	}
+}
+
+bool UTypedElementBridgeDataStorageFactory::IsEnabled()
+{
+	return UE::Editor::DataStorage::Compatibility::Private::CVarBridgeEnabled->GetBool();
+}
+
+void UTypedElementBridgeDataStorageFactory::RegisterQuery_NewUObject(ITypedElementDataStorageInterface& DataStorage)
+{
+	using namespace TypedElementQueryBuilder;
+	using DSI = ITypedElementDataStorageInterface;
+	
+	RemoveTypedElementRowHandleQuery = DataStorage.RegisterQuery(
+		Select()
+			.ReadOnly<UE::Editor::DataStorage::Compatibility::FTypedElementColumn>()
+		.Compile());
+}
+
+void UTypedElementBridgeDataStorageFactory::UnregisterQuery_NewUObject(ITypedElementDataStorageInterface& DataStorage)
+{
+}
+
+void UTypedElementBridgeDataStorageFactory::CleanupTypedElementColumns(ITypedElementDataStorageInterface& DataStorage)
+{
+	using namespace TypedElementQueryBuilder;
+	using namespace TypedElementDataStorage;
+
+	// Remove any TEv1 handles
+	{
+		TArray<TypedElementRowHandle> Handles;
+		using namespace TypedElementQueryBuilder;
+		DataStorage.RunQuery(
+			RemoveTypedElementRowHandleQuery,
+			CreateDirectQueryCallbackBinding(
+				[&Handles](ITypedElementDataStorageInterface::IDirectQueryContext& Context)
+				{
+					Handles.Append(Context.GetRowHandles());
+				}));
+		
+		DataStorage.BatchAddRemoveColumns(TConstArrayView<TypedElementRowHandle>(Handles), {}, {UE::Editor::DataStorage::Compatibility::FTypedElementColumn::StaticStruct()});
+	}
+}
+
+void UTypedElementBridgeDataStorageFactory::HandleOnEnabled(IConsoleVariable* CVar)
+{
+	ITypedElementDataStorageInterface* DataStorage = UTypedElementRegistry::GetInstance()->GetMutableDataStorage();
+	bool bIsEnabled = CVar->GetBool();
+
+	if (bIsEnabled)
+	{
+		RegisterQuery_NewUObject(*DataStorage);
+		UE::Editor::DataStorage::Compatibility::OnTypedElementBridgeEnabled().Broadcast(bIsEnabled);
+	}
+	else
+	{
+		UE::Editor::DataStorage::Compatibility::OnTypedElementBridgeEnabled().Broadcast(bIsEnabled);
+		UnregisterQuery_NewUObject(*DataStorage);
+		CleanupTypedElementColumns(*DataStorage);
+	}
+}
diff --git a/Engine/Plugins/Experimental/EditorDataStorage/Source/TedsCore/Private/Compatibility/TedsTypedElementBridgeQueries.h b/Engine/Plugins/Experimental/EditorDataStorage/Source/TedsCore/Private/Compatibility/TedsTypedElementBridgeQueries.h
new file mode 100644
--- /dev/null
+++ b/Engine/Plugins/Experimental/EditorDataStorage/Source/TedsCore/Private/Compatibility/TedsTypedElementBridgeQueries.h
@@ -0,0 +1,39 @@
+// Copyright Epic Games, Inc. All Rights Reserved.
+
+#pragma once
+
+#include "Elements/Framework/TypedElementHandle.h"
+#include "Elements/Interfaces/TypedElementDataStorageInterface.h"
+#include "Elements/Interfaces/TypedElementDataStorageFactory.h"
+#include "UObject/UObjectGlobals.h"
+
+#include "TedsTypedElementBridgeQueries.generated.h"
+
+class UTypedElementRegistry;
+
+/**
+ * This class is responsible for running queries that will ensure Typed Element Handles
+ * are cleaned up when TEDS is shut down.
+ */
+UCLASS(Transient)
+class UTypedElementBridgeDataStorageFactory : public UTypedElementDataStorageFactory
+{
+	GENERATED_BODY()
+public:
+	~UTypedElementBridgeDataStorageFactory() override = default;
+	virtual uint8 GetOrder() const override;
+	virtual void PreRegister(ITypedElementDataStorageInterface& DataStorage) override;
+	virtual void PreShutdown(ITypedElementDataStorageInterface& DataStorage) override;
+	virtual void RegisterQueries(ITypedElementDataStorageInterface& DataStorage) override;
+
+	static bool IsEnabled();
+
+private:
+	void RegisterQuery_NewUObject(ITypedElementDataStorageInterface& DataStorage);
+	void UnregisterQuery_NewUObject(ITypedElementDataStorageInterface& DataStorage);
+	void CleanupTypedElementColumns(ITypedElementDataStorageInterface& DataStorage);
+	void HandleOnEnabled(IConsoleVariable* CVar);
+	
+	TypedElementQueryHandle RemoveTypedElementRowHandleQuery = TypedElementDataStorage::InvalidQueryHandle;
+	FDelegateHandle DebugEnabledDelegateHandle;
+};
diff --git a/Engine/Plugins/Experimental/EditorDataStorage/Source/TedsCore/Private/Compatibility/TypedElementActorLabelQueries.cpp b/Engine/Plugins/Experimental/EditorDataStorage/Source/TedsCore/Private/Compatibility/TypedElementActorLabelQueries.cpp
--- a/Engine/Plugins/Experimental/EditorDataStorage/Source/TedsCore/Private/Compatibility/TypedElementActorLabelQueries.cpp
+++ b/Engine/Plugins/Experimental/EditorDataStorage/Source/TedsCore/Private/Compatibility/TypedElementActorLabelQueries.cpp
@@ -11,15 +11,15 @@
 #include "Hash/CityHash.h"
 #include "ScopedTransaction.h"
 
-#define LOCTEXT_NAMESPACE "TypedElementDataStorage"
+#define LOCTEXT_NAMESPACE "TedsCore"
 
-void UTypedElementActorLabelFactory::RegisterQueries(ITypedElementDataStorageInterface& DataStorage)
+void UActorLabelDataStorageFactory::RegisterQueries(ITypedElementDataStorageInterface& DataStorage)
 {
 	RegisterActorLabelToColumnQuery(DataStorage);
 	RegisterLabelColumnToActorQuery(DataStorage);
 }
 
-void UTypedElementActorLabelFactory::RegisterActorLabelToColumnQuery(ITypedElementDataStorageInterface& DataStorage) const
+void UActorLabelDataStorageFactory::RegisterActorLabelToColumnQuery(ITypedElementDataStorageInterface& DataStorage) const
 {
 	using namespace TypedElementQueryBuilder;
 	using namespace TypedElementDataStorage;
@@ -49,7 +49,7 @@ void UTypedElementActorLabelFactory::RegisterActorLabelToColumnQuery(ITypedEleme
 	);
 }
 
-void UTypedElementActorLabelFactory::RegisterLabelColumnToActorQuery(ITypedElementDataStorageInterface& DataStorage) const
+void UActorLabelDataStorageFactory::RegisterLabelColumnToActorQuery(ITypedElementDataStorageInterface& DataStorage) const
 {
 	using namespace TypedElementQueryBuilder;
 	using namespace TypedElementDataStorage;
diff --git a/Engine/Plugins/Experimental/EditorDataStorage/Source/TedsCore/Private/Compatibility/TypedElementActorLabelQueries.h b/Engine/Plugins/Experimental/EditorDataStorage/Source/TedsCore/Private/Compatibility/TypedElementActorLabelQueries.h
--- a/Engine/Plugins/Experimental/EditorDataStorage/Source/TedsCore/Private/Compatibility/TypedElementActorLabelQueries.h
+++ b/Engine/Plugins/Experimental/EditorDataStorage/Source/TedsCore/Private/Compatibility/TypedElementActorLabelQueries.h
@@ -9,12 +9,12 @@
 
 
 UCLASS()
-class UTypedElementActorLabelFactory : public UTypedElementDataStorageFactory
+class UActorLabelDataStorageFactory : public UTypedElementDataStorageFactory
 {
 	GENERATED_BODY()
 
 public:
-	~UTypedElementActorLabelFactory() override = default;
+	~UActorLabelDataStorageFactory() override = default;
 
 	void RegisterQueries(ITypedElementDataStorageInterface& DataStorage) override;
 
diff --git a/Engine/Plugins/Experimental/EditorDataStorage/Source/TedsCore/Private/Compatibility/TypedElementActorParentQueries.cpp b/Engine/Plugins/Experimental/EditorDataStorage/Source/TedsCore/Private/Compatibility/TypedElementActorParentQueries.cpp
--- a/Engine/Plugins/Experimental/EditorDataStorage/Source/TedsCore/Private/Compatibility/TypedElementActorParentQueries.cpp
+++ b/Engine/Plugins/Experimental/EditorDataStorage/Source/TedsCore/Private/Compatibility/TypedElementActorParentQueries.cpp
@@ -20,7 +20,7 @@ namespace UE::TypedElementActorParentQueries::Local
 		TEXT("Mirror parent information for actors to TEDS (only works when set on startup)"));
 };
 
-void UTypedElementActorParentFactory::RegisterQueries(ITypedElementDataStorageInterface& DataStorage)
+void UActorParentDataStorageFactory::RegisterQueries(ITypedElementDataStorageInterface& DataStorage)
 {
 	if(UE::TypedElementActorParentQueries::Local::bAddParentColumnToActors)
 	{
@@ -29,7 +29,7 @@ void UTypedElementActorParentFactory::RegisterQueries(ITypedElementDataStorageIn
 	}
 }
 
-void UTypedElementActorParentFactory::RegisterAddParentColumn(ITypedElementDataStorageInterface& DataStorage) const
+void UActorParentDataStorageFactory::RegisterAddParentColumn(ITypedElementDataStorageInterface& DataStorage) const
 {
 	using namespace TypedElementDataStorage;
 	using namespace TypedElementQueryBuilder;
@@ -66,7 +66,7 @@ void UTypedElementActorParentFactory::RegisterAddParentColumn(ITypedElementDataS
 	);
 }
 
-void UTypedElementActorParentFactory::RegisterUpdateOrRemoveParentColumn(ITypedElementDataStorageInterface& DataStorage) const
+void UActorParentDataStorageFactory::RegisterUpdateOrRemoveParentColumn(ITypedElementDataStorageInterface& DataStorage) const
 {
 	using namespace TypedElementDataStorage;
 	using namespace TypedElementQueryBuilder;
diff --git a/Engine/Plugins/Experimental/EditorDataStorage/Source/TedsCore/Private/Compatibility/TypedElementActorParentQueries.h b/Engine/Plugins/Experimental/EditorDataStorage/Source/TedsCore/Private/Compatibility/TypedElementActorParentQueries.h
--- a/Engine/Plugins/Experimental/EditorDataStorage/Source/TedsCore/Private/Compatibility/TypedElementActorParentQueries.h
+++ b/Engine/Plugins/Experimental/EditorDataStorage/Source/TedsCore/Private/Compatibility/TypedElementActorParentQueries.h
@@ -8,12 +8,12 @@
 #include "TypedElementActorParentQueries.generated.h"
 
 UCLASS()
-class UTypedElementActorParentFactory : public UTypedElementDataStorageFactory
+class UActorParentDataStorageFactory : public UTypedElementDataStorageFactory
 {
 	GENERATED_BODY()
 
 public:
-	~UTypedElementActorParentFactory() override = default;
+	~UActorParentDataStorageFactory() override = default;
 
 	void RegisterQueries(ITypedElementDataStorageInterface& DataStorage) override;
 
diff --git a/Engine/Plugins/Experimental/EditorDataStorage/Source/TedsCore/Private/Compatibility/TypedElementActorTransformQueries.cpp b/Engine/Plugins/Experimental/EditorDataStorage/Source/TedsCore/Private/Compatibility/TypedElementActorTransformQueries.cpp
--- a/Engine/Plugins/Experimental/EditorDataStorage/Source/TedsCore/Private/Compatibility/TypedElementActorTransformQueries.cpp
+++ b/Engine/Plugins/Experimental/EditorDataStorage/Source/TedsCore/Private/Compatibility/TypedElementActorTransformQueries.cpp
@@ -8,14 +8,14 @@
 #include "Elements/Framework/TypedElementQueryBuilder.h"
 #include "GameFramework/Actor.h"
 
-void UTypedElementActorTransformFactory::RegisterQueries(ITypedElementDataStorageInterface& DataStorage)
+void UActorTransformDataStorageFactory::RegisterQueries(ITypedElementDataStorageInterface& DataStorage)
 {
 	RegisterActorAddTransformColumn(DataStorage);
 	RegisterActorLocalTransformToColumn(DataStorage);
 	RegisterLocalTransformColumnToActor(DataStorage);
 }
 
-void UTypedElementActorTransformFactory::RegisterActorAddTransformColumn(ITypedElementDataStorageInterface& DataStorage) const
+void UActorTransformDataStorageFactory::RegisterActorAddTransformColumn(ITypedElementDataStorageInterface& DataStorage) const
 {
 	using namespace TypedElementQueryBuilder;
 	using namespace TypedElementDataStorage;
@@ -41,7 +41,7 @@ void UTypedElementActorTransformFactory::RegisterActorAddTransformColumn(ITypedE
 	);
 }
 
-void UTypedElementActorTransformFactory::RegisterActorLocalTransformToColumn(ITypedElementDataStorageInterface& DataStorage) const
+void UActorTransformDataStorageFactory::RegisterActorLocalTransformToColumn(ITypedElementDataStorageInterface& DataStorage) const
 {
 	using namespace TypedElementQueryBuilder;
 	using namespace TypedElementDataStorage;
@@ -70,7 +70,7 @@ void UTypedElementActorTransformFactory::RegisterActorLocalTransformToColumn(ITy
 	);
 }
 
-void UTypedElementActorTransformFactory::RegisterLocalTransformColumnToActor(ITypedElementDataStorageInterface& DataStorage) const
+void UActorTransformDataStorageFactory::RegisterLocalTransformColumnToActor(ITypedElementDataStorageInterface& DataStorage) const
 {
 	using namespace TypedElementQueryBuilder;
 	using namespace TypedElementDataStorage;
diff --git a/Engine/Plugins/Experimental/EditorDataStorage/Source/TedsCore/Private/Compatibility/TypedElementActorTransformQueries.h b/Engine/Plugins/Experimental/EditorDataStorage/Source/TedsCore/Private/Compatibility/TypedElementActorTransformQueries.h
--- a/Engine/Plugins/Experimental/EditorDataStorage/Source/TedsCore/Private/Compatibility/TypedElementActorTransformQueries.h
+++ b/Engine/Plugins/Experimental/EditorDataStorage/Source/TedsCore/Private/Compatibility/TypedElementActorTransformQueries.h
@@ -8,18 +8,18 @@
 #include "TypedElementActorTransformQueries.generated.h"
 
 UCLASS()
-class UTypedElementActorTransformFactory : public UTypedElementDataStorageFactory
+class UActorTransformDataStorageFactory : public UTypedElementDataStorageFactory
 {
 	GENERATED_BODY()
 
 public:
-	~UTypedElementActorTransformFactory() override = default;
+	~UActorTransformDataStorageFactory() override = default;
 
 	void RegisterQueries(ITypedElementDataStorageInterface& DataStorage) override;
 
 private:
 	/**
-	 * Checks actors that don't have a tranform column and adds one if an actor has been
+	 * Checks actors that don't have a transform column and adds one if an actor has been
 	 * assigned a transform.
 	 */
 	void RegisterActorAddTransformColumn(ITypedElementDataStorageInterface& DataStorage) const;
diff --git a/Engine/Plugins/Experimental/EditorDataStorage/Source/TedsCore/Private/Compatibility/TypedElementActorViewportProcessors.cpp b/Engine/Plugins/Experimental/EditorDataStorage/Source/TedsCore/Private/Compatibility/TypedElementActorViewportProcessors.cpp
--- a/Engine/Plugins/Experimental/EditorDataStorage/Source/TedsCore/Private/Compatibility/TypedElementActorViewportProcessors.cpp
+++ b/Engine/Plugins/Experimental/EditorDataStorage/Source/TedsCore/Private/Compatibility/TypedElementActorViewportProcessors.cpp
@@ -12,165 +12,168 @@
 #include "Elements/Columns/TypedElementSelectionColumns.h"
 #include "ProfilingDebugging/CpuProfilerTrace.h"
 
-FAutoConsoleCommandWithArgsAndOutputDevice SetOutlineColorConsoleCommand(
-	TEXT("TEDS.Debug.SetOutlineColor"),
-	TEXT("Adds an outline color to selected objects."),
-	FConsoleCommandWithArgsAndOutputDeviceDelegate::CreateLambda([](const TArray<FString>& Args, FOutputDevice& Output)
-		{
-			using namespace TypedElementQueryBuilder;
-			using DSI = ITypedElementDataStorageInterface;
+namespace UE::Editor::DataStorage::Private
+	{
+	FAutoConsoleCommandWithArgsAndOutputDevice SetOutlineColorConsoleCommand(
+		TEXT("TEDS.Debug.SetOutlineColor"),
+		TEXT("Adds an outline color to selected objects."),
+		FConsoleCommandWithArgsAndOutputDeviceDelegate::CreateLambda([](const TArray<FString>& Args, FOutputDevice& Output)
+			{
+				using namespace TypedElementQueryBuilder;
+				using DSI = ITypedElementDataStorageInterface;
 
-			TRACE_CPUPROFILER_EVENT_SCOPE(TEDS.Debug.AddOverlayColorToSelectionCommand);
+				TRACE_CPUPROFILER_EVENT_SCOPE(TEDS.Debug.AddOverlayColorToSelectionCommand);
 
-			if (ITypedElementDataStorageInterface* DataStorage = UTypedElementRegistry::GetInstance()->GetMutableDataStorage())
-			{
-				static TypedElementQueryHandle OverlayQuery = TypedElementInvalidQueryHandle;
-				if (OverlayQuery == TypedElementInvalidQueryHandle)
+				if (ITypedElementDataStorageInterface* DataStorage = UTypedElementRegistry::GetInstance()->GetMutableDataStorage())
 				{
-					OverlayQuery = DataStorage->RegisterQuery(
-						Select()
-						.Where()
-							.All<FTypedElementSelectionColumn>()
-						.Compile());
-				}
+					static TypedElementQueryHandle OverlayQuery = TypedElementInvalidQueryHandle;
+					if (OverlayQuery == TypedElementInvalidQueryHandle)
+					{
+						OverlayQuery = DataStorage->RegisterQuery(
+							Select()
+							.Where()
+								.All<FTypedElementSelectionColumn>()
+							.Compile());
+					}
 				
-				if (OverlayQuery == TypedElementInvalidQueryHandle)
-				{
-					return;
-				}
+					if (OverlayQuery == TypedElementInvalidQueryHandle)
+					{
+						return;
+					}
 
-				if (Args.IsEmpty())
-				{
-					Output.Log(TEXT("Provide a color index (0-7) to use as outline"));
-					return;
-				}
+					if (Args.IsEmpty())
+					{
+						Output.Log(TEXT("Provide a color index (0-7) to use as outline"));
+						return;
+					}
 				
-				// Parse the color
-				int32 ColorIndex;
-				LexFromString(ColorIndex, *Args[0]);
+					// Parse the color
+					int32 ColorIndex;
+					LexFromString(ColorIndex, *Args[0]);
 
-				if (!(ColorIndex >= 0 && ColorIndex <= 7))
-				{
-					Output.Log(TEXT("Color index must be in range [0,7]"));
-					return;
-				}
+					if (!(ColorIndex >= 0 && ColorIndex <= 7))
+					{
+						Output.Log(TEXT("Color index must be in range [0,7]"));
+						return;
+					}
 
-				TArray<TypedElementRowHandle> RowHandles;
+					TArray<TypedElementRowHandle> RowHandles;
 				
-				DataStorage->RunQuery(OverlayQuery, [ColorIndex, &RowHandles](const DSI::FQueryDescription&, DSI::IDirectQueryContext& Context)
-				{
-					RowHandles = Context.GetRowHandles();
-				});
+					DataStorage->RunQuery(OverlayQuery, [ColorIndex, &RowHandles](const DSI::FQueryDescription&, DSI::IDirectQueryContext& Context)
+					{
+						RowHandles = Context.GetRowHandles();
+					});
 
-				for (TypedElementRowHandle Row : RowHandles)
-				{
-					DataStorage->AddColumn(Row, 
-						FTypedElementViewportOutlineColorColumn{ .SelectionOutlineColorIndex = static_cast<uint8>(ColorIndex) });
-					DataStorage->AddColumn<FTypedElementSyncBackToWorldTag>(Row);
+					for (TypedElementRowHandle Row : RowHandles)
+					{
+						DataStorage->AddColumn(Row, 
+							FTypedElementViewportOutlineColorColumn{ .SelectionOutlineColorIndex = static_cast<uint8>(ColorIndex) });
+						DataStorage->AddColumn<FTypedElementSyncBackToWorldTag>(Row);
+					}
 				}
-			}
-		}));
+			}));
 
-FAutoConsoleCommandWithArgsAndOutputDevice SetSelectionOverlayColorConsoleCommand(
-	TEXT("TEDS.Debug.SetOverlayColor"),
-	TEXT("Adds an overlay color to selected objects."),
-	FConsoleCommandWithArgsAndOutputDeviceDelegate::CreateLambda([](const TArray<FString>& Args, FOutputDevice& Output)
-		{
-			using namespace TypedElementQueryBuilder;
-			using DSI = ITypedElementDataStorageInterface;
+	FAutoConsoleCommandWithArgsAndOutputDevice SetSelectionOverlayColorConsoleCommand(
+		TEXT("TEDS.Debug.SetOverlayColor"),
+		TEXT("Adds an overlay color to selected objects."),
+		FConsoleCommandWithArgsAndOutputDeviceDelegate::CreateLambda([](const TArray<FString>& Args, FOutputDevice& Output)
+			{
+				using namespace TypedElementQueryBuilder;
+				using DSI = ITypedElementDataStorageInterface;
 
-			TRACE_CPUPROFILER_EVENT_SCOPE(TEDS.Debug.AddOverlayColorToSelectionCommand);
+				TRACE_CPUPROFILER_EVENT_SCOPE(TEDS.Debug.AddOverlayColorToSelectionCommand);
 
-			if (ITypedElementDataStorageInterface* DataStorage = UTypedElementRegistry::GetInstance()->GetMutableDataStorage())
-			{
-				static TypedElementQueryHandle OverlayQuery = TypedElementInvalidQueryHandle;
-				if (OverlayQuery == TypedElementInvalidQueryHandle)
+				if (ITypedElementDataStorageInterface* DataStorage = UTypedElementRegistry::GetInstance()->GetMutableDataStorage())
 				{
-					OverlayQuery = DataStorage->RegisterQuery(
-						Select()
-						.Where()
-							.All<FTypedElementSelectionColumn>()
-						.Compile());
-				}
+					static TypedElementQueryHandle OverlayQuery = TypedElementInvalidQueryHandle;
+					if (OverlayQuery == TypedElementInvalidQueryHandle)
+					{
+						OverlayQuery = DataStorage->RegisterQuery(
+							Select()
+							.Where()
+								.All<FTypedElementSelectionColumn>()
+							.Compile());
+					}
 				
-				if (OverlayQuery == TypedElementInvalidQueryHandle)
-				{
-					return;
-				}
+					if (OverlayQuery == TypedElementInvalidQueryHandle)
+					{
+						return;
+					}
 
-				if (Args.IsEmpty())
-				{
-					Output.Log(TEXT("Provide a color in hexadecimal format (#RRGGBBAA) to overlay."));
-					return;
-				}
+					if (Args.IsEmpty())
+					{
+						Output.Log(TEXT("Provide a color in hexadecimal format (#RRGGBBAA) to overlay."));
+						return;
+					}
 				
-				// Parse the color
-				FColor Color = FColor::FromHex(Args[0]);
-				Color.A = FMath::Clamp(Color.A, 0, 128);
+					// Parse the color
+					FColor Color = FColor::FromHex(Args[0]);
+					Color.A = FMath::Clamp(Color.A, 0, 128);
 
-				TArray<TypedElementRowHandle> RowHandles;
+					TArray<TypedElementRowHandle> RowHandles;
 				
-				DataStorage->RunQuery(OverlayQuery, [Color, &RowHandles](const DSI::FQueryDescription&, DSI::IDirectQueryContext& Context)
-				{
-					RowHandles = Context.GetRowHandles();
-				});
+					DataStorage->RunQuery(OverlayQuery, [Color, &RowHandles](const DSI::FQueryDescription&, DSI::IDirectQueryContext& Context)
+					{
+						RowHandles = Context.GetRowHandles();
+					});
 
-				for (TypedElementRowHandle Row : RowHandles)
-				{
-					DataStorage->RemoveColumn<FTypedElementViewportOverlayColorColumn>(Row);
-					DataStorage->AddColumn(Row, FTypedElementViewportOverlayColorColumn{ .OverlayColor = Color });
+					for (TypedElementRowHandle Row : RowHandles)
+					{
+						DataStorage->RemoveColumn<FTypedElementViewportOverlayColorColumn>(Row);
+						DataStorage->AddColumn(Row, FTypedElementViewportOverlayColorColumn{ .OverlayColor = Color });
+					}
 				}
-			}
-		}));
+			}));
 
-FAutoConsoleCommandWithArgsAndOutputDevice RemoveSelectionOverlayColorConsoleCommand(
-	TEXT("TEDS.Debug.RemoveOverlayColor"),
-	TEXT("Removes an overlay color to selected objects."),
-	FConsoleCommandWithArgsAndOutputDeviceDelegate::CreateLambda([](const TArray<FString>& Args, FOutputDevice& Output)
-		{
-			using namespace TypedElementQueryBuilder;
-			using DSI = ITypedElementDataStorageInterface;
+	FAutoConsoleCommandWithArgsAndOutputDevice RemoveSelectionOverlayColorConsoleCommand(
+		TEXT("TEDS.Debug.RemoveOverlayColor"),
+		TEXT("Removes an overlay color to selected objects."),
+		FConsoleCommandWithArgsAndOutputDeviceDelegate::CreateLambda([](const TArray<FString>& Args, FOutputDevice& Output)
+			{
+				using namespace TypedElementQueryBuilder;
+				using DSI = ITypedElementDataStorageInterface;
 
-			TRACE_CPUPROFILER_EVENT_SCOPE(TEDS.Debug.AddOverlayColorToSelectionCommand);
+				TRACE_CPUPROFILER_EVENT_SCOPE(TEDS.Debug.AddOverlayColorToSelectionCommand);
 
-			if (ITypedElementDataStorageInterface* DataStorage = UTypedElementRegistry::GetInstance()->GetMutableDataStorage())
-			{
-				static TypedElementQueryHandle OverlayQuery = TypedElementInvalidQueryHandle;
-				if (OverlayQuery == TypedElementInvalidQueryHandle)
+				if (ITypedElementDataStorageInterface* DataStorage = UTypedElementRegistry::GetInstance()->GetMutableDataStorage())
 				{
-					OverlayQuery = DataStorage->RegisterQuery(
-						Select()
-						.Where()
-							.All<FTypedElementSelectionColumn>()
-						.Compile());
-				}
+					static TypedElementQueryHandle OverlayQuery = TypedElementInvalidQueryHandle;
+					if (OverlayQuery == TypedElementInvalidQueryHandle)
+					{
+						OverlayQuery = DataStorage->RegisterQuery(
+							Select()
+							.Where()
+								.All<FTypedElementSelectionColumn>()
+							.Compile());
+					}
 				
-				if (OverlayQuery == TypedElementInvalidQueryHandle)
-				{
-					return;
-				}
+					if (OverlayQuery == TypedElementInvalidQueryHandle)
+					{
+						return;
+					}
 
-				TArray<TypedElementRowHandle> RowHandles;
+					TArray<TypedElementRowHandle> RowHandles;
 				
-				DataStorage->RunQuery(OverlayQuery, [&RowHandles](const DSI::FQueryDescription&, DSI::IDirectQueryContext& Context)
-				{
-					RowHandles = Context.GetRowHandles();
-				});
+					DataStorage->RunQuery(OverlayQuery, [&RowHandles](const DSI::FQueryDescription&, DSI::IDirectQueryContext& Context)
+					{
+						RowHandles = Context.GetRowHandles();
+					});
 
-				for (TypedElementRowHandle Row : RowHandles)
-				{
-					DataStorage->RemoveColumn<FTypedElementViewportOverlayColorColumn>(Row);
+					for (TypedElementRowHandle Row : RowHandles)
+					{
+						DataStorage->RemoveColumn<FTypedElementViewportOverlayColorColumn>(Row);
+					}
 				}
-			}
-		}));
+			}));
+} // namespace UE::Editor::DataStorage::Private
 
-void UTypedElementActorViewportFactory::RegisterQueries(ITypedElementDataStorageInterface& DataStorage)
+void UActorViewportDataStorageFactory::RegisterQueries(ITypedElementDataStorageInterface& DataStorage)
 {
 	RegisterOutlineColorColumnToActor(DataStorage);
 	RegisterOverlayColorColumnToActor(DataStorage);
 }
 
-void UTypedElementActorViewportFactory::RegisterOutlineColorColumnToActor(ITypedElementDataStorageInterface& DataStorage)
+void UActorViewportDataStorageFactory::RegisterOutlineColorColumnToActor(ITypedElementDataStorageInterface& DataStorage)
 {
 	using namespace TypedElementQueryBuilder;
 	using namespace TypedElementDataStorage;
@@ -198,7 +201,7 @@ void UTypedElementActorViewportFactory::RegisterOutlineColorColumnToActor(ITyped
 	);
 }
 
-void UTypedElementActorViewportFactory::RegisterOverlayColorColumnToActor(ITypedElementDataStorageInterface& DataStorage)
+void UActorViewportDataStorageFactory::RegisterOverlayColorColumnToActor(ITypedElementDataStorageInterface& DataStorage)
 {
 	using namespace TypedElementQueryBuilder;
 	using namespace TypedElementDataStorage;
diff --git a/Engine/Plugins/Experimental/EditorDataStorage/Source/TedsCore/Private/Compatibility/TypedElementActorViewportProcessors.h b/Engine/Plugins/Experimental/EditorDataStorage/Source/TedsCore/Private/Compatibility/TypedElementActorViewportProcessors.h
--- a/Engine/Plugins/Experimental/EditorDataStorage/Source/TedsCore/Private/Compatibility/TypedElementActorViewportProcessors.h
+++ b/Engine/Plugins/Experimental/EditorDataStorage/Source/TedsCore/Private/Compatibility/TypedElementActorViewportProcessors.h
@@ -1,4 +1,4 @@
-﻿// Copyright Epic Games, Inc. All Rights Reserved.
+// Copyright Epic Games, Inc. All Rights Reserved.
 
 #pragma once
 
@@ -8,12 +8,12 @@
 #include "TypedElementActorViewportProcessors.generated.h"
 
 UCLASS()
-class UTypedElementActorViewportFactory : public UTypedElementDataStorageFactory
+class UActorViewportDataStorageFactory : public UTypedElementDataStorageFactory
 {
 	GENERATED_BODY()
 
 public:
-	~UTypedElementActorViewportFactory() override = default;
+	~UActorViewportDataStorageFactory() override = default;
 
 	void RegisterQueries(ITypedElementDataStorageInterface& DataStorage) override;
 
diff --git a/Engine/Plugins/Experimental/EditorDataStorage/Source/TedsCore/Private/Compatibility/TypedElementObjectReinstancingManager.cpp b/Engine/Plugins/Experimental/EditorDataStorage/Source/TedsCore/Private/Compatibility/TypedElementObjectReinstancingManager.cpp
--- a/Engine/Plugins/Experimental/EditorDataStorage/Source/TedsCore/Private/Compatibility/TypedElementObjectReinstancingManager.cpp
+++ b/Engine/Plugins/Experimental/EditorDataStorage/Source/TedsCore/Private/Compatibility/TypedElementObjectReinstancingManager.cpp
@@ -11,12 +11,12 @@
 
 DECLARE_LOG_CATEGORY_CLASS(LogTedsObjectReinstancing, Log, Log)
 
-UTypedElementObjectReinstancingManager::UTypedElementObjectReinstancingManager()
+UTedsObjectReinstancingManager::UTedsObjectReinstancingManager()
 	: MementoRowBaseTable(TypedElementInvalidTableHandle)
 {
 }
 
-void UTypedElementObjectReinstancingManager::Initialize(UTypedElementDatabase& InDatabase, UTypedElementDatabaseCompatibility& InDataStorageCompatibility)
+void UTedsObjectReinstancingManager::Initialize(UTypedElementDatabase& InDatabase, UTypedElementDatabaseCompatibility& InDataStorageCompatibility)
 {
 	using namespace TypedElementDataStorage;
 	using namespace UE::Editor::DataStorage;
@@ -25,9 +25,9 @@ void UTypedElementObjectReinstancingManager::Initialize(UTypedElementDatabase& I
 	DataStorageCompatibility = &InDataStorageCompatibility;
 	
 	UpdateCompletedCallbackHandle = 
-		Database->OnUpdateCompleted().AddUObject(this, &UTypedElementObjectReinstancingManager::UpdateCompleted);
+		Database->OnUpdateCompleted().AddUObject(this, &UTedsObjectReinstancingManager::UpdateCompleted);
 	ReinstancingCallbackHandle = 
-		FCoreUObjectDelegates::OnObjectsReinstanced.AddUObject(this, &UTypedElementObjectReinstancingManager::HandleOnObjectsReinstanced);
+		FCoreUObjectDelegates::OnObjectsReinstanced.AddUObject(this, &UTedsObjectReinstancingManager::HandleOnObjectsReinstanced);
 	ObjectRemovedCallbackHandle = DataStorageCompatibility->RegisterObjectRemovedCallback(
 		[this](const void* Object, const FObjectTypeInfo& TypeInfo, RowHandle Row)
 		{
@@ -35,7 +35,7 @@ void UTypedElementObjectReinstancingManager::Initialize(UTypedElementDatabase& I
 		});
 }
 
-void UTypedElementObjectReinstancingManager::Deinitialize()
+void UTedsObjectReinstancingManager::Deinitialize()
 {
 	FCoreUObjectDelegates::OnObjectsReinstanced.Remove(ReinstancingCallbackHandle);
 	DataStorageCompatibility->UnregisterObjectRemovedCallback(ObjectRemovedCallbackHandle);
@@ -45,7 +45,7 @@ void UTypedElementObjectReinstancingManager::Deinitialize()
 	Database = nullptr;
 }
 
-void UTypedElementObjectReinstancingManager::UpdateCompleted()
+void UTedsObjectReinstancingManager::UpdateCompleted()
 {
 	using namespace TypedElementDataStorage;
 
@@ -57,7 +57,7 @@ void UTypedElementObjectReinstancingManager::UpdateCompleted()
 	OldObjectToMementoMap.Reset();
 }
 
-void UTypedElementObjectReinstancingManager::HandleOnObjectPreRemoved(
+void UTedsObjectReinstancingManager::HandleOnObjectPreRemoved(
 	const void* Object, 
 	const UE::Editor::DataStorage::FObjectTypeInfo& TypeInfo, 
 	TypedElementDataStorage::RowHandle ObjectRow)
@@ -67,7 +67,7 @@ void UTypedElementObjectReinstancingManager::HandleOnObjectPreRemoved(
 	OldObjectToMementoMap.Add(Object, Memento);
 }
 
-void UTypedElementObjectReinstancingManager::HandleOnObjectsReinstanced(
+void UTedsObjectReinstancingManager::HandleOnObjectsReinstanced(
 	const FCoreUObjectDelegates::FReplacementObjectMap& ObjectReplacementMap)
 {
 	ITypedElementDataStorageInterface* Interface = Database;
diff --git a/Engine/Plugins/Experimental/EditorDataStorage/Source/TedsCore/Private/Compatibility/TypedElementObjectReinstancingManager.h b/Engine/Plugins/Experimental/EditorDataStorage/Source/TedsCore/Private/Compatibility/TypedElementObjectReinstancingManager.h
--- a/Engine/Plugins/Experimental/EditorDataStorage/Source/TedsCore/Private/Compatibility/TypedElementObjectReinstancingManager.h
+++ b/Engine/Plugins/Experimental/EditorDataStorage/Source/TedsCore/Private/Compatibility/TypedElementObjectReinstancingManager.h
@@ -20,11 +20,11 @@ namespace UE::Editor::DataStorage
 }
 
 UCLASS(Transient)
-class UTypedElementObjectReinstancingManager : public UObject
+class UTedsObjectReinstancingManager : public UObject
 {
 	GENERATED_BODY()
 public:
-	UTypedElementObjectReinstancingManager();
+	UTedsObjectReinstancingManager();
 
 	void Initialize(UTypedElementDatabase& InDatabase, UTypedElementDatabaseCompatibility& InDataStorageCompatibility);
 	void Deinitialize();
diff --git a/Engine/Plugins/Experimental/EditorDataStorage/Source/TedsCore/Private/Compatibility/TypedElementUObjectWorldQueries.cpp b/Engine/Plugins/Experimental/EditorDataStorage/Source/TedsCore/Private/Compatibility/TypedElementUObjectWorldQueries.cpp
--- a/Engine/Plugins/Experimental/EditorDataStorage/Source/TedsCore/Private/Compatibility/TypedElementUObjectWorldQueries.cpp
+++ b/Engine/Plugins/Experimental/EditorDataStorage/Source/TedsCore/Private/Compatibility/TypedElementUObjectWorldQueries.cpp
@@ -7,13 +7,13 @@
 #include "Elements/Framework/TypedElementQueryBuilder.h"
 #include "Engine/World.h"
 
-void UTypedElementUObjectWorldFactory::RegisterQueries(ITypedElementDataStorageInterface& DataStorage)
+void UObjectWorldDataStorageFactory::RegisterQueries(ITypedElementDataStorageInterface& DataStorage)
 {
 	RegisterAddWorldColumn(DataStorage);
 	RegisterUpdateOrRemoveWorldColumn(DataStorage);
 }
 
-void UTypedElementUObjectWorldFactory::RegisterAddWorldColumn(ITypedElementDataStorageInterface& DataStorage) const
+void UObjectWorldDataStorageFactory::RegisterAddWorldColumn(ITypedElementDataStorageInterface& DataStorage) const
 {
 	using namespace TypedElementDataStorage;
 	using namespace TypedElementQueryBuilder;
@@ -43,7 +43,7 @@ void UTypedElementUObjectWorldFactory::RegisterAddWorldColumn(ITypedElementDataS
 	);
 }
 
-void UTypedElementUObjectWorldFactory::RegisterUpdateOrRemoveWorldColumn(ITypedElementDataStorageInterface& DataStorage) const
+void UObjectWorldDataStorageFactory::RegisterUpdateOrRemoveWorldColumn(ITypedElementDataStorageInterface& DataStorage) const
 {
 	using namespace TypedElementDataStorage;
 	using namespace TypedElementQueryBuilder;
diff --git a/Engine/Plugins/Experimental/EditorDataStorage/Source/TedsCore/Private/Compatibility/TypedElementUObjectWorldQueries.h b/Engine/Plugins/Experimental/EditorDataStorage/Source/TedsCore/Private/Compatibility/TypedElementUObjectWorldQueries.h
--- a/Engine/Plugins/Experimental/EditorDataStorage/Source/TedsCore/Private/Compatibility/TypedElementUObjectWorldQueries.h
+++ b/Engine/Plugins/Experimental/EditorDataStorage/Source/TedsCore/Private/Compatibility/TypedElementUObjectWorldQueries.h
@@ -8,12 +8,12 @@
 #include "TypedElementUObjectWorldQueries.generated.h"
 
 UCLASS()
-class UTypedElementUObjectWorldFactory : public UTypedElementDataStorageFactory
+class UObjectWorldDataStorageFactory : public UTypedElementDataStorageFactory
 {
 	GENERATED_BODY()
 
 public:
-	~UTypedElementUObjectWorldFactory() override = default;
+	~UObjectWorldDataStorageFactory() override = default;
 
 	void RegisterQueries(ITypedElementDataStorageInterface& DataStorage) override;
 
diff --git a/Engine/Plugins/Experimental/EditorDataStorage/Source/TedsCore/Private/TypedElementsDataStorage.cpp b/Engine/Plugins/Experimental/EditorDataStorage/Source/TedsCore/Private/TypedElementsDataStorage.cpp
--- a/Engine/Plugins/Experimental/EditorDataStorage/Source/TedsCore/Private/TypedElementsDataStorage.cpp
+++ b/Engine/Plugins/Experimental/EditorDataStorage/Source/TedsCore/Private/TypedElementsDataStorage.cpp
@@ -87,7 +87,7 @@ void FTypedElementsDataStorageModule::StartupModule()
 				DatabaseUi = NewObject<UTypedElementDatabaseUi>();
 				DatabaseUi->Initialize(Database.Get(), DatabaseCompatibility.Get());
 
-				ObjectReinstancingManager = NewObject<UTypedElementObjectReinstancingManager>();
+				ObjectReinstancingManager = NewObject<UTedsObjectReinstancingManager>();
 				ObjectReinstancingManager->Initialize(*Database, *DatabaseCompatibility);
 
 				// Register the various database instances.
diff --git a/Engine/Plugins/Experimental/EditorDataStorage/Source/TedsCore/Public/Compatibility/Columns/TypedElement.h b/Engine/Plugins/Experimental/EditorDataStorage/Source/TedsCore/Public/Compatibility/Columns/TypedElement.h
new file mode 100644
--- /dev/null
+++ b/Engine/Plugins/Experimental/EditorDataStorage/Source/TedsCore/Public/Compatibility/Columns/TypedElement.h
@@ -0,0 +1,23 @@
+// Copyright Epic Games, Inc. All Rights Reserved.
+
+#pragma once
+
+#include "Elements/Common/TypedElementCommonTypes.h"
+#include "Elements/Framework/TypedElementHandle.h"
+
+#include "TypedElement.generated.h"
+
+// A column which contains a TypedElementHandle from the TypedElementFramework
+// Prefer to use the namespaced alias when referencing in code - UE::Editor::DataStorage::Compatibility::FTypedElementColumn
+USTRUCT()
+struct FTedsTypedElementColumn : public FTypedElementDataStorageColumn
+{
+	GENERATED_BODY()
+
+	FTypedElementHandle Handle;
+};
+
+namespace UE::Editor::DataStorage::Compatibility
+{
+	using FTypedElementColumn = FTedsTypedElementColumn;
+}  // namespace UE::Editor::DataStorage::Compatibility
diff --git a/Engine/Plugins/Experimental/EditorDataStorage/Source/TedsCore/Public/Compatibility/TedsTypedElementBridge.h b/Engine/Plugins/Experimental/EditorDataStorage/Source/TedsCore/Public/Compatibility/TedsTypedElementBridge.h
--- a/Engine/Plugins/Experimental/EditorDataStorage/Source/TedsCore/Public/Compatibility/TedsTypedElementBridge.h
+++ b/Engine/Plugins/Experimental/EditorDataStorage/Source/TedsCore/Public/Compatibility/TedsTypedElementBridge.h
@@ -2,52 +2,12 @@
 
 #pragma once
 
-#include "Elements/Framework/TypedElementHandle.h"
-#include "Elements/Interfaces/TypedElementDataStorageInterface.h"
-#include "Elements/Interfaces/TypedElementDataStorageFactory.h"
-#include "UObject/UObjectGlobals.h"
+#include "Delegates/Delegate.h"
 
-#include "TedsTypedElementBridge.generated.h"
-
-class UTypedElementRegistry;
-
-DECLARE_MULTICAST_DELEGATE_OneParam(FOnTEDSTypedElementBridgeEnable, bool /*bEnabled*/);
-
-/**
- * This class is responsible for running queries that will ensure TypedElementHandles
- * are cleaned up when TEDS is shut down.
- */
-UCLASS(Transient)
-class UTEDSTypedElementBridge : public UTypedElementDataStorageFactory
+namespace UE::Editor::DataStorage::Compatibility
 {
-	GENERATED_BODY()
-public:
-	~UTEDSTypedElementBridge() override = default;
-	virtual uint8 GetOrder() const override;
-	virtual void PreRegister(ITypedElementDataStorageInterface& DataStorage) override;
-	virtual void PreShutdown(ITypedElementDataStorageInterface& DataStorage) override;
-	virtual void RegisterQueries(ITypedElementDataStorageInterface& DataStorage) override;
+	DECLARE_MULTICAST_DELEGATE_OneParam(FOnTypedElementBridgeEnabled, bool /*bEnabled*/);
 
-	static TEDSCORE_API FOnTEDSTypedElementBridgeEnable& OnEnabled();
-	static TEDSCORE_API bool IsEnabled();
-
-private:
-	void RegisterQuery_NewUObject(ITypedElementDataStorageInterface& DataStorage);
-	void UnregisterQuery_NewUObject(ITypedElementDataStorageInterface& DataStorage);
-	void CleanupTypedElementColumns(ITypedElementDataStorageInterface& DataStorage);
-	void HandleOnEnabled(IConsoleVariable* CVar);
-	
-	TypedElementQueryHandle RemoveTypedElementRowHandleQuery = TypedElementDataStorage::InvalidQueryHandle;
-	FDelegateHandle DebugEnabledDelegateHandle;
-};
-
-// A column which contains a TypedElementHandle
-// TypedElement, in this context, refers to the interface-based
-// TypedElements instead of the data-based TypedElementDataStorage (TEDS)
-USTRUCT()
-struct FTEDSTypedElementColumn : public FTypedElementDataStorageColumn
-{
-	GENERATED_BODY()
-	
-	FTypedElementHandle Handle;
-};
\ No newline at end of file
+	TEDSCORE_API FOnTypedElementBridgeEnabled& OnTypedElementBridgeEnabled();
+	TEDSCORE_API bool IsTypedElementBridgeEnabled();
+} // namespace UE::Editor::DataStorage::Compatibility
diff --git a/Engine/Plugins/Experimental/EditorDataStorage/Source/TedsCore/Public/TypedElementsDataStorage.h b/Engine/Plugins/Experimental/EditorDataStorage/Source/TedsCore/Public/TypedElementsDataStorage.h
--- a/Engine/Plugins/Experimental/EditorDataStorage/Source/TedsCore/Public/TypedElementsDataStorage.h
+++ b/Engine/Plugins/Experimental/EditorDataStorage/Source/TedsCore/Public/TypedElementsDataStorage.h
@@ -11,7 +11,7 @@ class FReferenceCollector;
 class UTypedElementDatabase;
 class UTypedElementDatabaseCompatibility;
 class UTypedElementDatabaseUi;
-class UTypedElementObjectReinstancingManager;
+class UTedsObjectReinstancingManager;
 
 class FTypedElementsDataStorageModule : public IModuleInterface, public FGCObject
 {
@@ -37,6 +37,6 @@ private:
 	TObjectPtr<UTypedElementDatabase> Database;
 	TObjectPtr<UTypedElementDatabaseCompatibility> DatabaseCompatibility;
 	TObjectPtr<UTypedElementDatabaseUi> DatabaseUi;
-	TObjectPtr<UTypedElementObjectReinstancingManager> ObjectReinstancingManager;
+	TObjectPtr<UTedsObjectReinstancingManager> ObjectReinstancingManager;
 	bool bInitialized{ false };
 };
