[Teds] Renaming more pieces of TedsCore - loose files and memento files

#rnx

diff --git a/Engine/Plugins/Experimental/EditorDataStorage/Source/TedsCore/Private/Commands/EditorDataStorageCompatibilityCommands.cpp b/Engine/Plugins/Experimental/EditorDataStorage/Source/TedsCore/Private/Commands/EditorDataStorageCompatibilityCommands.cpp
--- a/Engine/Plugins/Experimental/EditorDataStorage/Source/TedsCore/Private/Commands/EditorDataStorageCompatibilityCommands.cpp
+++ b/Engine/Plugins/Experimental/EditorDataStorage/Source/TedsCore/Private/Commands/EditorDataStorageCompatibilityCommands.cpp
@@ -153,7 +153,7 @@ namespace UE::Editor::DataStorage
 		return Commands.GetCommandCount<FTypeInfoReinstanced>() > 0;
 	}
 
-	void FPatchData::RunPatch(CompatibilityCommandBuffer::FCollection& Commands, UTypedElementDatabaseCompatibility& StorageCompat,
+	void FPatchData::RunPatch(CompatibilityCommandBuffer::FCollection& Commands, UEditorDataStorageCompatibility& StorageCompat,
 		FScratchBuffer& ScratchBuffer)
 	{
 		using namespace TypedElementDataStorage;
@@ -227,7 +227,7 @@ namespace UE::Editor::DataStorage
 	// FPrepareCommands
 	//
 
-	FPrepareCommands::FPrepareCommands(ITypedElementDataStorageInterface& InStorage, UTypedElementDatabaseCompatibility& InStorageCompat,
+	FPrepareCommands::FPrepareCommands(ITypedElementDataStorageInterface& InStorage, UEditorDataStorageCompatibility& InStorageCompat,
 		CompatibilityCommandBuffer::FCollection& InCommands)
 		: Storage(InStorage)
 		, StorageCompat(InStorageCompat)
@@ -317,7 +317,7 @@ namespace UE::Editor::DataStorage
 		}
 	}
 
-	void FPrepareCommands::RunPreparation(ITypedElementDataStorageInterface& Storage, UTypedElementDatabaseCompatibility& StorageCompat,
+	void FPrepareCommands::RunPreparation(ITypedElementDataStorageInterface& Storage, UEditorDataStorageCompatibility& StorageCompat,
 		CompatibilityCommandBuffer::FCollection& Commands)
 	{
 		FPrepareCommands PrepareVisitor(Storage, StorageCompat, Commands);
@@ -507,7 +507,7 @@ namespace UE::Editor::DataStorage
 	// FCommandProcessor
 	//
 
-	FCommandProcessor::FCommandProcessor(ITypedElementDataStorageInterface& InStorage, UTypedElementDatabaseCompatibility& InStorageCompatibility)
+	FCommandProcessor::FCommandProcessor(ITypedElementDataStorageInterface& InStorage, UEditorDataStorageCompatibility& InStorageCompatibility)
 		: Storage(InStorage)
 		, StorageCompatibility(InStorageCompatibility)
 		, MementoSystem(InStorageCompatibility.Environment->GetMementoSystem())
diff --git a/Engine/Plugins/Experimental/EditorDataStorage/Source/TedsCore/Private/Commands/EditorDataStorageCompatibilityCommands.h b/Engine/Plugins/Experimental/EditorDataStorage/Source/TedsCore/Private/Commands/EditorDataStorageCompatibilityCommands.h
--- a/Engine/Plugins/Experimental/EditorDataStorage/Source/TedsCore/Private/Commands/EditorDataStorageCompatibilityCommands.h
+++ b/Engine/Plugins/Experimental/EditorDataStorage/Source/TedsCore/Private/Commands/EditorDataStorageCompatibilityCommands.h
@@ -10,14 +10,14 @@
 #include "UObject/WeakObjectPtrTemplates.h"
 
 class ITypedElementDataStorageInterface;
-class UTypedElementMementoSystem;
 class UScriptStruct;
-class UTypedElementDatabaseCompatibility;
+class UEditorDataStorageCompatibility;
 class UObject;
 
 namespace UE::Editor::DataStorage
 {
 	class FScratchBuffer;
+	class FMementoSystem;
 
 	enum class EObjectType : uint8
 	{
@@ -221,14 +221,14 @@ namespace UE::Editor::DataStorage
 		};
 
 		static bool IsPatchingRequired(const CompatibilityCommandBuffer::FCollection& Commands);
-		static void RunPatch(CompatibilityCommandBuffer::FCollection& Commands, UTypedElementDatabaseCompatibility& StorageCompat,
+		static void RunPatch(CompatibilityCommandBuffer::FCollection& Commands, UEditorDataStorageCompatibility& StorageCompat,
 			FScratchBuffer& ScratchBuffer);
 	};
 
 	/** Prepares each command for further processing, e.g. resolving the target table. */
 	struct FPrepareCommands final
 	{
-		FPrepareCommands(ITypedElementDataStorageInterface& InStorage, UTypedElementDatabaseCompatibility& InStorageCompat,
+		FPrepareCommands(ITypedElementDataStorageInterface& InStorage, UEditorDataStorageCompatibility& InStorageCompat,
 			CompatibilityCommandBuffer::FCollection& InCommands);
 
 		template<typename T>
@@ -242,11 +242,11 @@ namespace UE::Editor::DataStorage
 		void operator()(FAddSyncFromWorldTag& Command);
 
 		ITypedElementDataStorageInterface& Storage;
-		UTypedElementDatabaseCompatibility& StorageCompat;
+		UEditorDataStorageCompatibility& StorageCompat;
 		CompatibilityCommandBuffer::FCollection& Commands;
 		int32 CurrentIndex = 0;
 
-		static void RunPreparation(ITypedElementDataStorageInterface& Storage, UTypedElementDatabaseCompatibility& StorageCompat, 
+		static void RunPreparation(ITypedElementDataStorageInterface& Storage, UEditorDataStorageCompatibility& StorageCompat, 
 			CompatibilityCommandBuffer::FCollection& Commands);
 	};
 
@@ -299,7 +299,7 @@ namespace UE::Editor::DataStorage
 	/** Executes the commands in the command buffer for TEDS Compatibility. */
 	struct FCommandProcessor final
 	{
-		FCommandProcessor(ITypedElementDataStorageInterface& InStorage, UTypedElementDatabaseCompatibility& InStorageCompatibility);
+		FCommandProcessor(ITypedElementDataStorageInterface& InStorage, UEditorDataStorageCompatibility& InStorageCompatibility);
 
 		void SetupRow(TypedElementDataStorage::RowHandle Row, UObject* Object);
 		void SetupRow(TypedElementDataStorage::RowHandle Row, void* Object, TWeakObjectPtr<UScriptStruct> TypeInfo);
@@ -324,8 +324,8 @@ namespace UE::Editor::DataStorage
 		void operator()(FAddSyncFromWorldTag& Command);
 
 		ITypedElementDataStorageInterface& Storage;
-		UTypedElementDatabaseCompatibility& StorageCompatibility;
-		UTypedElementMementoSystem& MementoSystem;
+		UEditorDataStorageCompatibility& StorageCompatibility;
+		FMementoSystem& MementoSystem;
 	};
 
 	struct FRecordCommands final
diff --git a/Engine/Plugins/Experimental/EditorDataStorage/Source/TedsCore/Private/Compatibility/TypedElementObjectReinstancingManager.cpp b/Engine/Plugins/Experimental/EditorDataStorage/Source/TedsCore/Private/Compatibility/TypedElementObjectReinstancingManager.cpp
--- a/Engine/Plugins/Experimental/EditorDataStorage/Source/TedsCore/Private/Compatibility/TypedElementObjectReinstancingManager.cpp
+++ b/Engine/Plugins/Experimental/EditorDataStorage/Source/TedsCore/Private/Compatibility/TypedElementObjectReinstancingManager.cpp
@@ -16,16 +16,16 @@ UTedsObjectReinstancingManager::UTedsObjectReinstancingManager()
 {
 }
 
-void UTedsObjectReinstancingManager::Initialize(UTypedElementDatabase& InDatabase, UTypedElementDatabaseCompatibility& InDataStorageCompatibility)
+void UTedsObjectReinstancingManager::Initialize(UEditorDataStorage& InDataStorage, UEditorDataStorageCompatibility& InDataStorageCompatibility)
 {
 	using namespace TypedElementDataStorage;
 	using namespace UE::Editor::DataStorage;
 
-	Database = &InDatabase;
+	DataStorage = &InDataStorage;
 	DataStorageCompatibility = &InDataStorageCompatibility;
 	
 	UpdateCompletedCallbackHandle = 
-		Database->OnUpdateCompleted().AddUObject(this, &UTedsObjectReinstancingManager::UpdateCompleted);
+		DataStorage->OnUpdateCompleted().AddUObject(this, &UTedsObjectReinstancingManager::UpdateCompleted);
 	ReinstancingCallbackHandle = 
 		FCoreUObjectDelegates::OnObjectsReinstanced.AddUObject(this, &UTedsObjectReinstancingManager::HandleOnObjectsReinstanced);
 	ObjectRemovedCallbackHandle = DataStorageCompatibility->RegisterObjectRemovedCallback(
@@ -39,17 +39,18 @@ void UTedsObjectReinstancingManager::Deinitialize()
 {
 	FCoreUObjectDelegates::OnObjectsReinstanced.Remove(ReinstancingCallbackHandle);
 	DataStorageCompatibility->UnregisterObjectRemovedCallback(ObjectRemovedCallbackHandle);
-	Database->OnUpdateCompleted().Remove(UpdateCompletedCallbackHandle);
+	DataStorage->OnUpdateCompleted().Remove(UpdateCompletedCallbackHandle);
 
 	DataStorageCompatibility = nullptr;
-	Database = nullptr;
+	DataStorage = nullptr;
 }
 
 void UTedsObjectReinstancingManager::UpdateCompleted()
 {
 	using namespace TypedElementDataStorage;
+	using namespace UE::Editor::DataStorage;
 
-	UTypedElementMementoSystem& MementoSystem = Database->GetEnvironment()->GetMementoSystem();
+	FMementoSystem& MementoSystem = DataStorage->GetEnvironment()->GetMementoSystem();
 	for (TMap<const void*, RowHandle>::TConstIterator It = OldObjectToMementoMap.CreateConstIterator(); It; ++It)
 	{
 		MementoSystem.DestroyMemento(It.Value());
@@ -63,14 +64,14 @@ void UTedsObjectReinstancingManager::HandleOnObjectPreRemoved(
 	TypedElementDataStorage::RowHandle ObjectRow)
 {
 	// This is the chance to record the old object to memento
-	TypedElementDataStorage::RowHandle Memento = Database->GetEnvironment()->GetMementoSystem().CreateMemento(ObjectRow);
+	TypedElementDataStorage::RowHandle Memento = DataStorage->GetEnvironment()->GetMementoSystem().CreateMemento(ObjectRow);
 	OldObjectToMementoMap.Add(Object, Memento);
 }
 
 void UTedsObjectReinstancingManager::HandleOnObjectsReinstanced(
 	const FCoreUObjectDelegates::FReplacementObjectMap& ObjectReplacementMap)
 {
-	ITypedElementDataStorageInterface* Interface = Database;
+	ITypedElementDataStorageInterface* Interface = DataStorage;
 	for (FCoreUObjectDelegates::FReplacementObjectMap::TConstIterator Iter = ObjectReplacementMap.CreateConstIterator(); Iter; ++Iter)
 	{
 		const void* PreDeleteObject = Iter->Key;
@@ -93,7 +94,7 @@ void UTedsObjectReinstancingManager::HandleOnObjectsReinstanced(
 			}
 
 			// Kick off re-instantiation of NewObjectRow from the Memento
-			Database->GetEnvironment()->GetMementoSystem().RestoreMemento(Memento, NewObjectRow);
+			DataStorage->GetEnvironment()->GetMementoSystem().RestoreMemento(Memento, NewObjectRow);
 		}
 	}
 }
diff --git a/Engine/Plugins/Experimental/EditorDataStorage/Source/TedsCore/Private/Compatibility/TypedElementObjectReinstancingManager.h b/Engine/Plugins/Experimental/EditorDataStorage/Source/TedsCore/Private/Compatibility/TypedElementObjectReinstancingManager.h
--- a/Engine/Plugins/Experimental/EditorDataStorage/Source/TedsCore/Private/Compatibility/TypedElementObjectReinstancingManager.h
+++ b/Engine/Plugins/Experimental/EditorDataStorage/Source/TedsCore/Private/Compatibility/TypedElementObjectReinstancingManager.h
@@ -11,8 +11,9 @@
 
 class ITypedElementDataStorageCompatibilityInterface;
 class ITypedElementDataStorageInterface;
-class UTypedElementDatabaseCompatibility;
+class UEditorDataStorageCompatibility;
 class UTypedElementMementoSystem;
+class UEditorDataStorage;
 
 namespace UE::Editor::DataStorage
 {
@@ -26,7 +27,7 @@ class UTedsObjectReinstancingManager : public UObject
 public:
 	UTedsObjectReinstancingManager();
 
-	void Initialize(UTypedElementDatabase& InDatabase, UTypedElementDatabaseCompatibility& InDataStorageCompatibility);
+	void Initialize(UEditorDataStorage& InDataStorage, UEditorDataStorageCompatibility& InDataStorageCompatibility);
 	void Deinitialize();
 
 private:
@@ -38,9 +39,9 @@ private:
 	void HandleOnObjectsReinstanced(const FCoreUObjectDelegates::FReplacementObjectMap& ObjectReplacementMap);
 
 	UPROPERTY()
-	TObjectPtr<UTypedElementDatabase> Database = nullptr;
+	TObjectPtr<UEditorDataStorage> DataStorage = nullptr;
 	UPROPERTY()
-	TObjectPtr<UTypedElementDatabaseCompatibility> DataStorageCompatibility = nullptr;
+	TObjectPtr<UEditorDataStorageCompatibility> DataStorageCompatibility = nullptr;
 	
 	// Reverse lookup that holds all populated mementos for recently deleted objects
 	// Entry removed when the memento is removed
diff --git a/Engine/Plugins/Experimental/EditorDataStorage/Source/TedsCore/Private/Memento/TypedElementMementoSystem.cpp b/Engine/Plugins/Experimental/EditorDataStorage/Source/TedsCore/Private/Memento/TypedElementMementoSystem.cpp
--- a/Engine/Plugins/Experimental/EditorDataStorage/Source/TedsCore/Private/Memento/TypedElementMementoSystem.cpp
+++ b/Engine/Plugins/Experimental/EditorDataStorage/Source/TedsCore/Private/Memento/TypedElementMementoSystem.cpp
@@ -9,119 +9,114 @@
 #include "TypedElementMementoRowTypes.h"
 #include "TypedElementDatabase.h"
 
-UTypedElementMementoSystem::UTypedElementMementoSystem(ITypedElementDataStorageInterface& InDataStorage)
-	: DataStorage(InDataStorage)
+namespace UE::Editor::DataStorage
 {
-	using namespace UE::Editor::DataStorage;
-
-	FScopedExclusiveLock Lock(EGlobalLockScope::Public);
+	FMementoSystem::FMementoSystem(ITypedElementDataStorageInterface& InDataStorage)
+		: DataStorage(InDataStorage)
+	{
+		FScopedExclusiveLock Lock(EGlobalLockScope::Public);
 
-	// Register tables that will be used by reinstancing
-	MementoRowBaseTable = DataStorage.RegisterTable<FTypedElementMementoTag>(TEXT("MementoRowBaseTable"));
+		// Register tables that will be used by reinstancing
+		MementoRowBaseTable = DataStorage.RegisterTable<FTypedElementMementoTag>(TEXT("MementoRowBaseTable"));
 
-	// Discover all MementoTranslators
-	{
-		constexpr bool bIncludeDerived = true;
-		constexpr EObjectFlags ExcludeFlags = EObjectFlags::RF_NoFlags;
-		ForEachObjectOfClass(UTypedElementMementoTranslatorBase::StaticClass(),
-			[this](UObject* Object)
-			{
-				const UTypedElementMementoTranslatorBase* TranslatorCandidate = Cast<UTypedElementMementoTranslatorBase>(Object);
-				// Exclude abstract classes
-				if (TranslatorCandidate->GetClass()->GetClassFlags() & EClassFlags::CLASS_Abstract)
+		// Discover all MementoTranslators
+		{
+			constexpr bool bIncludeDerived = true;
+			constexpr EObjectFlags ExcludeFlags = EObjectFlags::RF_NoFlags;
+			ForEachObjectOfClass(UTedsMementoTranslatorBase::StaticClass(),
+				[this](UObject* Object)
 				{
-					return;
-				}
-				MementoTranslators.Add(TranslatorCandidate);
-			},
-			bIncludeDerived, ExcludeFlags);
+					const UTedsMementoTranslatorBase* TranslatorCandidate = Cast<UTedsMementoTranslatorBase>(Object);
+					// Exclude abstract classes
+					if (TranslatorCandidate->GetClass()->GetClassFlags() & EClassFlags::CLASS_Abstract)
+					{
+						return;
+					}
+					MementoTranslators.Add(TranslatorCandidate);
+				},
+				bIncludeDerived, ExcludeFlags);
+		}
 	}
-}
-
-TypedElementRowHandle UTypedElementMementoSystem::CreateMemento(TypedElementDataStorage::RowHandle SourceRow)
-{
-	using namespace TypedElementDataStorage;
-	using namespace UE::Editor::DataStorage;
 
-	FScopedSharedLock Lock(EGlobalLockScope::Public);
+	TypedElementRowHandle FMementoSystem::CreateMemento(TypedElementDataStorage::RowHandle SourceRow)
+	{
+		using namespace TypedElementDataStorage;
 
-	RowHandle MementoRow = DataStorage.AddRow(MementoRowBaseTable);
-	CreateMementoInternal(MementoRow, SourceRow);
-	return MementoRow;
-}
+		FScopedSharedLock Lock(EGlobalLockScope::Public);
 
-void UTypedElementMementoSystem::CreateMemento(TypedElementDataStorage::RowHandle ReservedMementoRow, TypedElementDataStorage::RowHandle SourceRow)
-{
-	using namespace UE::Editor::DataStorage;
+		RowHandle MementoRow = DataStorage.AddRow(MementoRowBaseTable);
+		CreateMementoInternal(MementoRow, SourceRow);
+		return MementoRow;
+	}
 
-	FScopedSharedLock Lock(EGlobalLockScope::Public);
+	void FMementoSystem::CreateMemento(TypedElementDataStorage::RowHandle ReservedMementoRow, TypedElementDataStorage::RowHandle SourceRow)
+	{
+		FScopedSharedLock Lock(EGlobalLockScope::Public);
 	
-	DataStorage.AddRow(ReservedMementoRow, MementoRowBaseTable);
-	CreateMementoInternal(ReservedMementoRow, SourceRow);
-}
+		DataStorage.AddRow(ReservedMementoRow, MementoRowBaseTable);
+		CreateMementoInternal(ReservedMementoRow, SourceRow);
+	}
 
-void UTypedElementMementoSystem::CreateMementoInternal(TypedElementDataStorage::RowHandle MementoRow, TypedElementDataStorage::RowHandle SourceRow)
-{
-	using namespace TypedElementDataStorage;
-	using namespace UE::Editor::DataStorage;
-	
-	for (const UTypedElementMementoTranslatorBase* Translator : MementoTranslators)
+	void FMementoSystem::CreateMementoInternal(TypedElementDataStorage::RowHandle MementoRow, TypedElementDataStorage::RowHandle SourceRow)
 	{
-		if (void* SourceColumn = DataStorage.GetColumnData(SourceRow, Translator->GetColumnType()))
+		using namespace TypedElementDataStorage;
+	
+		for (const UTedsMementoTranslatorBase* Translator : MementoTranslators)
 		{
-			const UScriptStruct* MementoType = Translator->GetMementoType();
-			DataStorage.AddColumnData(MementoRow, MementoType,
-				[Translator, SourceColumn](void* MementoColumn, const UScriptStruct& ColumnType)
-				{
-					ColumnType.InitializeStruct(MementoColumn);
-					FScopedSharedLock Lock(EGlobalLockScope::Public);
-					Translator->TranslateColumnToMemento(SourceColumn, MementoColumn);
-				},
-				[](const UScriptStruct& ColumnType, void* Destination, void* Source)
-				{
-					ColumnType.CopyScriptStruct(Destination, Source);
-				});
-
-			UE_LOG(LogEditorDataStorage, VeryVerbose,
-				TEXT("Column->Memento: %llu -> %llu"), SourceRow, MementoRow);
+			if (void* SourceColumn = DataStorage.GetColumnData(SourceRow, Translator->GetColumnType()))
+			{
+				const UScriptStruct* MementoType = Translator->GetMementoType();
+				DataStorage.AddColumnData(MementoRow, MementoType,
+					[Translator, SourceColumn](void* MementoColumn, const UScriptStruct& ColumnType)
+					{
+						ColumnType.InitializeStruct(MementoColumn);
+						FScopedSharedLock Lock(EGlobalLockScope::Public);
+						Translator->TranslateColumnToMemento(SourceColumn, MementoColumn);
+					},
+					[](const UScriptStruct& ColumnType, void* Destination, void* Source)
+					{
+						ColumnType.CopyScriptStruct(Destination, Source);
+					});
+
+				UE_LOG(LogEditorDataStorage, VeryVerbose,
+					TEXT("Column->Memento: %llu -> %llu"), SourceRow, MementoRow);
+			}
 		}
 	}
-}
 
-void UTypedElementMementoSystem::RestoreMemento(TypedElementDataStorage::RowHandle MementoRow, TypedElementDataStorage::RowHandle TargetRow)
-{
-	using namespace UE::Editor::DataStorage;
-
-	FScopedSharedLock Lock(EGlobalLockScope::Public);
-
-	for (const UTypedElementMementoTranslatorBase* Translator : MementoTranslators)
+	void FMementoSystem::RestoreMemento(TypedElementDataStorage::RowHandle MementoRow, TypedElementDataStorage::RowHandle TargetRow)
 	{
-		if (void* MementoColumn = DataStorage.GetColumnData(MementoRow, Translator->GetMementoType()))
-		{
-			const UScriptStruct* TargetType = Translator->GetColumnType();
-			DataStorage.AddColumnData(TargetRow, TargetType,
-				[Translator, MementoColumn](void* TargetColumn, const UScriptStruct& ColumnType)
-				{
-					ColumnType.InitializeStruct(TargetColumn);
-					FScopedSharedLock Lock(EGlobalLockScope::Public);
-					Translator->TranslateMementoToColumn(MementoColumn, TargetColumn);
-				},
-				[](const UScriptStruct& ColumnType, void* Destination, void* Source)
-				{
-					ColumnType.CopyScriptStruct(Destination, Source);
-				});
+		FScopedSharedLock Lock(EGlobalLockScope::Public);
 
-			UE_LOG(LogEditorDataStorage, VeryVerbose, 
-				TEXT("Memento->Column: %llu -> %llu"), MementoRow, TargetRow);
+		for (const UTedsMementoTranslatorBase* Translator : MementoTranslators)
+		{
+			if (void* MementoColumn = DataStorage.GetColumnData(MementoRow, Translator->GetMementoType()))
+			{
+				const UScriptStruct* TargetType = Translator->GetColumnType();
+				DataStorage.AddColumnData(TargetRow, TargetType,
+					[Translator, MementoColumn](void* TargetColumn, const UScriptStruct& ColumnType)
+					{
+						ColumnType.InitializeStruct(TargetColumn);
+						FScopedSharedLock Lock(EGlobalLockScope::Public);
+						Translator->TranslateMementoToColumn(MementoColumn, TargetColumn);
+					},
+					[](const UScriptStruct& ColumnType, void* Destination, void* Source)
+					{
+						ColumnType.CopyScriptStruct(Destination, Source);
+					});
+
+				UE_LOG(LogEditorDataStorage, VeryVerbose, 
+					TEXT("Memento->Column: %llu -> %llu"), MementoRow, TargetRow);
+			}
 		}
 	}
-}
 
-void UTypedElementMementoSystem::DestroyMemento(TypedElementDataStorage::RowHandle MementoRow)
-{
-	// No need to lock this as no internal data is used.
+	void FMementoSystem::DestroyMemento(TypedElementDataStorage::RowHandle MementoRow)
+	{
+		// No need to lock this as no internal data is used.
 
-	checkf(DataStorage.IsRowAvailable(MementoRow) && DataStorage.HasColumns<FTypedElementMementoTag>(MementoRow),
-		TEXT("Deleting memento row that's not marked as such."));
-	DataStorage.RemoveRow(MementoRow);
-}
+		checkf(DataStorage.IsRowAvailable(MementoRow) && DataStorage.HasColumns<FTypedElementMementoTag>(MementoRow),
+			TEXT("Deleting memento row that's not marked as such."));
+		DataStorage.RemoveRow(MementoRow);
+	}
+} // namespace UE::Editor::DataStorage
diff --git a/Engine/Plugins/Experimental/EditorDataStorage/Source/TedsCore/Private/Memento/TypedElementMementoSystem.h b/Engine/Plugins/Experimental/EditorDataStorage/Source/TedsCore/Private/Memento/TypedElementMementoSystem.h
--- a/Engine/Plugins/Experimental/EditorDataStorage/Source/TedsCore/Private/Memento/TypedElementMementoSystem.h
+++ b/Engine/Plugins/Experimental/EditorDataStorage/Source/TedsCore/Private/Memento/TypedElementMementoSystem.h
@@ -6,23 +6,25 @@
 #include "Elements/Common/TypedElementHandles.h"
 
 class ITypedElementDataStorageInterface;
-class UTypedElementDatabase;
-class UTypedElementMementoTranslatorBase;
-	
-class UTypedElementMementoSystem
+class UTedsMementoTranslatorBase;
+
+namespace UE::Editor::DataStorage
 {
-public:
-	explicit UTypedElementMementoSystem(ITypedElementDataStorageInterface& InDataStorage);
-	
-	TypedElementDataStorage::RowHandle CreateMemento(TypedElementDataStorage::RowHandle SourceRow);
-	void CreateMemento(TypedElementDataStorage::RowHandle ReservedMementoRow, TypedElementDataStorage::RowHandle SourceRow);
-	void RestoreMemento(TypedElementDataStorage::RowHandle MementoRow,  TypedElementDataStorage::RowHandle TargetRow);
-	void DestroyMemento(TypedElementDataStorage::RowHandle MementoRow);
+	class FMementoSystem
+	{
+	public:
+		explicit FMementoSystem(ITypedElementDataStorageInterface& InDataStorage);
+
+		TypedElementDataStorage::RowHandle CreateMemento(TypedElementDataStorage::RowHandle SourceRow);
+		void CreateMemento(TypedElementDataStorage::RowHandle ReservedMementoRow, TypedElementDataStorage::RowHandle SourceRow);
+		void RestoreMemento(TypedElementDataStorage::RowHandle MementoRow, TypedElementDataStorage::RowHandle TargetRow);
+		void DestroyMemento(TypedElementDataStorage::RowHandle MementoRow);
 
-private:
-	void CreateMementoInternal(TypedElementDataStorage::RowHandle MementoRow, TypedElementDataStorage::RowHandle SourceRow);
+	private:
+		void CreateMementoInternal(TypedElementDataStorage::RowHandle MementoRow, TypedElementDataStorage::RowHandle SourceRow);
 
-	TArray<const UTypedElementMementoTranslatorBase*> MementoTranslators;
-	TypedElementDataStorage::TableHandle MementoRowBaseTable;
-	ITypedElementDataStorageInterface& DataStorage;
-};
+		TArray<const UTedsMementoTranslatorBase*> MementoTranslators;
+		TypedElementDataStorage::TableHandle MementoRowBaseTable;
+		ITypedElementDataStorageInterface& DataStorage;
+	};
+} // namespace UE::Editor::DataStorage
diff --git a/Engine/Plugins/Experimental/EditorDataStorage/Source/TedsCore/Private/Memento/TypedElementMementoTranslators.cpp b/Engine/Plugins/Experimental/EditorDataStorage/Source/TedsCore/Private/Memento/TypedElementMementoTranslators.cpp
--- a/Engine/Plugins/Experimental/EditorDataStorage/Source/TedsCore/Private/Memento/TypedElementMementoTranslators.cpp
+++ b/Engine/Plugins/Experimental/EditorDataStorage/Source/TedsCore/Private/Memento/TypedElementMementoTranslators.cpp
@@ -1,4 +1,4 @@
-﻿// Copyright Epic Games, Inc. All Rights Reserved.
+// Copyright Epic Games, Inc. All Rights Reserved.
 
 #include "Memento/TypedElementMementoTranslators.h"
 
@@ -6,12 +6,12 @@
 #include "StructUtils/PropertyBag.h"
 
 
-const UScriptStruct* UTypedElementDefaultMementoTranslator::GetMementoType() const
+const UScriptStruct* UTedsDefaultMementoTranslator::GetMementoType() const
 {
 	return MementoType;
 }
 
-void UTypedElementDefaultMementoTranslator::PostInitProperties()
+void UTedsDefaultMementoTranslator::PostInitProperties()
 {
 	Super::PostInitProperties();
 	
@@ -50,7 +50,7 @@ void UTypedElementDefaultMementoTranslator::PostInitProperties()
 	}
 }
 
-void UTypedElementDefaultMementoTranslator::TranslateColumnToMemento(const void* TypeErasedColumn, void* TypeErasedMemento) const
+void UTedsDefaultMementoTranslator::TranslateColumnToMemento(const void* TypeErasedColumn, void* TypeErasedMemento) const
 {
 	const std::byte* BaseAddressColumn = static_cast<const std::byte*>(TypeErasedColumn);
 	std::byte* BaseAddressMemento = static_cast<std::byte*>(TypeErasedMemento);
@@ -69,7 +69,7 @@ void UTypedElementDefaultMementoTranslator::TranslateColumnToMemento(const void*
 	}	
 }
 
-void UTypedElementDefaultMementoTranslator::TranslateMementoToColumn(const void* TypeErasedMemento,
+void UTedsDefaultMementoTranslator::TranslateMementoToColumn(const void* TypeErasedMemento,
 	void* TypeErasedColumn) const
 {
 	const std::byte* BaseAddressMemento = static_cast<const std::byte*>(TypeErasedMemento);
diff --git a/Engine/Plugins/Experimental/EditorDataStorage/Source/TedsCore/Private/Memento/TypedElementMementoizedColumns.h b/Engine/Plugins/Experimental/EditorDataStorage/Source/TedsCore/Private/Memento/TypedElementMementoizedColumns.h
--- a/Engine/Plugins/Experimental/EditorDataStorage/Source/TedsCore/Private/Memento/TypedElementMementoizedColumns.h
+++ b/Engine/Plugins/Experimental/EditorDataStorage/Source/TedsCore/Private/Memento/TypedElementMementoizedColumns.h
@@ -1,4 +1,4 @@
-﻿// Copyright Epic Games, Inc. All Rights Reserved.
+// Copyright Epic Games, Inc. All Rights Reserved.
 
 #pragma once
 
@@ -10,7 +10,7 @@
  * Enable SelectionColumn to be mementoized
  */
 UCLASS()
-class UTypedElementSelectionColumnMementoTranslator final : public UTypedElementDefaultMementoTranslator
+class UTedsSelectionColumnMementoTranslator final : public UTedsDefaultMementoTranslator
 {
 	GENERATED_BODY()
 public:
diff --git a/Engine/Plugins/Experimental/EditorDataStorage/Source/TedsCore/Private/TypedElementDatabase.cpp b/Engine/Plugins/Experimental/EditorDataStorage/Source/TedsCore/Private/TypedElementDatabase.cpp
--- a/Engine/Plugins/Experimental/EditorDataStorage/Source/TedsCore/Private/TypedElementDatabase.cpp
+++ b/Engine/Plugins/Experimental/EditorDataStorage/Source/TedsCore/Private/TypedElementDatabase.cpp
@@ -19,13 +19,13 @@
 #include "TypedElementDatabaseEnvironment.h"
 #include "UObject/UObjectIterator.h"
 
-const FName UTypedElementDatabase::TickGroupName_Default(TEXT("Default"));
-const FName UTypedElementDatabase::TickGroupName_PreUpdate(TEXT("PreUpdate"));
-const FName UTypedElementDatabase::TickGroupName_Update(TEXT("Update"));
-const FName UTypedElementDatabase::TickGroupName_PostUpdate(TEXT("PostUpdate"));
-const FName UTypedElementDatabase::TickGroupName_SyncWidget(TEXT("SyncWidgets"));
-const FName UTypedElementDatabase::TickGroupName_SyncExternalToDataStorage(TEXT("SyncExternalToDataStorage"));
-const FName UTypedElementDatabase::TickGroupName_SyncDataStorageToExternal(TEXT("SyncDataStorageToExternal"));
+const FName UEditorDataStorage::TickGroupName_Default(TEXT("Default"));
+const FName UEditorDataStorage::TickGroupName_PreUpdate(TEXT("PreUpdate"));
+const FName UEditorDataStorage::TickGroupName_Update(TEXT("Update"));
+const FName UEditorDataStorage::TickGroupName_PostUpdate(TEXT("PostUpdate"));
+const FName UEditorDataStorage::TickGroupName_SyncWidget(TEXT("SyncWidgets"));
+const FName UEditorDataStorage::TickGroupName_SyncExternalToDataStorage(TEXT("SyncExternalToDataStorage"));
+const FName UEditorDataStorage::TickGroupName_SyncDataStorageToExternal(TEXT("SyncDataStorageToExternal"));
 
 FAutoConsoleCommandWithOutputDevice PrintQueryCallbacksConsoleCommand(
 	TEXT("TEDS.PrintQueryCallbacks"),
@@ -34,7 +34,7 @@ FAutoConsoleCommandWithOutputDevice PrintQueryCallbacksConsoleCommand(
 		{
 			if (UTypedElementRegistry* Registry = UTypedElementRegistry::GetInstance())
 			{
-				if (UTypedElementDatabase* DataStorage = Cast<UTypedElementDatabase>(Registry->GetMutableDataStorage()))
+				if (UEditorDataStorage* DataStorage = Cast<UEditorDataStorage>(Registry->GetMutableDataStorage()))
 				{
 					DataStorage->DebugPrintQueryCallbacks(Output);
 				}
@@ -102,15 +102,15 @@ namespace TypedElementDatabasePrivate
 	}
 }
 
-void UTypedElementDatabase::Initialize()
+void UEditorDataStorage::Initialize()
 {
 	using namespace TypedElementDataStorage;
 
 	check(GEditor);
 	UMassEntityEditorSubsystem* Mass = GEditor->GetEditorSubsystem<UMassEntityEditorSubsystem>();
 	check(Mass);
-	OnPreMassTickHandle = Mass->GetOnPreTickDelegate().AddUObject(this, &UTypedElementDatabase::OnPreMassTick);
-	OnPostMassTickHandle = Mass->GetOnPostTickDelegate().AddUObject(this, &UTypedElementDatabase::OnPostMassTick);
+	OnPreMassTickHandle = Mass->GetOnPreTickDelegate().AddUObject(this, &UEditorDataStorage::OnPreMassTick);
+	OnPostMassTickHandle = Mass->GetOnPostTickDelegate().AddUObject(this, &UEditorDataStorage::OnPostMassTick);
 
 	ActiveEditorEntityManager = Mass->GetMutableEntityManager();
 	ActiveEditorPhaseManager = Mass->GetMutablePhaseManager();
@@ -163,7 +163,7 @@ void UTypedElementDatabase::Initialize()
 	}
 }
 
-void UTypedElementDatabase::SetFactories(TConstArrayView<UClass*> FactoryClasses)
+void UEditorDataStorage::SetFactories(TConstArrayView<UClass*> FactoryClasses)
 {
 	Factories.Reserve(FactoryClasses.Num());
 
@@ -199,7 +199,7 @@ void UTypedElementDatabase::SetFactories(TConstArrayView<UClass*> FactoryClasses
 	}
 }
 
-void UTypedElementDatabase::ResetFactories()
+void UEditorDataStorage::ResetFactories()
 {
 	for (int32 Index = Factories.Num() - 1; Index >= 0; --Index)
 	{
@@ -209,17 +209,17 @@ void UTypedElementDatabase::ResetFactories()
 	Factories.Empty();
 }
 
-UTypedElementDatabase::FactoryIterator UTypedElementDatabase::CreateFactoryIterator()
+UEditorDataStorage::FactoryIterator UEditorDataStorage::CreateFactoryIterator()
 {
-	return UTypedElementDatabase::FactoryIterator(this);
+	return UEditorDataStorage::FactoryIterator(this);
 }
 
-UTypedElementDatabase::FactoryConstIterator UTypedElementDatabase::CreateFactoryIterator() const
+UEditorDataStorage::FactoryConstIterator UEditorDataStorage::CreateFactoryIterator() const
 {
-	return UTypedElementDatabase::FactoryConstIterator(this);
+	return UEditorDataStorage::FactoryConstIterator(this);
 }
 
-const UTypedElementDataStorageFactory* UTypedElementDatabase::FindFactory(const UClass* FactoryType) const
+const UTypedElementDataStorageFactory* UEditorDataStorage::FindFactory(const UClass* FactoryType) const
 {
 	for (const FFactoryTypePair& Factory : Factories)
 	{
@@ -231,14 +231,14 @@ const UTypedElementDataStorageFactory* UTypedElementDatabase::FindFactory(const
 	return nullptr;
 }
 
-void UTypedElementDatabase::Deinitialize()
+void UEditorDataStorage::Deinitialize()
 {
 	checkf(Factories.IsEmpty(), TEXT("ResetFactories should have been called before deinitialized"));
 	
 	Reset();
 }
 
-void UTypedElementDatabase::OnPreMassTick(float DeltaTime)
+void UEditorDataStorage::OnPreMassTick(float DeltaTime)
 {
 	checkf(IsAvailable(), TEXT("Typed Element Database was ticked while it's not ready."));
 	
@@ -248,7 +248,7 @@ void UTypedElementDatabase::OnPreMassTick(float DeltaTime)
 	Environment->GetDirectDeferredCommands().ProcessCommands();
 }
 
-void UTypedElementDatabase::OnPostMassTick(float DeltaTime)
+void UEditorDataStorage::OnPostMassTick(float DeltaTime)
 {
 	checkf(IsAvailable(), TEXT("Typed Element Database was ticked while it's not ready."));
 	
@@ -256,17 +256,17 @@ void UTypedElementDatabase::OnPostMassTick(float DeltaTime)
 	OnUpdateCompletedDelegate.Broadcast();
 }
 
-TSharedPtr<FMassEntityManager> UTypedElementDatabase::GetActiveMutableEditorEntityManager()
+TSharedPtr<FMassEntityManager> UEditorDataStorage::GetActiveMutableEditorEntityManager()
 {
 	return ActiveEditorEntityManager;
 }
 
-TSharedPtr<const FMassEntityManager> UTypedElementDatabase::GetActiveEditorEntityManager() const
+TSharedPtr<const FMassEntityManager> UEditorDataStorage::GetActiveEditorEntityManager() const
 {
 	return ActiveEditorEntityManager;
 }
 
-TypedElementDataStorage::TableHandle UTypedElementDatabase::RegisterTable(TConstArrayView<const UScriptStruct*> ColumnList, const FName Name)
+TypedElementDataStorage::TableHandle UEditorDataStorage::RegisterTable(TConstArrayView<const UScriptStruct*> ColumnList, const FName Name)
 {
 	using namespace TypedElementDataStorage;
 
@@ -286,7 +286,7 @@ TypedElementDataStorage::TableHandle UTypedElementDatabase::RegisterTable(TConst
 	return InvalidTableHandle;
 }
 
-TypedElementDataStorage::TableHandle UTypedElementDatabase::RegisterTable(TypedElementDataStorage::TableHandle SourceTable,
+TypedElementDataStorage::TableHandle UEditorDataStorage::RegisterTable(TypedElementDataStorage::TableHandle SourceTable,
 	TConstArrayView<const UScriptStruct*> ColumnList, const FName Name)
 {
 	using namespace TypedElementDataStorage;
@@ -307,7 +307,7 @@ TypedElementDataStorage::TableHandle UTypedElementDatabase::RegisterTable(TypedE
 	return InvalidTableHandle;
 }
 
-TypedElementDataStorage::TableHandle UTypedElementDatabase::FindTable(const FName Name)
+TypedElementDataStorage::TableHandle UEditorDataStorage::FindTable(const FName Name)
 {
 	using namespace TypedElementDataStorage;
 
@@ -315,14 +315,14 @@ TypedElementDataStorage::TableHandle UTypedElementDatabase::FindTable(const FNam
 	return TableHandle ? *TableHandle : InvalidTableHandle;
 }
 
-TypedElementDataStorage::RowHandle UTypedElementDatabase::ReserveRow()
+TypedElementDataStorage::RowHandle UEditorDataStorage::ReserveRow()
 {
 	return ActiveEditorEntityManager 
 		? ActiveEditorEntityManager->ReserveEntity().AsNumber()
 		: TypedElementDataStorage::InvalidRowHandle;
 }
 
-void UTypedElementDatabase::BatchReserveRows(int32 Count, TFunctionRef<void(TypedElementDataStorage::RowHandle)> ReservationCallback)
+void UEditorDataStorage::BatchReserveRows(int32 Count, TFunctionRef<void(TypedElementDataStorage::RowHandle)> ReservationCallback)
 {
 	using namespace TypedElementDataStorage;
 
@@ -338,7 +338,7 @@ void UTypedElementDatabase::BatchReserveRows(int32 Count, TFunctionRef<void(Type
 	}
 }
 
-void UTypedElementDatabase::BatchReserveRows(TArrayView<TypedElementDataStorage::RowHandle> ReservedRows)
+void UEditorDataStorage::BatchReserveRows(TArrayView<TypedElementDataStorage::RowHandle> ReservedRows)
 {
 	using namespace TypedElementDataStorage;
 
@@ -354,7 +354,7 @@ void UTypedElementDatabase::BatchReserveRows(TArrayView<TypedElementDataStorage:
 	}
 }
 
-TypedElementDataStorage::RowHandle UTypedElementDatabase::AddRow(TypedElementDataStorage::TableHandle Table)
+TypedElementDataStorage::RowHandle UEditorDataStorage::AddRow(TypedElementDataStorage::TableHandle Table)
 {
 	checkf(Table < Tables.Num(), TEXT("Attempting to add a row to a non-existing table."));
 	return ActiveEditorEntityManager 
@@ -362,7 +362,7 @@ TypedElementDataStorage::RowHandle UTypedElementDatabase::AddRow(TypedElementDat
 		: TypedElementDataStorage::InvalidRowHandle;
 }
 
-TypedElementDataStorage::RowHandle UTypedElementDatabase::AddRow(TypedElementDataStorage::TableHandle Table,
+TypedElementDataStorage::RowHandle UEditorDataStorage::AddRow(TypedElementDataStorage::TableHandle Table,
 	TypedElementDataStorage::RowCreationCallbackRef OnCreated)
 {
 	using namespace TypedElementDataStorage;
@@ -385,7 +385,7 @@ TypedElementDataStorage::RowHandle UTypedElementDatabase::AddRow(TypedElementDat
 	return InvalidRowHandle;
 }
 
-bool UTypedElementDatabase::AddRow(TypedElementDataStorage::RowHandle ReservedRow, TypedElementDataStorage::TableHandle Table)
+bool UEditorDataStorage::AddRow(TypedElementDataStorage::RowHandle ReservedRow, TypedElementDataStorage::TableHandle Table)
 {
 	checkf(!IsRowAssigned(ReservedRow), TEXT("Attempting to assign a table to row that already has a table assigned."));
 	checkf(Table < Tables.Num(), TEXT("Attempting to add a row to a non-existing table."));
@@ -400,7 +400,7 @@ bool UTypedElementDatabase::AddRow(TypedElementDataStorage::RowHandle ReservedRo
 	}
 }
 
-bool UTypedElementDatabase::AddRow(TypedElementDataStorage::RowHandle ReservedRow, TypedElementDataStorage::TableHandle Table,
+bool UEditorDataStorage::AddRow(TypedElementDataStorage::RowHandle ReservedRow, TypedElementDataStorage::TableHandle Table,
 	TypedElementDataStorage::RowCreationCallbackRef OnCreated)
 {
 	OnCreated.CheckCallable();
@@ -417,7 +417,7 @@ bool UTypedElementDatabase::AddRow(TypedElementDataStorage::RowHandle ReservedRo
 	return false;
 }
 
-bool UTypedElementDatabase::BatchAddRow(
+bool UEditorDataStorage::BatchAddRow(
 	TypedElementDataStorage::TableHandle Table, int32 Count, TypedElementDataStorage::RowCreationCallbackRef OnCreated)
 {
 	OnCreated.CheckCallable();
@@ -440,7 +440,7 @@ bool UTypedElementDatabase::BatchAddRow(
 	return false;
 }
 
-bool UTypedElementDatabase::BatchAddRow(TypedElementDataStorage::TableHandle Table, 
+bool UEditorDataStorage::BatchAddRow(TypedElementDataStorage::TableHandle Table, 
 	TConstArrayView<TypedElementDataStorage::RowHandle> ReservedHandles, TypedElementDataStorage::RowCreationCallbackRef OnCreated)
 {
 	OnCreated.CheckCallable();
@@ -469,7 +469,7 @@ bool UTypedElementDatabase::BatchAddRow(TypedElementDataStorage::TableHandle Tab
 }
 
 
-void UTypedElementDatabase::RemoveRow(TypedElementDataStorage::RowHandle Row)
+void UEditorDataStorage::RemoveRow(TypedElementDataStorage::RowHandle Row)
 {
 	FMassEntityHandle Entity = FMassEntityHandle::FromNumber(Row);
 	if (ActiveEditorEntityManager && ActiveEditorEntityManager->IsEntityValid(Entity))
@@ -487,17 +487,17 @@ void UTypedElementDatabase::RemoveRow(TypedElementDataStorage::RowHandle Row)
 	}
 }
 
-bool UTypedElementDatabase::IsRowAvailable(TypedElementDataStorage::RowHandle Row) const
+bool UEditorDataStorage::IsRowAvailable(TypedElementDataStorage::RowHandle Row) const
 {
 	return ActiveEditorEntityManager ? UE::Editor::DataStorage::Legacy::FCommandBuffer::Execute_IsRowAvailable(*ActiveEditorEntityManager, Row) : false;
 }
 
-bool UTypedElementDatabase::IsRowAssigned(TypedElementDataStorage::RowHandle Row) const
+bool UEditorDataStorage::IsRowAssigned(TypedElementDataStorage::RowHandle Row) const
 {
 	return ActiveEditorEntityManager ? UE::Editor::DataStorage::Legacy::FCommandBuffer::Execute_IsRowAssigned(*ActiveEditorEntityManager, Row) : false;
 }
 
-void UTypedElementDatabase::AddColumn(TypedElementRowHandle Row, const UScriptStruct* ColumnType)
+void UEditorDataStorage::AddColumn(TypedElementRowHandle Row, const UScriptStruct* ColumnType)
 {
 	if (ColumnType && ActiveEditorEntityManager)
 	{
@@ -512,7 +512,7 @@ void UTypedElementDatabase::AddColumn(TypedElementRowHandle Row, const UScriptSt
 	}
 }
 
-void UTypedElementDatabase::AddColumnData(TypedElementRowHandle Row, const UScriptStruct* ColumnType,
+void UEditorDataStorage::AddColumnData(TypedElementRowHandle Row, const UScriptStruct* ColumnType,
 	const TypedElementDataStorage::ColumnCreationCallbackRef& Initializer,
 	TypedElementDataStorage::ColumnCopyOrMoveCallback Relocator)
 {
@@ -539,7 +539,7 @@ void UTypedElementDatabase::AddColumnData(TypedElementRowHandle Row, const UScri
 	}
 }
 
-void UTypedElementDatabase::RemoveColumn(TypedElementRowHandle Row, const UScriptStruct* ColumnType)
+void UEditorDataStorage::RemoveColumn(TypedElementRowHandle Row, const UScriptStruct* ColumnType)
 {
 	if (ColumnType && ActiveEditorEntityManager)
 	{
@@ -554,7 +554,7 @@ void UTypedElementDatabase::RemoveColumn(TypedElementRowHandle Row, const UScrip
 	}
 }
 
-const void* UTypedElementDatabase::GetColumnData(TypedElementRowHandle Row, const UScriptStruct* ColumnType) const
+const void* UEditorDataStorage::GetColumnData(TypedElementRowHandle Row, const UScriptStruct* ColumnType) const
 {
 	FMassEntityHandle Entity = FMassEntityHandle::FromNumber(Row);
 	if (ActiveEditorEntityManager &&
@@ -576,12 +576,12 @@ const void* UTypedElementDatabase::GetColumnData(TypedElementRowHandle Row, cons
 	return nullptr;
 }
 
-void* UTypedElementDatabase::GetColumnData(TypedElementRowHandle Row, const UScriptStruct* ColumnType)
+void* UEditorDataStorage::GetColumnData(TypedElementRowHandle Row, const UScriptStruct* ColumnType)
 {
-	return const_cast<void*>(static_cast<const UTypedElementDatabase*>(this)->GetColumnData(Row, ColumnType));
+	return const_cast<void*>(static_cast<const UEditorDataStorage*>(this)->GetColumnData(Row, ColumnType));
 }
 
-void UTypedElementDatabase::AddColumns(TypedElementRowHandle Row, TConstArrayView<const UScriptStruct*> Columns)
+void UEditorDataStorage::AddColumns(TypedElementRowHandle Row, TConstArrayView<const UScriptStruct*> Columns)
 {
 	if (ActiveEditorEntityManager)
 	{
@@ -604,7 +604,7 @@ void UTypedElementDatabase::AddColumns(TypedElementRowHandle Row, TConstArrayVie
 	}
 }
 
-void UTypedElementDatabase::AddColumn(TypedElementRowHandle Row, const UE::Editor::DataStorage::FDynamicTag& Tag, const FName& InValue)
+void UEditorDataStorage::AddColumn(TypedElementRowHandle Row, const UE::Editor::DataStorage::FDynamicTag& Tag, const FName& InValue)
 {
 	if (ActiveEditorEntityManager)
 	{
@@ -618,7 +618,7 @@ void UTypedElementDatabase::AddColumn(TypedElementRowHandle Row, const UE::Edito
 	}
 }
 
-void UTypedElementDatabase::RemoveColumn(TypedElementRowHandle Row, const UE::Editor::DataStorage::FDynamicTag& Tag)
+void UEditorDataStorage::RemoveColumn(TypedElementRowHandle Row, const UE::Editor::DataStorage::FDynamicTag& Tag)
 {
 	if (ActiveEditorEntityManager)
 	{
@@ -631,7 +631,7 @@ void UTypedElementDatabase::RemoveColumn(TypedElementRowHandle Row, const UE::Ed
 	}
 }
 
-void UTypedElementDatabase::RemoveColumns(TypedElementRowHandle Row, TConstArrayView<const UScriptStruct*> Columns)
+void UEditorDataStorage::RemoveColumns(TypedElementRowHandle Row, TConstArrayView<const UScriptStruct*> Columns)
 {
 	FMassEntityHandle Entity = FMassEntityHandle::FromNumber(Row);
 	if (ActiveEditorEntityManager)
@@ -654,7 +654,7 @@ void UTypedElementDatabase::RemoveColumns(TypedElementRowHandle Row, TConstArray
 	}
 }
 
-void UTypedElementDatabase::AddRemoveColumns(TypedElementRowHandle Row,
+void UEditorDataStorage::AddRemoveColumns(TypedElementRowHandle Row,
 	TConstArrayView<const UScriptStruct*> ColumnsToAdd, TConstArrayView<const UScriptStruct*> ColumnsToRemove)
 {
 	if (ActiveEditorEntityManager)
@@ -695,7 +695,7 @@ void UTypedElementDatabase::AddRemoveColumns(TypedElementRowHandle Row,
 	}
 }
 
-void UTypedElementDatabase::BatchAddRemoveColumns(TConstArrayView<TypedElementRowHandle> Rows, 
+void UEditorDataStorage::BatchAddRemoveColumns(TConstArrayView<TypedElementRowHandle> Rows, 
 	TConstArrayView<const UScriptStruct*> ColumnsToAdd, TConstArrayView<const UScriptStruct*> ColumnsToRemove)
 {	
 	if (ActiveEditorEntityManager)
@@ -764,7 +764,7 @@ void UTypedElementDatabase::BatchAddRemoveColumns(TConstArrayView<TypedElementRo
 	}
 }
 
-bool UTypedElementDatabase::HasColumns(TypedElementRowHandle Row, TConstArrayView<const UScriptStruct*> ColumnTypes) const
+bool UEditorDataStorage::HasColumns(TypedElementRowHandle Row, TConstArrayView<const UScriptStruct*> ColumnTypes) const
 {
 	if (ActiveEditorEntityManager)
 	{
@@ -808,7 +808,7 @@ bool UTypedElementDatabase::HasColumns(TypedElementRowHandle Row, TConstArrayVie
 	return false;
 }
 
-bool UTypedElementDatabase::HasColumns(TypedElementRowHandle Row, TConstArrayView<TWeakObjectPtr<const UScriptStruct>> ColumnTypes) const
+bool UEditorDataStorage::HasColumns(TypedElementRowHandle Row, TConstArrayView<TWeakObjectPtr<const UScriptStruct>> ColumnTypes) const
 {
 	if (ActiveEditorEntityManager)
 	{
@@ -854,7 +854,7 @@ bool UTypedElementDatabase::HasColumns(TypedElementRowHandle Row, TConstArrayVie
 	return false;
 }
 
-void UTypedElementDatabase::ListColumns(TypedElementDataStorage::RowHandle Row, TypedElementDataStorage::ColumnListCallbackRef Callback) const
+void UEditorDataStorage::ListColumns(TypedElementDataStorage::RowHandle Row, TypedElementDataStorage::ColumnListCallbackRef Callback) const
 {
 	if (ActiveEditorEntityManager)
 	{
@@ -879,7 +879,7 @@ void UTypedElementDatabase::ListColumns(TypedElementDataStorage::RowHandle Row,
 	}
 }
 
-void UTypedElementDatabase::ListColumns(TypedElementDataStorage::RowHandle Row, TypedElementDataStorage::ColumnListWithDataCallbackRef Callback)
+void UEditorDataStorage::ListColumns(TypedElementDataStorage::RowHandle Row, TypedElementDataStorage::ColumnListWithDataCallbackRef Callback)
 {
 	if (ActiveEditorEntityManager)
 	{
@@ -912,7 +912,7 @@ void UTypedElementDatabase::ListColumns(TypedElementDataStorage::RowHandle Row,
 	}
 }
 
-bool UTypedElementDatabase::MatchesColumns(TypedElementDataStorage::RowHandle Row, const TypedElementDataStorage::FQueryConditions& Conditions) const
+bool UEditorDataStorage::MatchesColumns(TypedElementDataStorage::RowHandle Row, const TypedElementDataStorage::FQueryConditions& Conditions) const
 {
 	if (ActiveEditorEntityManager)
 	{
@@ -956,25 +956,25 @@ bool UTypedElementDatabase::MatchesColumns(TypedElementDataStorage::RowHandle Ro
 	return false;
 }
 
-void UTypedElementDatabase::RegisterTickGroup(
+void UEditorDataStorage::RegisterTickGroup(
 	FName GroupName, EQueryTickPhase Phase, FName BeforeGroup, FName AfterGroup, TypedElementDataStorage::EExecutionMode ExecutionMode)
 {
 	Environment->GetQueryStore().RegisterTickGroup(GroupName, Phase, BeforeGroup, AfterGroup, ExecutionMode);
 }
 
-void UTypedElementDatabase::UnregisterTickGroup(FName GroupName, EQueryTickPhase Phase)
+void UEditorDataStorage::UnregisterTickGroup(FName GroupName, EQueryTickPhase Phase)
 {
 	Environment->GetQueryStore().UnregisterTickGroup(GroupName, Phase);
 }
 
-TypedElementQueryHandle UTypedElementDatabase::RegisterQuery(FQueryDescription&& Query)
+TypedElementQueryHandle UEditorDataStorage::RegisterQuery(FQueryDescription&& Query)
 {
 	return (ActiveEditorEntityManager && ActiveEditorPhaseManager)
 		? Environment->GetQueryStore().RegisterQuery(MoveTemp(Query), *Environment, *ActiveEditorEntityManager, *ActiveEditorPhaseManager).Packed()
 		: TypedElementInvalidQueryHandle;
 }
 
-void UTypedElementDatabase::UnregisterQuery(TypedElementQueryHandle Query)
+void UEditorDataStorage::UnregisterQuery(TypedElementQueryHandle Query)
 {
 	if (ActiveEditorEntityManager && ActiveEditorPhaseManager)
 	{
@@ -983,13 +983,13 @@ void UTypedElementDatabase::UnregisterQuery(TypedElementQueryHandle Query)
 	}
 }
 
-const ITypedElementDataStorageInterface::FQueryDescription& UTypedElementDatabase::GetQueryDescription(TypedElementQueryHandle Query) const
+const ITypedElementDataStorageInterface::FQueryDescription& UEditorDataStorage::GetQueryDescription(TypedElementQueryHandle Query) const
 {
 	const UE::Editor::DataStorage::FExtendedQueryStore::Handle StorageHandle(Query);
 	return Environment->GetQueryStore().GetQueryDescription(StorageHandle);
 }
 
-FName UTypedElementDatabase::GetQueryTickGroupName(EQueryTickGroups Group) const
+FName UEditorDataStorage::GetQueryTickGroupName(EQueryTickGroups Group) const
 {
 	switch (Group)
 	{
@@ -1013,7 +1013,7 @@ FName UTypedElementDatabase::GetQueryTickGroupName(EQueryTickGroups Group) const
 	}
 }
 
-ITypedElementDataStorageInterface::FQueryResult UTypedElementDatabase::RunQuery(TypedElementQueryHandle Query)
+ITypedElementDataStorageInterface::FQueryResult UEditorDataStorage::RunQuery(TypedElementQueryHandle Query)
 {
 	TRACE_CPUPROFILER_EVENT_SCOPE(TEDS.RunQuery);
 
@@ -1028,7 +1028,7 @@ ITypedElementDataStorageInterface::FQueryResult UTypedElementDatabase::RunQuery(
 	}
 }
 
-ITypedElementDataStorageInterface::FQueryResult UTypedElementDatabase::RunQuery(
+ITypedElementDataStorageInterface::FQueryResult UEditorDataStorage::RunQuery(
 	TypedElementQueryHandle Query, ITypedElementDataStorageInterface::DirectQueryCallbackRef Callback)
 {
 	TRACE_CPUPROFILER_EVENT_SCOPE(TEDS.RunQuery);
@@ -1045,7 +1045,7 @@ ITypedElementDataStorageInterface::FQueryResult UTypedElementDatabase::RunQuery(
 	}
 }
 
-ITypedElementDataStorageInterface::FQueryResult UTypedElementDatabase::RunQuery(
+ITypedElementDataStorageInterface::FQueryResult UEditorDataStorage::RunQuery(
 	TypedElementQueryHandle Query, TypedElementDataStorage::EDirectQueryExecutionFlags Flags,
 	ITypedElementDataStorageInterface::DirectQueryCallbackRef Callback)
 {
@@ -1062,7 +1062,7 @@ ITypedElementDataStorageInterface::FQueryResult UTypedElementDatabase::RunQuery(
 	}
 }
 
-void UTypedElementDatabase::ActivateQueries(FName ActivationName)
+void UEditorDataStorage::ActivateQueries(FName ActivationName)
 {
 	if (ActiveEditorEntityManager)
 	{
@@ -1070,49 +1070,49 @@ void UTypedElementDatabase::ActivateQueries(FName ActivationName)
 	}
 }
 
-TypedElementDataStorage::RowHandle UTypedElementDatabase::FindIndexedRow(TypedElementDataStorage::IndexHash Index) const
+TypedElementDataStorage::RowHandle UEditorDataStorage::FindIndexedRow(TypedElementDataStorage::IndexHash Index) const
 {
 	return Environment->GetIndexTable().FindIndexedRow(UE::Editor::DataStorage::EGlobalLockScope::Public, Index);
 }
 
-void UTypedElementDatabase::IndexRow(TypedElementDataStorage::IndexHash Index, TypedElementDataStorage::RowHandle Row)
+void UEditorDataStorage::IndexRow(TypedElementDataStorage::IndexHash Index, TypedElementDataStorage::RowHandle Row)
 {
 	Environment->GetIndexTable().IndexRow(UE::Editor::DataStorage::EGlobalLockScope::Public, Index, Row);
 }
 
-void UTypedElementDatabase::BatchIndexRows(
+void UEditorDataStorage::BatchIndexRows(
 	TConstArrayView<TPair<TypedElementDataStorage::IndexHash, TypedElementDataStorage::RowHandle>> IndexRowPairs)
 {
 	Environment->GetIndexTable().BatchIndexRows(UE::Editor::DataStorage::EGlobalLockScope::Public, IndexRowPairs);
 }
 
-void UTypedElementDatabase::ReindexRow(TypedElementDataStorage::IndexHash OriginalIndex, TypedElementDataStorage::IndexHash NewIndex, 
+void UEditorDataStorage::ReindexRow(TypedElementDataStorage::IndexHash OriginalIndex, TypedElementDataStorage::IndexHash NewIndex, 
 	TypedElementDataStorage::RowHandle RowHandle)
 {
 	Environment->GetIndexTable().ReindexRow(UE::Editor::DataStorage::EGlobalLockScope::Public, OriginalIndex, NewIndex, RowHandle);
 }
 
-void UTypedElementDatabase::RemoveIndex(TypedElementDataStorage::IndexHash Index)
+void UEditorDataStorage::RemoveIndex(TypedElementDataStorage::IndexHash Index)
 {
 	Environment->GetIndexTable().RemoveIndex(UE::Editor::DataStorage::EGlobalLockScope::Public, Index);
 }
 
-FTypedElementOnDataStorageUpdate& UTypedElementDatabase::OnUpdate()
+FTypedElementOnDataStorageUpdate& UEditorDataStorage::OnUpdate()
 {
 	return OnUpdateDelegate;
 }
 
-FTypedElementOnDataStorageUpdate& UTypedElementDatabase::OnUpdateCompleted()
+FTypedElementOnDataStorageUpdate& UEditorDataStorage::OnUpdateCompleted()
 {
 	return OnUpdateCompletedDelegate;
 }
 
-bool UTypedElementDatabase::IsAvailable() const
+bool UEditorDataStorage::IsAvailable() const
 {
 	return bool(ActiveEditorEntityManager);
 }
 
-void* UTypedElementDatabase::GetExternalSystemAddress(UClass* Target)
+void* UEditorDataStorage::GetExternalSystemAddress(UClass* Target)
 {
 	if (Target && Target->IsChildOf<USubsystem>())
 	{
@@ -1121,16 +1121,16 @@ void* UTypedElementDatabase::GetExternalSystemAddress(UClass* Target)
 	return nullptr;
 }
 
-bool UTypedElementDatabase::SupportsExtension(FName Extension) const
+bool UEditorDataStorage::SupportsExtension(FName Extension) const
 {
 	return false;
 }
 
-void UTypedElementDatabase::ListExtensions(TFunctionRef<void(FName)> Callback) const
+void UEditorDataStorage::ListExtensions(TFunctionRef<void(FName)> Callback) const
 {
 }
 
-void UTypedElementDatabase::PreparePhase(EQueryTickPhase Phase, float DeltaTime)
+void UEditorDataStorage::PreparePhase(EQueryTickPhase Phase, float DeltaTime)
 {
 	using namespace UE::Editor::DataStorage;
 
@@ -1150,7 +1150,7 @@ void UTypedElementDatabase::PreparePhase(EQueryTickPhase Phase, float DeltaTime)
 	}
 }
 
-void UTypedElementDatabase::FinalizePhase(EQueryTickPhase Phase, float DeltaTime)
+void UEditorDataStorage::FinalizePhase(EQueryTickPhase Phase, float DeltaTime)
 {
 	using namespace UE::Editor::DataStorage;
 
@@ -1170,7 +1170,7 @@ void UTypedElementDatabase::FinalizePhase(EQueryTickPhase Phase, float DeltaTime
 	}
 }
 
-void UTypedElementDatabase::Reset()
+void UEditorDataStorage::Reset()
 {
 	if (UMassEntityEditorSubsystem* Mass = GEditor->GetEditorSubsystem<UMassEntityEditorSubsystem>())
 	{
@@ -1191,7 +1191,7 @@ void UTypedElementDatabase::Reset()
 	ActiveEditorEntityManager.Reset();
 }
 
-int32 UTypedElementDatabase::GetTableChunkSize(FName TableName) const
+int32 UEditorDataStorage::GetTableChunkSize(FName TableName) const
 {
 	const UEditorDataStorageSettings* Settings = GetDefault<UEditorDataStorageSettings>();
 	if (const EChunkMemorySize* TableSpecificSize = Settings->TableSpecificChunkMemorySize.Find(TableName))
@@ -1204,17 +1204,17 @@ int32 UTypedElementDatabase::GetTableChunkSize(FName TableName) const
 	}
 }
 
-TSharedPtr<UE::Editor::DataStorage::FEnvironment> UTypedElementDatabase::GetEnvironment()
+TSharedPtr<UE::Editor::DataStorage::FEnvironment> UEditorDataStorage::GetEnvironment()
 {
 	return Environment;
 }
 
-TSharedPtr<const UE::Editor::DataStorage::FEnvironment> UTypedElementDatabase::GetEnvironment() const
+TSharedPtr<const UE::Editor::DataStorage::FEnvironment> UEditorDataStorage::GetEnvironment() const
 {
 	return Environment;
 }
 
-FMassArchetypeHandle UTypedElementDatabase::LookupArchetype(TypedElementDataStorage::TableHandle InTableHandle) const
+FMassArchetypeHandle UEditorDataStorage::LookupArchetype(TypedElementDataStorage::TableHandle InTableHandle) const
 {
 	const uint32 TableIndex = InTableHandle;
 	if (Tables.IsValidIndex(TableIndex))
@@ -1224,14 +1224,14 @@ FMassArchetypeHandle UTypedElementDatabase::LookupArchetype(TypedElementDataStor
 	return FMassArchetypeHandle();
 }
 
-void UTypedElementDatabase::DebugPrintQueryCallbacks(FOutputDevice& Output)
+void UEditorDataStorage::DebugPrintQueryCallbacks(FOutputDevice& Output)
 {
 	Environment->GetQueryStore().DebugPrintQueryCallbacks(Output);
 }
 
-void UTypedElementDatabase::AddReferencedObjects(UObject* InThis, FReferenceCollector& Collector)
+void UEditorDataStorage::AddReferencedObjects(UObject* InThis, FReferenceCollector& Collector)
 {
-	UTypedElementDatabase* Database = static_cast<UTypedElementDatabase*>(InThis);
+	UEditorDataStorage* Database = static_cast<UEditorDataStorage*>(InThis);
 
 	for (auto& FactoryPair : Database->Factories)
 	{
diff --git a/Engine/Plugins/Experimental/EditorDataStorage/Source/TedsCore/Private/TypedElementDatabase.h b/Engine/Plugins/Experimental/EditorDataStorage/Source/TedsCore/Private/TypedElementDatabase.h
--- a/Engine/Plugins/Experimental/EditorDataStorage/Source/TedsCore/Private/TypedElementDatabase.h
+++ b/Engine/Plugins/Experimental/EditorDataStorage/Source/TedsCore/Private/TypedElementDatabase.h
@@ -25,7 +25,7 @@ namespace UE::Editor::DataStorage
 } // namespace UE::Editor::DataStorage
 
 UCLASS()
-class TEDSCORE_API UTypedElementDatabase
+class TEDSCORE_API UEditorDataStorage
 	: public UObject
 	, public ITypedElementDataStorageInterface
 {
@@ -52,11 +52,11 @@ public:
 		int32 Index = 0;
 	};
 
-	using FactoryIterator = TFactoryIterator<UTypedElementDataStorageFactory, UTypedElementDatabase>;
-	using FactoryConstIterator = TFactoryIterator<const UTypedElementDataStorageFactory, const UTypedElementDatabase>;
+	using FactoryIterator = TFactoryIterator<UTypedElementDataStorageFactory, UEditorDataStorage>;
+	using FactoryConstIterator = TFactoryIterator<const UTypedElementDataStorageFactory, const UEditorDataStorage>;
 
 public:
-	~UTypedElementDatabase() override = default;
+	~UEditorDataStorage() override = default;
 	
 	void Initialize();
 	
@@ -208,17 +208,17 @@ private:
 };
 
 template <typename FactoryType, typename DatabaseType>
-UTypedElementDatabase::TFactoryIterator<FactoryType, DatabaseType>::TFactoryIterator(DatabasePtr InDatabase): Database(InDatabase)
+UEditorDataStorage::TFactoryIterator<FactoryType, DatabaseType>::TFactoryIterator(DatabasePtr InDatabase): Database(InDatabase)
 {}
 
 template <typename FactoryType, typename DatabaseType>
-typename UTypedElementDatabase::TFactoryIterator<FactoryType, DatabaseType>::FactoryPtr UTypedElementDatabase::TFactoryIterator<FactoryType, DatabaseType>::operator*() const
+typename UEditorDataStorage::TFactoryIterator<FactoryType, DatabaseType>::FactoryPtr UEditorDataStorage::TFactoryIterator<FactoryType, DatabaseType>::operator*() const
 {
 	return Database->Factories[Index].Instance;
 }
 
 template <typename FactoryType, typename DatabaseType>
-typename UTypedElementDatabase::TFactoryIterator<FactoryType, DatabaseType>::ThisType& UTypedElementDatabase::TFactoryIterator<FactoryType, DatabaseType>::operator++()
+typename UEditorDataStorage::TFactoryIterator<FactoryType, DatabaseType>::ThisType& UEditorDataStorage::TFactoryIterator<FactoryType, DatabaseType>::operator++()
 {
 	if (Database != nullptr && Index < Database->Factories.Num())
 	{
@@ -228,13 +228,13 @@ typename UTypedElementDatabase::TFactoryIterator<FactoryType, DatabaseType>::Thi
 }
 
 template <typename FactoryType, typename DatabaseType>
-UTypedElementDatabase::TFactoryIterator<FactoryType, DatabaseType>::operator bool() const
+UEditorDataStorage::TFactoryIterator<FactoryType, DatabaseType>::operator bool() const
 {
 	return Database != nullptr && Index < Database->Factories.Num();
 }
 
 template <typename FactoryTypeT>
-const FactoryTypeT* UTypedElementDatabase::FindFactory() const
+const FactoryTypeT* UEditorDataStorage::FindFactory() const
 {
 	return static_cast<const FactoryTypeT*>(FindFactory(FactoryTypeT::StaticClass()));
 }
diff --git a/Engine/Plugins/Experimental/EditorDataStorage/Source/TedsCore/Private/TypedElementDatabaseCompatibility.cpp b/Engine/Plugins/Experimental/EditorDataStorage/Source/TedsCore/Private/TypedElementDatabaseCompatibility.cpp
--- a/Engine/Plugins/Experimental/EditorDataStorage/Source/TedsCore/Private/TypedElementDatabaseCompatibility.cpp
+++ b/Engine/Plugins/Experimental/EditorDataStorage/Source/TedsCore/Private/TypedElementDatabaseCompatibility.cpp
@@ -24,7 +24,7 @@
 
 DEFINE_LOG_CATEGORY_STATIC(LogTedsCompat, Log, All);
 
-namespace TypedElementDataStorage
+namespace UE::Editor::DataStorage::Private
 {
 	bool bIntegrateWithGC = true;
 	FAutoConsoleVariableRef CVarIntegrateWithGC(
@@ -59,12 +59,12 @@ namespace TypedElementDataStorage
 			"0 - disable\n"
 			"1 - summarize number of nops\n"
 			"2 - include nops"));
-}
 
-static const FName IntegrateWithGCName(TEXT("IntegrateWithGC"));
-static const FName CompatibilityUsesCommandBufferExtensionName(TEXT("CompatiblityUsesCommandBuffer"));
+	static const FName IntegrateWithGCName(TEXT("IntegrateWithGC"));
+	static const FName CompatibilityUsesCommandBufferExtensionName(TEXT("CompatiblityUsesCommandBuffer"));
+} // namespace UE::Editor::DataStorage::Private
 
-void UTypedElementDatabaseCompatibility::Initialize(UTypedElementDatabase* InStorage)
+void UEditorDataStorageCompatibility::Initialize(UEditorDataStorage* InStorage)
 {
 	using namespace TypedElementDataStorage;
 	using namespace UE::Editor::DataStorage;
@@ -79,22 +79,21 @@ void UTypedElementDatabaseCompatibility::Initialize(UTypedElementDatabase* InSto
 	
 	Prepare();
 
-	InStorage->OnUpdate().AddUObject(this, &UTypedElementDatabaseCompatibility::Tick);
+	InStorage->OnUpdate().AddUObject(this, &UEditorDataStorageCompatibility::Tick);
 
-	PreEditChangePropertyDelegateHandle = FCoreUObjectDelegates::OnPreObjectPropertyChanged.AddUObject(this, &UTypedElementDatabaseCompatibility::OnPrePropertyChanged);
-	PostEditChangePropertyDelegateHandle = FCoreUObjectDelegates::OnObjectPropertyChanged.AddUObject(this, &UTypedElementDatabaseCompatibility::OnPostEditChangeProperty);
-	ObjectModifiedDelegateHandle = FCoreUObjectDelegates::OnObjectModified.AddUObject(this, &UTypedElementDatabaseCompatibility::OnObjectModified);
-	ObjectReinstancedDelegateHandle = FCoreUObjectDelegates::OnObjectsReinstanced.AddUObject(this, &UTypedElementDatabaseCompatibility::OnObjectReinstanced);
+	PreEditChangePropertyDelegateHandle = FCoreUObjectDelegates::OnPreObjectPropertyChanged.AddUObject(this, &UEditorDataStorageCompatibility::OnPrePropertyChanged);
+	PostEditChangePropertyDelegateHandle = FCoreUObjectDelegates::OnObjectPropertyChanged.AddUObject(this, &UEditorDataStorageCompatibility::OnPostEditChangeProperty);
+	ObjectModifiedDelegateHandle = FCoreUObjectDelegates::OnObjectModified.AddUObject(this, &UEditorDataStorageCompatibility::OnObjectModified);
+	ObjectReinstancedDelegateHandle = FCoreUObjectDelegates::OnObjectsReinstanced.AddUObject(this, &UEditorDataStorageCompatibility::OnObjectReinstanced);
 	
-	PostGcUnreachableAnalysisHandle = FCoreUObjectDelegates::PostReachabilityAnalysis.AddUObject(this, &UTypedElementDatabaseCompatibility::OnPostGcUnreachableAnalysis);
+	PostGcUnreachableAnalysisHandle = FCoreUObjectDelegates::PostReachabilityAnalysis.AddUObject(this, &UEditorDataStorageCompatibility::OnPostGcUnreachableAnalysis);
 	// Used to get all the worlds and register the actor create/destroy handles on them.
-	PostWorldInitializationDelegateHandle = FWorldDelegates::OnPostWorldInitialization.AddUObject(this, &UTypedElementDatabaseCompatibility::OnPostWorldInitialization);
-	PreWorldFinishDestroyDelegateHandle = FWorldDelegates::OnPreWorldFinishDestroy.AddUObject(this, &UTypedElementDatabaseCompatibility::OnPreWorldFinishDestroy);
+	PostWorldInitializationDelegateHandle = FWorldDelegates::OnPostWorldInitialization.AddUObject(this, &UEditorDataStorageCompatibility::OnPostWorldInitialization);
+	PreWorldFinishDestroyDelegateHandle = FWorldDelegates::OnPreWorldFinishDestroy.AddUObject(this, &UEditorDataStorageCompatibility::OnPreWorldFinishDestroy);
 }
 
-void UTypedElementDatabaseCompatibility::Deinitialize()
+void UEditorDataStorageCompatibility::Deinitialize()
 {
-	using namespace TypedElementDataStorage;
 	using namespace UE::Editor::DataStorage;
 
 	FScopedExclusiveLock Lock(EGlobalLockScope::Public);
@@ -116,7 +115,7 @@ void UTypedElementDatabaseCompatibility::Deinitialize()
 	Reset();
 }
 
-void UTypedElementDatabaseCompatibility::RegisterRegistrationFilter(ObjectRegistrationFilter Filter)
+void UEditorDataStorageCompatibility::RegisterRegistrationFilter(ObjectRegistrationFilter Filter)
 {
 	using namespace UE::Editor::DataStorage;
 
@@ -124,7 +123,7 @@ void UTypedElementDatabaseCompatibility::RegisterRegistrationFilter(ObjectRegist
 	ObjectRegistrationFilters.Add(MoveTemp(Filter));
 }
 
-void UTypedElementDatabaseCompatibility::RegisterDealiaserCallback(ObjectToRowDealiaser Dealiaser)
+void UEditorDataStorageCompatibility::RegisterDealiaserCallback(ObjectToRowDealiaser Dealiaser)
 {
 	using namespace UE::Editor::DataStorage;
 
@@ -132,12 +131,12 @@ void UTypedElementDatabaseCompatibility::RegisterDealiaserCallback(ObjectToRowDe
 	ObjectToRowDialiasers.Add(MoveTemp(Dealiaser));
 }
 
-void UTypedElementDatabaseCompatibility::RegisterTypeTableAssociation(
+void UEditorDataStorageCompatibility::RegisterTypeTableAssociation(
 	TObjectPtr<UStruct> TypeInfo, TypedElementDataStorage::TableHandle Table)
 {
 	using namespace UE::Editor::DataStorage;
 	
-	if (TypedElementDataStorage::bUseCommandBuffer)
+	if (Private::bUseCommandBuffer)
 	{
 		QueuedCommands.AddCommand(UE::Editor::DataStorage::FRegisterTypeTableAssociation{ .TypeInfo = TypeInfo, .Table = Table });
 	}
@@ -148,12 +147,12 @@ void UTypedElementDatabaseCompatibility::RegisterTypeTableAssociation(
 	}
 }
 
-FDelegateHandle UTypedElementDatabaseCompatibility::RegisterObjectAddedCallback(UE::Editor::DataStorage::ObjectAddedCallback&& OnObjectAdded)
+FDelegateHandle UEditorDataStorageCompatibility::RegisterObjectAddedCallback(UE::Editor::DataStorage::ObjectAddedCallback&& OnObjectAdded)
 {
 	using namespace UE::Editor::DataStorage;
 
 	FDelegateHandle Handle(FDelegateHandle::GenerateNewHandle);
-	if (TypedElementDataStorage::bUseCommandBuffer)
+	if (Private::bUseCommandBuffer)
 	{
 		QueuedCommands.AddCommand(UE::Editor::DataStorage::FRegisterObjectAddedCallback{ .Callback = MoveTemp(OnObjectAdded), .Handle = Handle });
 	}
@@ -165,11 +164,11 @@ FDelegateHandle UTypedElementDatabaseCompatibility::RegisterObjectAddedCallback(
 	return Handle;
 }
 
-void UTypedElementDatabaseCompatibility::UnregisterObjectAddedCallback(FDelegateHandle Handle)
+void UEditorDataStorageCompatibility::UnregisterObjectAddedCallback(FDelegateHandle Handle)
 {
 	using namespace UE::Editor::DataStorage;
 
-	if (TypedElementDataStorage::bUseCommandBuffer)
+	if (Private::bUseCommandBuffer)
 	{
 		QueuedCommands.AddCommand(UE::Editor::DataStorage::FUnregisterObjectAddedCallback{ .Handle = Handle });
 	}
@@ -184,7 +183,7 @@ void UTypedElementDatabaseCompatibility::UnregisterObjectAddedCallback(FDelegate
 	}
 }
 
-FDelegateHandle UTypedElementDatabaseCompatibility::RegisterObjectRemovedCallback(UE::Editor::DataStorage::ObjectRemovedCallback&& OnObjectAdded)
+FDelegateHandle UEditorDataStorageCompatibility::RegisterObjectRemovedCallback(UE::Editor::DataStorage::ObjectRemovedCallback&& OnObjectAdded)
 {
 	using namespace UE::Editor::DataStorage;
 
@@ -196,7 +195,7 @@ FDelegateHandle UTypedElementDatabaseCompatibility::RegisterObjectRemovedCallbac
 	return Handle;
 }
 
-void UTypedElementDatabaseCompatibility::UnregisterObjectRemovedCallback(FDelegateHandle Handle)
+void UEditorDataStorageCompatibility::UnregisterObjectRemovedCallback(FDelegateHandle Handle)
 {
 	using namespace UE::Editor::DataStorage;
 
@@ -209,7 +208,7 @@ void UTypedElementDatabaseCompatibility::UnregisterObjectRemovedCallback(FDelega
 	});
 }
 
-TypedElementRowHandle UTypedElementDatabaseCompatibility::AddCompatibleObjectExplicit(UObject* Object)
+TypedElementRowHandle UEditorDataStorageCompatibility::AddCompatibleObjectExplicit(UObject* Object)
 {
 	// Because AddCompatibleObjectExplicitTransactionable needs a finer grained control over the lock, there's no higher up lock here.
 
@@ -219,7 +218,7 @@ TypedElementRowHandle UTypedElementDatabaseCompatibility::AddCompatibleObjectExp
 	return bCanAddObject ? AddCompatibleObjectExplicitTransactionable<true>(Object) : TypedElementDataStorage::InvalidRowHandle;
 }
 
-TypedElementRowHandle UTypedElementDatabaseCompatibility::AddCompatibleObjectExplicit(void* Object, TWeakObjectPtr<UScriptStruct> TypeInfo)
+TypedElementRowHandle UEditorDataStorageCompatibility::AddCompatibleObjectExplicit(void* Object, TWeakObjectPtr<UScriptStruct> TypeInfo)
 {
 	using namespace TypedElementDataStorage;
 	using namespace UE::Editor::DataStorage;
@@ -232,7 +231,7 @@ TypedElementRowHandle UTypedElementDatabaseCompatibility::AddCompatibleObjectExp
 	{
 		Result = Storage->ReserveRow();
 		Storage->IndexRow(GenerateIndexHash(Object), Result);
-		if (bUseCommandBuffer)
+		if (Private::bUseCommandBuffer)
 		{
 			QueuedCommands.AddCommand(FAddCompatibleExternalObject{ .Object = Object, .TypeInfo = TypeInfo, .Row = Result });
 		}
@@ -244,19 +243,19 @@ TypedElementRowHandle UTypedElementDatabaseCompatibility::AddCompatibleObjectExp
 	return Result;
 }
 
-void UTypedElementDatabaseCompatibility::RemoveCompatibleObjectExplicit(UObject* Object)
+void UEditorDataStorageCompatibility::RemoveCompatibleObjectExplicit(UObject* Object)
 {
 	RemoveCompatibleObjectExplicitTransactionable<true>(Object);
 }
 
-void UTypedElementDatabaseCompatibility::RemoveCompatibleObjectExplicit(void* Object)
+void UEditorDataStorageCompatibility::RemoveCompatibleObjectExplicit(void* Object)
 {
 	using namespace TypedElementDataStorage;
 	using namespace UE::Editor::DataStorage;
 
 	checkf(Storage, TEXT("Removing compatible objects is not supported before Typed Element's Database compatibility manager has been initialized."));
 
-	if (TypedElementDataStorage::bUseCommandBuffer && TypedElementDataStorage::bUseDeferredRemovesInCompat)
+	if (Private::bUseCommandBuffer && Private::bUseDeferredRemovesInCompat)
 	{
 		QueuedCommands.AddCommand(FRemoveCompatibleExternalObject{ .Object = Object });
 	}
@@ -278,7 +277,7 @@ void UTypedElementDatabaseCompatibility::RemoveCompatibleObjectExplicit(void* Ob
 	}
 }
 
-TypedElementRowHandle UTypedElementDatabaseCompatibility::FindRowWithCompatibleObjectExplicit(const UObject* Object) const
+TypedElementRowHandle UEditorDataStorageCompatibility::FindRowWithCompatibleObjectExplicit(const UObject* Object) const
 {
 	using namespace TypedElementDataStorage;
 	using namespace UE::Editor::DataStorage;
@@ -293,7 +292,7 @@ TypedElementRowHandle UTypedElementDatabaseCompatibility::FindRowWithCompatibleO
 	return InvalidRowHandle;
 }
 
-TypedElementRowHandle UTypedElementDatabaseCompatibility::FindRowWithCompatibleObjectExplicit(const void* Object) const
+TypedElementRowHandle UEditorDataStorageCompatibility::FindRowWithCompatibleObjectExplicit(const void* Object) const
 {
 	// Thread safety is only needed by FindIndexedRow which internally takes care of it.
 
@@ -302,19 +301,19 @@ TypedElementRowHandle UTypedElementDatabaseCompatibility::FindRowWithCompatibleO
 	return (Object && Storage && Storage->IsAvailable()) ? Storage->FindIndexedRow(GenerateIndexHash(Object)) : InvalidRowHandle;
 }
 
-bool UTypedElementDatabaseCompatibility::SupportsExtension(FName Extension) const
+bool UEditorDataStorageCompatibility::SupportsExtension(FName Extension) const
 {
 	// No thread safety needed.
 
-	using namespace TypedElementDataStorage;
+	using namespace UE::Editor::DataStorage;
 
-	if (Extension == IntegrateWithGCName)
+	if (Extension == Private::IntegrateWithGCName)
 	{
-		return bIntegrateWithGC;
+		return Private::bIntegrateWithGC;
 	}
-	else if (Extension == CompatibilityUsesCommandBufferExtensionName)
+	else if (Extension == Private::CompatibilityUsesCommandBufferExtensionName)
 	{
-		return bUseCommandBuffer;
+		return Private::bUseCommandBuffer;
 	}
 	else
 	{
@@ -322,34 +321,34 @@ bool UTypedElementDatabaseCompatibility::SupportsExtension(FName Extension) cons
 	}
 }
 
-void UTypedElementDatabaseCompatibility::ListExtensions(TFunctionRef<void(FName)> Callback) const
+void UEditorDataStorageCompatibility::ListExtensions(TFunctionRef<void(FName)> Callback) const
 {
 	// No thread safety needed.
 
-	using namespace TypedElementDataStorage;
+	using namespace UE::Editor::DataStorage;
 
-	if (bIntegrateWithGC)
+	if (Private::bIntegrateWithGC)
 	{
-		Callback(IntegrateWithGCName);
+		Callback(Private::IntegrateWithGCName);
 	}
-	if (bUseCommandBuffer)
+	if (Private::bUseCommandBuffer)
 	{
-		Callback(CompatibilityUsesCommandBufferExtensionName);
+		Callback(Private::CompatibilityUsesCommandBufferExtensionName);
 	}
 }
 
-void UTypedElementDatabaseCompatibility::Prepare()
+void UEditorDataStorageCompatibility::Prepare()
 {
 	// Thread-safe as this is only called from a function that has an exclusive lock.
 	CreateStandardArchetypes();
 	RegisterTypeInformationQueries();
 }
 
-void UTypedElementDatabaseCompatibility::Reset()
+void UEditorDataStorageCompatibility::Reset()
 {
 }
 
-void UTypedElementDatabaseCompatibility::CreateStandardArchetypes()
+void UEditorDataStorageCompatibility::CreateStandardArchetypes()
 {
 	// Thread-safe as this is only called from a function that has an exclusive lock.
 
@@ -377,7 +376,7 @@ void UTypedElementDatabaseCompatibility::CreateStandardArchetypes()
 	RegisterTypeTableAssociation(UObject::StaticClass(), StandardUObjectTable);
 }
 
-void UTypedElementDatabaseCompatibility::RegisterTypeInformationQueries()
+void UEditorDataStorageCompatibility::RegisterTypeInformationQueries()
 {
 	// Thread-safe as this is only called from a function that has an exclusive lock.
 
@@ -399,7 +398,7 @@ void UTypedElementDatabaseCompatibility::RegisterTypeInformationQueries()
 		.Compile());
 }
 
-bool UTypedElementDatabaseCompatibility::ShouldAddObject(const UObject* Object) const
+bool UEditorDataStorageCompatibility::ShouldAddObject(const UObject* Object) const
 {
 	using namespace TypedElementDataStorage;
 	using namespace UE::Editor::DataStorage;
@@ -419,7 +418,7 @@ bool UTypedElementDatabaseCompatibility::ShouldAddObject(const UObject* Object)
 	return Include;
 }
 
-TypedElementDataStorage::TableHandle UTypedElementDatabaseCompatibility::FindBestMatchingTable(const UStruct* TypeInfo) const
+TypedElementDataStorage::TableHandle UEditorDataStorageCompatibility::FindBestMatchingTable(const UStruct* TypeInfo) const
 {
 	using namespace TypedElementDataStorage;
 	using namespace UE::Editor::DataStorage;
@@ -439,7 +438,7 @@ TypedElementDataStorage::TableHandle UTypedElementDatabaseCompatibility::FindBes
 }
 
 template<bool bEnableTransactions>
-TypedElementRowHandle UTypedElementDatabaseCompatibility::AddCompatibleObjectExplicitTransactionable(UObject* Object)
+TypedElementRowHandle UEditorDataStorageCompatibility::AddCompatibleObjectExplicitTransactionable(UObject* Object)
 {
 	using namespace TypedElementDataStorage;
 	using namespace UE::Editor::DataStorage;
@@ -451,7 +450,7 @@ TypedElementRowHandle UTypedElementDatabaseCompatibility::AddCompatibleObjectExp
 	{
 		Result = Storage->ReserveRow();
 		Storage->IndexRow(GenerateIndexHash(Object), Result);
-		if (bUseCommandBuffer)
+		if (Private::bUseCommandBuffer)
 		{
 			QueuedCommands.AddCommand(FAddCompatibleUObject{ .Object = Object, .Row = Result });
 		}
@@ -472,7 +471,7 @@ TypedElementRowHandle UTypedElementDatabaseCompatibility::AddCompatibleObjectExp
 }
 
 template<bool bEnableTransactions>
-void UTypedElementDatabaseCompatibility::RemoveCompatibleObjectExplicitTransactionable(const UObject* Object)
+void UEditorDataStorageCompatibility::RemoveCompatibleObjectExplicitTransactionable(const UObject* Object)
 {
 	using namespace TypedElementDataStorage;
 	using namespace UE::Editor::DataStorage;
@@ -482,7 +481,7 @@ void UTypedElementDatabaseCompatibility::RemoveCompatibleObjectExplicitTransacti
 	
 	if constexpr (!bEnableTransactions)
 	{
-		if (TypedElementDataStorage::bUseCommandBuffer && TypedElementDataStorage::bUseDeferredRemovesInCompat)
+		if (Private::bUseCommandBuffer && Private::bUseDeferredRemovesInCompat)
 		{
 			// There's no need for an transaction recording so the full operation can be done as part of the commands processing.
 			QueuedCommands.AddCommand(FRemoveCompatibleUObject{ .Object = Object });
@@ -504,7 +503,7 @@ void UTypedElementDatabaseCompatibility::RemoveCompatibleObjectExplicitTransacti
 }
 
 template<bool bEnableTransactions>
-void UTypedElementDatabaseCompatibility::RemoveCompatibleObjectExplicitTransactionable(
+void UEditorDataStorageCompatibility::RemoveCompatibleObjectExplicitTransactionable(
 	const UObject* Object, TypedElementDataStorage::RowHandle ObjectRow)
 {
 	using namespace TypedElementDataStorage;
@@ -513,7 +512,7 @@ void UTypedElementDatabaseCompatibility::RemoveCompatibleObjectExplicitTransacti
 	checkf(Storage,
 		TEXT("Removing compatible objects is not supported before Typed Element's Database compatibility manager has been initialized."));
 
-	if (TypedElementDataStorage::bUseCommandBuffer && TypedElementDataStorage::bUseDeferredRemovesInCompat)
+	if (Private::bUseCommandBuffer && Private::bUseDeferredRemovesInCompat)
 	{
 		if constexpr (bEnableTransactions)
 		{
@@ -547,7 +546,7 @@ void UTypedElementDatabaseCompatibility::RemoveCompatibleObjectExplicitTransacti
 	}
 }
 
-TypedElementRowHandle UTypedElementDatabaseCompatibility::DealiasObject(const UObject* Object) const
+TypedElementRowHandle UEditorDataStorageCompatibility::DealiasObject(const UObject* Object) const
 {
 	// Thread safe because it's only called from functions that already lock.
 
@@ -561,7 +560,7 @@ TypedElementRowHandle UTypedElementDatabaseCompatibility::DealiasObject(const UO
 	return TypedElementDataStorage::InvalidRowHandle;
 }
 
-void UTypedElementDatabaseCompatibility::Tick()
+void UEditorDataStorageCompatibility::Tick()
 {
 	using namespace UE::Editor::DataStorage;
 	
@@ -572,7 +571,7 @@ void UTypedElementDatabaseCompatibility::Tick()
 	// Delay processing until the required systems are available by not clearing any lists or doing any work.
 	if (Storage && Storage->IsAvailable())
 	{
-		if (TypedElementDataStorage::bUseCommandBuffer)
+		if (Private::bUseCommandBuffer)
 		{
 			TickPendingCommands();
 		}
@@ -592,12 +591,12 @@ void UTypedElementDatabaseCompatibility::Tick()
 // FPendingTypeInformatUpdate
 // 
 
-UTypedElementDatabaseCompatibility::FPendingTypeInformationUpdate::FPendingTypeInformationUpdate()
+UEditorDataStorageCompatibility::FPendingTypeInformationUpdate::FPendingTypeInformationUpdate()
 	: PendingTypeInformationUpdatesActive(&PendingTypeInformationUpdates[0])
 	, PendingTypeInformationUpdatesSwapped(&PendingTypeInformationUpdates[1])
 {}
 
-void UTypedElementDatabaseCompatibility::FPendingTypeInformationUpdate::AddTypeInformation(const TMap<UObject*, UObject*>& ReplacedObjects)
+void UEditorDataStorageCompatibility::FPendingTypeInformationUpdate::AddTypeInformation(const TMap<UObject*, UObject*>& ReplacedObjects)
 {
 	UE::TUniqueLock Lock(Safeguard);
 
@@ -611,7 +610,7 @@ void UTypedElementDatabaseCompatibility::FPendingTypeInformationUpdate::AddTypeI
 	}
 }
 
-void UTypedElementDatabaseCompatibility::FPendingTypeInformationUpdate::Process(UTypedElementDatabaseCompatibility& Compatibility)
+void UEditorDataStorageCompatibility::FPendingTypeInformationUpdate::Process(UEditorDataStorageCompatibility& Compatibility)
 {
 	using namespace TypedElementDataStorage;
 	using namespace TypedElementQueryBuilder;
@@ -684,7 +683,7 @@ void UTypedElementDatabaseCompatibility::FPendingTypeInformationUpdate::Process(
 	}
 }
 
-TOptional<TWeakObjectPtr<UObject>> UTypedElementDatabaseCompatibility::FPendingTypeInformationUpdate::ProcessResolveTypeRecursively(
+TOptional<TWeakObjectPtr<UObject>> UEditorDataStorageCompatibility::FPendingTypeInformationUpdate::ProcessResolveTypeRecursively(
 	const TWeakObjectPtr<UObject>& Target)
 {
 	// Thread-safety guaranteed because this is a private function that only gets called from functions that called inside a mutex.
@@ -709,28 +708,28 @@ TOptional<TWeakObjectPtr<UObject>> UTypedElementDatabaseCompatibility::FPendingT
 //
 
 template<typename AddressType>
-void UTypedElementDatabaseCompatibility::PendingRegistration<AddressType>::Add(TypedElementRowHandle ReservedRowHandle, AddressType Address)
+void UEditorDataStorageCompatibility::PendingRegistration<AddressType>::Add(TypedElementRowHandle ReservedRowHandle, AddressType Address)
 {
 	// Thread-safe as it's only called from functions that already lock using.
 	Entries.Emplace(FEntry{ .Address = Address, .Row = ReservedRowHandle });
 }
 
 template<typename AddressType>
-bool UTypedElementDatabaseCompatibility::PendingRegistration<AddressType>::IsEmpty() const
+bool UEditorDataStorageCompatibility::PendingRegistration<AddressType>::IsEmpty() const
 {
 	// Thread-safe as it's only called from functions that already lock using.
 	return Entries.IsEmpty();
 }
 
 template<typename AddressType>
-int32 UTypedElementDatabaseCompatibility::PendingRegistration<AddressType>::Num() const
+int32 UEditorDataStorageCompatibility::PendingRegistration<AddressType>::Num() const
 {
 	// Thread-safe as it's only called from functions that already lock using.
 	return Entries.Num();
 }
 
 template<typename AddressType>
-void UTypedElementDatabaseCompatibility::PendingRegistration<AddressType>::ForEachAddress(const TFunctionRef<void(AddressType&)>& Callback)
+void UEditorDataStorageCompatibility::PendingRegistration<AddressType>::ForEachAddress(const TFunctionRef<void(AddressType&)>& Callback)
 {
 	// Thread-safe as it's only called from functions that already lock using.
 	for (FEntry& Entry : Entries)
@@ -740,8 +739,8 @@ void UTypedElementDatabaseCompatibility::PendingRegistration<AddressType>::ForEa
 }
 
 template<typename AddressType>
-void UTypedElementDatabaseCompatibility::PendingRegistration<AddressType>::ProcessEntries(ITypedElementDataStorageInterface& StorageInterface,
-	UTypedElementDatabaseCompatibility& Compatibility, const TFunctionRef<void(TypedElementRowHandle, const AddressType&)>& SetupRowCallback)
+void UEditorDataStorageCompatibility::PendingRegistration<AddressType>::ProcessEntries(ITypedElementDataStorageInterface& StorageInterface,
+	UEditorDataStorageCompatibility& Compatibility, const TFunctionRef<void(TypedElementRowHandle, const AddressType&)>& SetupRowCallback)
 {
 	// Thread-safe as it's only called from functions that already lock using.
 	using namespace TypedElementDataStorage;
@@ -848,13 +847,13 @@ void UTypedElementDatabaseCompatibility::PendingRegistration<AddressType>::Proce
 }
 
 template<typename AddressType>
-void UTypedElementDatabaseCompatibility::PendingRegistration<AddressType>::Reset()
+void UEditorDataStorageCompatibility::PendingRegistration<AddressType>::Reset()
 {
 	// Thread-safe as it's only called from functions that already lock using.
 	Entries.Reset();
 }
 
-void UTypedElementDatabaseCompatibility::TickPendingCommands()
+void UEditorDataStorageCompatibility::TickPendingCommands()
 {
 	using namespace UE::Editor::DataStorage;
 	using namespace TypedElementDataStorage;
@@ -880,17 +879,17 @@ void UTypedElementDatabaseCompatibility::TickPendingCommands()
 
 	if (CommandCount > 0)
 	{
-		if (bOptimizeCommandBuffer)
+		if (Private::bOptimizeCommandBuffer)
 		{
 			TEDS_EVENT_SCOPE(TEXT("Compatibility Tick - Optimization"))
 			FSorter::SortCommands(PendingCommands);
 			FCommandOptimizer::Run(PendingCommands, Environment->GetScratchBuffer());
 		}
 
-		if (PrintCompatCommandBuffer > 0 )
+		if (Private::PrintCompatCommandBuffer > 0 )
 		{
 			TEDS_EVENT_SCOPE(TEXT("Compatibility Tick - Logging"))
-			FString CommandsAsString = FRecordCommands::PrintToString(PendingCommands, PrintCompatCommandBuffer == 2);
+			FString CommandsAsString = FRecordCommands::PrintToString(PendingCommands, Private::PrintCompatCommandBuffer == 2);
 			UE_LOG(LogTedsCompat, Log, TEXT("Pending Commands:\n%s%u Nops"), 
 				*CommandsAsString, PendingCommands.GetCommandCount<FNopCommand>());
 		}
@@ -901,7 +900,7 @@ void UTypedElementDatabaseCompatibility::TickPendingCommands()
 	PendingCommands.Reset();
 }
 
-void UTypedElementDatabaseCompatibility::TickPendingUObjectRegistration()
+void UEditorDataStorageCompatibility::TickPendingUObjectRegistration()
 {
 	// Thread safe because it's only called from functions that already lock.
 
@@ -931,7 +930,7 @@ void UTypedElementDatabaseCompatibility::TickPendingUObjectRegistration()
 	}
 }
 
-void UTypedElementDatabaseCompatibility::TickPendingExternalObjectRegistration()
+void UEditorDataStorageCompatibility::TickPendingExternalObjectRegistration()
 {
 	// Thread safe because it's only called from functions that already lock.
 
@@ -952,7 +951,7 @@ void UTypedElementDatabaseCompatibility::TickPendingExternalObjectRegistration()
 	}
 }
 
-void UTypedElementDatabaseCompatibility::TickObjectSync()
+void UEditorDataStorageCompatibility::TickObjectSync()
 {
 	using namespace TypedElementDataStorage;
 	
@@ -996,11 +995,11 @@ void UTypedElementDatabaseCompatibility::TickObjectSync()
 	}
 }
 
-void UTypedElementDatabaseCompatibility::OnPrePropertyChanged(UObject* Object, const FEditPropertyChain& PropertyChain)
+void UEditorDataStorageCompatibility::OnPrePropertyChanged(UObject* Object, const FEditPropertyChain& PropertyChain)
 {
 	using namespace UE::Editor::DataStorage;
 
-	if (TypedElementDataStorage::bUseCommandBuffer)
+	if (Private::bUseCommandBuffer)
 	{
 		QueuedCommands.AddCommand(FAddInteractiveSyncFromWorldTag{ .Target = Object });
 	}
@@ -1012,13 +1011,13 @@ void UTypedElementDatabaseCompatibility::OnPrePropertyChanged(UObject* Object, c
 	}
 }
 
-void UTypedElementDatabaseCompatibility::OnPostEditChangeProperty(
+void UEditorDataStorageCompatibility::OnPostEditChangeProperty(
 	UObject* Object,
 	FPropertyChangedEvent& PropertyChangedEvent)
 {
 	using namespace UE::Editor::DataStorage;
 
-	if (TypedElementDataStorage::bUseCommandBuffer)
+	if (Private::bUseCommandBuffer)
 	{
 		if (PropertyChangedEvent.ChangeType != EPropertyChangeType::Interactive)
 		{
@@ -1041,11 +1040,11 @@ void UTypedElementDatabaseCompatibility::OnPostEditChangeProperty(
 	}
 }
 
-void UTypedElementDatabaseCompatibility::OnObjectModified(UObject* Object)
+void UEditorDataStorageCompatibility::OnObjectModified(UObject* Object)
 {
 	using namespace UE::Editor::DataStorage;
 
-	if (TypedElementDataStorage::bUseCommandBuffer)
+	if (Private::bUseCommandBuffer)
 	{
 		QueuedCommands.AddCommand(FAddSyncFromWorldTag{ .Target = Object });
 	}
@@ -1061,7 +1060,7 @@ void UTypedElementDatabaseCompatibility::OnObjectModified(UObject* Object)
 	}
 }
 
-void UTypedElementDatabaseCompatibility::TriggerOnObjectAdded(
+void UEditorDataStorageCompatibility::TriggerOnObjectAdded(
 	const void* Object, UE::Editor::DataStorage::FObjectTypeInfo TypeInfo, TypedElementDataStorage::RowHandle Row) const
 {
 	using namespace UE::Editor::DataStorage;
@@ -1075,7 +1074,7 @@ void UTypedElementDatabaseCompatibility::TriggerOnObjectAdded(
 	}
 }
 
-void UTypedElementDatabaseCompatibility::TriggerOnPreObjectRemoved(
+void UEditorDataStorageCompatibility::TriggerOnPreObjectRemoved(
 	const void* Object, UE::Editor::DataStorage::FObjectTypeInfo TypeInfo, TypedElementDataStorage::RowHandle Row) const
 {
 	using namespace UE::Editor::DataStorage;
@@ -1089,11 +1088,11 @@ void UTypedElementDatabaseCompatibility::TriggerOnPreObjectRemoved(
 	}
 }
 
-void UTypedElementDatabaseCompatibility::OnObjectReinstanced(const FCoreUObjectDelegates::FReplacementObjectMap& ReplacedObjects)
+void UEditorDataStorageCompatibility::OnObjectReinstanced(const FCoreUObjectDelegates::FReplacementObjectMap& ReplacedObjects)
 {
 	using namespace UE::Editor::DataStorage;
 	
-	if (TypedElementDataStorage::bUseCommandBuffer)
+	if (Private::bUseCommandBuffer)
 	{
 		bool bHasUpdatedTypeInformation = false;
 		for (TMap<UObject*, UObject*>::TConstIterator It = ReplacedObjects.CreateConstIterator(); It; ++It)
@@ -1114,18 +1113,18 @@ void UTypedElementDatabaseCompatibility::OnObjectReinstanced(const FCoreUObjectD
 	}
 }
 
-void UTypedElementDatabaseCompatibility::OnPostGcUnreachableAnalysis()
+void UEditorDataStorageCompatibility::OnPostGcUnreachableAnalysis()
 {
 	using namespace TypedElementDataStorage;
 	using namespace TypedElementQueryBuilder;
 	using namespace UE::Editor::DataStorage;
 
-	if (bIntegrateWithGC)
+	if (Private::bIntegrateWithGC)
 	{
 		TEDS_EVENT_SCOPE(TEXT("Post GC clean up"));
 		FScopedExclusiveLock Lock(EGlobalLockScope::Public);
 
-		if (TypedElementDataStorage::bUseCommandBuffer)
+		if (Private::bUseCommandBuffer)
 		{
 			Storage->RunQuery(UObjectQuery, CreateDirectQueryCallbackBinding(
 				[this](RowHandle Row, const FTypedElementUObjectIdColumn& ObjectId)
@@ -1177,18 +1176,18 @@ void UTypedElementDatabaseCompatibility::OnPostGcUnreachableAnalysis()
 	}
 }
 
-void UTypedElementDatabaseCompatibility::OnPostWorldInitialization(UWorld* World, const UWorld::InitializationValues InitializationValues)
+void UEditorDataStorageCompatibility::OnPostWorldInitialization(UWorld* World, const UWorld::InitializationValues InitializationValues)
 {
 	using namespace UE::Editor::DataStorage;
 
 	FScopedExclusiveLock Lock(EGlobalLockScope::Public);
 
 	FDelegateHandle Handle = World->AddOnActorDestroyedHandler(
-		FOnActorDestroyed::FDelegate::CreateUObject(this, &UTypedElementDatabaseCompatibility::OnActorDestroyed));
+		FOnActorDestroyed::FDelegate::CreateUObject(this, &UEditorDataStorageCompatibility::OnActorDestroyed));
 	ActorDestroyedDelegateHandles.Add(World, Handle);
 }
 
-void UTypedElementDatabaseCompatibility::OnPreWorldFinishDestroy(UWorld* World)
+void UEditorDataStorageCompatibility::OnPreWorldFinishDestroy(UWorld* World)
 {
 	using namespace UE::Editor::DataStorage;
 
@@ -1201,38 +1200,39 @@ void UTypedElementDatabaseCompatibility::OnPreWorldFinishDestroy(UWorld* World)
 	}
 }
 
-void UTypedElementDatabaseCompatibility::OnActorDestroyed(AActor* Actor)
+void UEditorDataStorageCompatibility::OnActorDestroyed(AActor* Actor)
 {
 	// The only function called is already thread safe.
 	RemoveCompatibleObjectExplicit(Actor);
 }
 
-SIZE_T GetTypeHash(const UTypedElementDatabaseCompatibility::FSyncTagInfo& Column)
+SIZE_T GetTypeHash(const UEditorDataStorageCompatibility::FSyncTagInfo& Column)
 {
 	// Thread safe as it only uses local data.
 	return HashCombine(Column.ColumnType.GetWeakPtrTypeHash(), Column.bAddColumn);
 }
 
 //
-// UTypedElementDatabaseCompatibility::FRegistrationCommandChange
+// UEditorDataStorageCompatibility::FRegistrationCommandChange
 //
 
-UTypedElementDatabaseCompatibility::FRegistrationCommandChange::FRegistrationCommandChange(
-	UTypedElementDatabaseCompatibility* InOwner, UObject* InTargetObject)
+UEditorDataStorageCompatibility::FRegistrationCommandChange::FRegistrationCommandChange(
+	UEditorDataStorageCompatibility* InOwner, UObject* InTargetObject)
 	: Owner(InOwner)
 	, TargetObject(InTargetObject)
 {
 }
 
-UTypedElementDatabaseCompatibility::FRegistrationCommandChange::~FRegistrationCommandChange()
+UEditorDataStorageCompatibility::FRegistrationCommandChange::~FRegistrationCommandChange()
 {
 	// Does not require any thread locking as IsRowAvailable is thread safe and DestroyMemento will lock.
+	using namespace UE::Editor::DataStorage;
 	
 	// If there has been no revert operation, there's also no memento.
-	if (UTypedElementDatabaseCompatibility* DataStorageCompat = Owner.Get(); 
+	if (UEditorDataStorageCompatibility* DataStorageCompat = Owner.Get(); 
 		DataStorageCompat && DataStorageCompat->Storage->IsRowAvailable(MementoRow))
 	{
-		if (TypedElementDataStorage::bUseCommandBuffer)
+		if (Private::bUseCommandBuffer)
 		{
 			DataStorageCompat->QueuedCommands.AddCommand(UE::Editor::DataStorage::FDestroyMemento{ .MementoRow = MementoRow});
 		}
@@ -1243,17 +1243,17 @@ UTypedElementDatabaseCompatibility::FRegistrationCommandChange::~FRegistrationCo
 	}
 }
 
-void UTypedElementDatabaseCompatibility::FRegistrationCommandChange::Apply(UObject* Object)
+void UEditorDataStorageCompatibility::FRegistrationCommandChange::Apply(UObject* Object)
 {
 	using namespace TypedElementDataStorage;
 	using namespace UE::Editor::DataStorage;
 
 	checkf(Owner.IsValid() && Owner.Get() == Object, 
 		TEXT("Applying registration transaction command within TEDS Compat was called after TEDS is not longer available."));
-	UTypedElementDatabaseCompatibility* DataStorageCompat = Owner.Get();
+	UEditorDataStorageCompatibility* DataStorageCompat = Owner.Get();
 	if (UObject* TargetRetrieved = TargetObject.Get(/*bEvenIfPendingKill=*/ true))
 	{
-		if (TypedElementDataStorage::bUseCommandBuffer)
+		if (Private::bUseCommandBuffer)
 		{
 			RowHandle ObjectRow = DataStorageCompat->AddCompatibleObjectExplicitTransactionable<false>(TargetRetrieved);
 			DataStorageCompat->QueuedCommands.AddCommand(FRestoreMemento{ .MementoRow = MementoRow, .TargetRow  = ObjectRow });
@@ -1268,7 +1268,7 @@ void UTypedElementDatabaseCompatibility::FRegistrationCommandChange::Apply(UObje
 	}
 }
 
-void UTypedElementDatabaseCompatibility::FRegistrationCommandChange::Revert(UObject* Object)
+void UEditorDataStorageCompatibility::FRegistrationCommandChange::Revert(UObject* Object)
 {
 	using namespace TypedElementDataStorage;
 	using namespace UE::Editor::DataStorage;
@@ -1279,13 +1279,13 @@ void UTypedElementDatabaseCompatibility::FRegistrationCommandChange::Revert(UObj
 	FScopedExclusiveLock Lock(EGlobalLockScope::Public);
 	if (UObject* TargetRetrieved = TargetObject.Get(/*bEvenIfPendingKill=*/ true))
 	{
-		UTypedElementDatabaseCompatibility* DataStorageCompat = Owner.Get(); 
+		UEditorDataStorageCompatibility* DataStorageCompat = Owner.Get(); 
 		ITypedElementDataStorageInterface* DataStorage = DataStorageCompat->Storage;
 			
 		RowHandle ObjectRow = DataStorageCompat->FindRowWithCompatibleObjectExplicit(TargetRetrieved);
 		if (DataStorage->IsRowAvailable(ObjectRow))
 		{
-			if (TypedElementDataStorage::bUseCommandBuffer && TypedElementDataStorage::bUseDeferredRemovesInCompat)
+			if (Private::bUseCommandBuffer && Private::bUseDeferredRemovesInCompat)
 			{
 				MementoRow = DataStorage->ReserveRow();
 				DataStorageCompat->QueuedCommands.AddCommand(FCreateMemento{ .ReservedMementoRow = MementoRow, .TargetRow = ObjectRow });
@@ -1299,18 +1299,18 @@ void UTypedElementDatabaseCompatibility::FRegistrationCommandChange::Revert(UObj
 	}
 }
 
-FString UTypedElementDatabaseCompatibility::FRegistrationCommandChange::ToString() const
+FString UEditorDataStorageCompatibility::FRegistrationCommandChange::ToString() const
 {
 	return TEXT("Typed Element Data Storage Compatibility - Registration");
 }
 
 
 //
-// UTypedElementDatabaseCompatibility::FDeregistrationCommandChange
+// UEditorDataStorageCompatibility::FDeregistrationCommandChange
 //
 
-UTypedElementDatabaseCompatibility::FDeregistrationCommandChange::FDeregistrationCommandChange(
-	UTypedElementDatabaseCompatibility* InOwner, UObject* InTargetObject)
+UEditorDataStorageCompatibility::FDeregistrationCommandChange::FDeregistrationCommandChange(
+	UEditorDataStorageCompatibility* InOwner, UObject* InTargetObject)
 	: Owner(InOwner)
 	, TargetObject(InTargetObject)
 {
@@ -1322,7 +1322,7 @@ UTypedElementDatabaseCompatibility::FDeregistrationCommandChange::FDeregistratio
 	RowHandle ObjectRow = InOwner->FindRowWithCompatibleObjectExplicit(InTargetObject);
 	if (DataStorage->IsRowAvailable(ObjectRow))
 	{
-		if (TypedElementDataStorage::bUseCommandBuffer && TypedElementDataStorage::bUseDeferredRemovesInCompat)
+		if (Private::bUseCommandBuffer && Private::bUseDeferredRemovesInCompat)
 		{
 			MementoRow = DataStorage->ReserveRow();
 			InOwner->QueuedCommands.AddCommand(FCreateMemento{ .ReservedMementoRow = MementoRow, .TargetRow = ObjectRow });
@@ -1334,15 +1334,15 @@ UTypedElementDatabaseCompatibility::FDeregistrationCommandChange::FDeregistratio
 	}
 }
 
-UTypedElementDatabaseCompatibility::FDeregistrationCommandChange::~FDeregistrationCommandChange()
+UEditorDataStorageCompatibility::FDeregistrationCommandChange::~FDeregistrationCommandChange()
 {
 	using namespace UE::Editor::DataStorage;
 
 	// There's no memento row if target object was never registered with TEDS Compat.
-	if (UTypedElementDatabaseCompatibility* DataStorageCompat = Owner.Get();
+	if (UEditorDataStorageCompatibility* DataStorageCompat = Owner.Get();
 		DataStorageCompat && DataStorageCompat->Storage->IsRowAvailable(MementoRow))
 	{
-		if (TypedElementDataStorage::bUseCommandBuffer)
+		if (Private::bUseCommandBuffer)
 		{
 			DataStorageCompat->QueuedCommands.AddCommand(FDestroyMemento{ .MementoRow = MementoRow });
 		}
@@ -1353,7 +1353,7 @@ UTypedElementDatabaseCompatibility::FDeregistrationCommandChange::~FDeregistrati
 	}
 }
 
-void UTypedElementDatabaseCompatibility::FDeregistrationCommandChange::Apply(UObject* Object)
+void UEditorDataStorageCompatibility::FDeregistrationCommandChange::Apply(UObject* Object)
 {
 	// All function calls are guaranteed to be thread safe.
 
@@ -1361,14 +1361,14 @@ void UTypedElementDatabaseCompatibility::FDeregistrationCommandChange::Apply(UOb
 
 	checkf(Owner.IsValid() && Owner.Get() == Object,
 		TEXT("Applying deregistration transaction command within TEDS Compat was called after TEDS is not longer available."));
-	UTypedElementDatabaseCompatibility* DataStorageCompat = Owner.Get();
+	UEditorDataStorageCompatibility* DataStorageCompat = Owner.Get();
 	if (UObject* TargetRetrieved = TargetObject.Get(/*bEvenIfPendingKill=*/ true))
 	{
 		DataStorageCompat->RemoveCompatibleObjectExplicitTransactionable<false>(TargetRetrieved);
 	}
 }
 
-void UTypedElementDatabaseCompatibility::FDeregistrationCommandChange::Revert(UObject* Object)
+void UEditorDataStorageCompatibility::FDeregistrationCommandChange::Revert(UObject* Object)
 {
 	using namespace TypedElementDataStorage;
 	using namespace UE::Editor::DataStorage;
@@ -1376,10 +1376,10 @@ void UTypedElementDatabaseCompatibility::FDeregistrationCommandChange::Revert(UO
 	checkf(Owner.IsValid() && Owner.Get() == Object,
 		TEXT("Reverting deregistration transaction command within TEDS Compat was called after TEDS is not longer available."));
 	
-	UTypedElementDatabaseCompatibility* DataStorageCompat = Owner.Get();
+	UEditorDataStorageCompatibility* DataStorageCompat = Owner.Get();
 	if (UObject* TargetRetrieved = TargetObject.Get(/*bEvenIfPendingKill=*/ true))
 	{
-		if(TypedElementDataStorage::bUseCommandBuffer)
+		if(Private::bUseCommandBuffer)
 		{
 			RowHandle ObjectRow = DataStorageCompat->AddCompatibleObjectExplicitTransactionable<false>(TargetRetrieved);
 			DataStorageCompat->QueuedCommands.AddCommand(FRestoreMemento{ .MementoRow = MementoRow, .TargetRow = ObjectRow });
@@ -1394,7 +1394,7 @@ void UTypedElementDatabaseCompatibility::FDeregistrationCommandChange::Revert(UO
 	}
 }
 
-FString UTypedElementDatabaseCompatibility::FDeregistrationCommandChange::ToString() const
+FString UEditorDataStorageCompatibility::FDeregistrationCommandChange::ToString() const
 {
 	return TEXT("Typed Element Data Storage Compatibility - Deregistration");
 }
diff --git a/Engine/Plugins/Experimental/EditorDataStorage/Source/TedsCore/Private/TypedElementDatabaseCompatibility.h b/Engine/Plugins/Experimental/EditorDataStorage/Source/TedsCore/Private/TypedElementDatabaseCompatibility.h
--- a/Engine/Plugins/Experimental/EditorDataStorage/Source/TedsCore/Private/TypedElementDatabaseCompatibility.h
+++ b/Engine/Plugins/Experimental/EditorDataStorage/Source/TedsCore/Private/TypedElementDatabaseCompatibility.h
@@ -31,7 +31,7 @@ namespace UE::Editor::DataStorage
 } // namespace UE::Editor::DataStorage
 
 UCLASS()
-class TEDSCORE_API UTypedElementDatabaseCompatibility
+class TEDSCORE_API UEditorDataStorageCompatibility
 	: public UObject
 	, public ITypedElementDataStorageCompatibilityInterface
 {
@@ -41,9 +41,9 @@ class TEDSCORE_API UTypedElementDatabaseCompatibility
 	friend struct UE::Editor::DataStorage::FPatchData;
 	friend struct UE::Editor::DataStorage::FPrepareCommands;
 public:
-	~UTypedElementDatabaseCompatibility() override = default;
+	~UEditorDataStorageCompatibility() override = default;
 
-	void Initialize(UTypedElementDatabase* InStorage);
+	void Initialize(UEditorDataStorage* InStorage);
 	void Deinitialize();
 
 	void RegisterRegistrationFilter(ObjectRegistrationFilter Filter) override;
@@ -67,13 +67,13 @@ public:
 	void ListExtensions(TFunctionRef<void(FName)> Callback) const override;
 	
 private:
-	// The below changes expect UTypedElementDatabaseCompatibility to be the object passed in to StoreUndo.
+	// The below changes expect UEditorDataStorageCompatibility to be the object passed in to StoreUndo.
 	// Note that we cannot pass in TargetObject to StoreUndo because doing so seems to stomp regular Modify()
 	// changes for that object.
 	class FRegistrationCommandChange final : public FCommandChange
 	{
 	public:
-		FRegistrationCommandChange(UTypedElementDatabaseCompatibility* InOwner, UObject* InTargetObject);
+		FRegistrationCommandChange(UEditorDataStorageCompatibility* InOwner, UObject* InTargetObject);
 		~FRegistrationCommandChange() override;
 
 		void Apply(UObject* Object) override;
@@ -81,14 +81,14 @@ private:
 		FString ToString() const override;
 
 	private:
-		TWeakObjectPtr<UTypedElementDatabaseCompatibility> Owner;
+		TWeakObjectPtr<UEditorDataStorageCompatibility> Owner;
 		TWeakObjectPtr<UObject> TargetObject;
 		TypedElementDataStorage::RowHandle MementoRow = TypedElementDataStorage::InvalidRowHandle;
 	};
 	class FDeregistrationCommandChange final : public FCommandChange
 	{
 	public:
-		FDeregistrationCommandChange(UTypedElementDatabaseCompatibility* InOwner, UObject* InTargetObject);
+		FDeregistrationCommandChange(UEditorDataStorageCompatibility* InOwner, UObject* InTargetObject);
 		~FDeregistrationCommandChange() override;
 
 		void Apply(UObject* Object) override;
@@ -96,7 +96,7 @@ private:
 		FString ToString() const override;
 
 	private:
-		TWeakObjectPtr<UTypedElementDatabaseCompatibility> Owner;
+		TWeakObjectPtr<UEditorDataStorageCompatibility> Owner;
 		TWeakObjectPtr<UObject> TargetObject;
 		TypedElementDataStorage::RowHandle MementoRow = TypedElementDataStorage::InvalidRowHandle;
 	};
@@ -142,7 +142,7 @@ private:
 		FPendingTypeInformationUpdate();
 
 		void AddTypeInformation(const TMap<UObject*, UObject*>& ReplacedObjects);
-		void Process(UTypedElementDatabaseCompatibility& Compatibility);
+		void Process(UEditorDataStorageCompatibility& Compatibility);
 
 	private:
 		TOptional<TWeakObjectPtr<UObject>> ProcessResolveTypeRecursively(const TWeakObjectPtr<UObject>& Target);
@@ -186,7 +186,7 @@ private:
 		int32 Num() const;
 		
 		void ForEachAddress(const TFunctionRef<void(AddressType&)>& Callback);
-		void ProcessEntries(ITypedElementDataStorageInterface& Storage, UTypedElementDatabaseCompatibility& Compatibility,
+		void ProcessEntries(ITypedElementDataStorageInterface& Storage, UEditorDataStorageCompatibility& Compatibility,
 			const TFunctionRef<void(TypedElementRowHandle, const AddressType&)>& SetupRowCallback);
 		void Reset();
 	};
@@ -244,4 +244,4 @@ private:
 	TypedElementDataStorage::QueryHandle UObjectQuery;
 };
 
-SIZE_T GetTypeHash(const UTypedElementDatabaseCompatibility::FSyncTagInfo& Column);
+SIZE_T GetTypeHash(const UEditorDataStorageCompatibility::FSyncTagInfo& Column);
diff --git a/Engine/Plugins/Experimental/EditorDataStorage/Source/TedsCore/Private/TypedElementDatabaseEnvironment.cpp b/Engine/Plugins/Experimental/EditorDataStorage/Source/TedsCore/Private/TypedElementDatabaseEnvironment.cpp
--- a/Engine/Plugins/Experimental/EditorDataStorage/Source/TedsCore/Private/TypedElementDatabaseEnvironment.cpp
+++ b/Engine/Plugins/Experimental/EditorDataStorage/Source/TedsCore/Private/TypedElementDatabaseEnvironment.cpp
@@ -5,7 +5,7 @@
 
 namespace UE::Editor::DataStorage
 {
-	FEnvironment::FEnvironment(UTypedElementDatabase& InDataStorage,
+	FEnvironment::FEnvironment(UEditorDataStorage& InDataStorage,
 		FMassEntityManager& InMassEntityManager, FMassProcessingPhaseManager& InMassPhaseManager)
 		: DataStorage(InDataStorage)
 		, DirectDeferredCommands(*this)
@@ -56,12 +56,12 @@ namespace UE::Editor::DataStorage
 		return Queries;
 	}
 
-	UTypedElementMementoSystem& FEnvironment::GetMementoSystem()
+	FMementoSystem& FEnvironment::GetMementoSystem()
 	{
 		return MementoSystem;
 	}
 
-	const UTypedElementMementoSystem& FEnvironment::GetMementoSystem() const
+	const FMementoSystem& FEnvironment::GetMementoSystem() const
 	{
 		return MementoSystem;
 	}
diff --git a/Engine/Plugins/Experimental/EditorDataStorage/Source/TedsCore/Private/TypedElementDatabaseEnvironment.h b/Engine/Plugins/Experimental/EditorDataStorage/Source/TedsCore/Private/TypedElementDatabaseEnvironment.h
--- a/Engine/Plugins/Experimental/EditorDataStorage/Source/TedsCore/Private/TypedElementDatabaseEnvironment.h
+++ b/Engine/Plugins/Experimental/EditorDataStorage/Source/TedsCore/Private/TypedElementDatabaseEnvironment.h
@@ -11,14 +11,14 @@
 #include "Memento/TypedElementMementoSystem.h"
 #include "Queries/TypedElementExtendedQueryStore.h"
 
-class UTypedElementDatabase;
+class UEditorDataStorage;
 
 namespace UE::Editor::DataStorage
 {
 	class FEnvironment final
 	{
 	public:
-		FEnvironment(UTypedElementDatabase& InDataStorage,
+		FEnvironment(UEditorDataStorage& InDataStorage,
 			FMassEntityManager& InMassEntityManager, FMassProcessingPhaseManager& InMassPhaseManager);
 
 		Legacy::FCommandBuffer& GetDirectDeferredCommands();
@@ -33,8 +33,8 @@ namespace UE::Editor::DataStorage
 		FExtendedQueryStore& GetQueryStore();
 		const FExtendedQueryStore& GetQueryStore() const;
 
-		UTypedElementMementoSystem& GetMementoSystem();
-		const UTypedElementMementoSystem& GetMementoSystem() const;
+		FMementoSystem& GetMementoSystem();
+		const FMementoSystem& GetMementoSystem() const;
 
 		FMassEntityManager& GetMassEntityManager();
 		const FMassEntityManager& GetMassEntityManager() const;
@@ -51,12 +51,12 @@ namespace UE::Editor::DataStorage
 		uint64 GetUpdateCycleId() const;
 
 	private:
-		UTypedElementDatabase& DataStorage;
+		UEditorDataStorage& DataStorage;
 		Legacy::FCommandBuffer DirectDeferredCommands;
 		FIndexTable IndexTable;
 		FScratchBuffer ScratchBuffer;
 		FExtendedQueryStore Queries;
-		UTypedElementMementoSystem MementoSystem;
+		FMementoSystem MementoSystem;
 		FDynamicColumnGenerator DynamicColumnGenerator;
 		FDynamicTagManager DynamicTagManager;
 
diff --git a/Engine/Plugins/Experimental/EditorDataStorage/Source/TedsCore/Private/TypedElementDatabaseUI.cpp b/Engine/Plugins/Experimental/EditorDataStorage/Source/TedsCore/Private/TypedElementDatabaseUI.cpp
--- a/Engine/Plugins/Experimental/EditorDataStorage/Source/TedsCore/Private/TypedElementDatabaseUI.cpp
+++ b/Engine/Plugins/Experimental/EditorDataStorage/Source/TedsCore/Private/TypedElementDatabaseUI.cpp
@@ -13,7 +13,7 @@
 #include "GenericPlatform/GenericPlatformMemory.h"
 #include "Widgets/SlateControlledConstruction.h"
 
-DEFINE_LOG_CATEGORY(LogTypedElementDatabaseUI);
+DEFINE_LOG_CATEGORY(LogEditorDataStorageUI);
 
 namespace Internal
 {
@@ -27,22 +27,22 @@ namespace Internal
 	template<class... Ts> TOverloaded(Ts...) -> TOverloaded<Ts...>;
 }
 
-void UTypedElementDatabaseUi::Initialize(
+void UEditorDataStorageUi::Initialize(
 	ITypedElementDataStorageInterface* StorageInterface,
 	ITypedElementDataStorageCompatibilityInterface* StorageCompatibilityInterface)
 {
-	checkf(StorageInterface, TEXT("Typed Element's Database compatibility manager is being initialized with an invalid storage target."));
+	checkf(StorageInterface, TEXT("TEDS' compatibility manager is being initialized with an invalid storage target."));
 
 	Storage = StorageInterface;
 	StorageCompatibility = StorageCompatibilityInterface;
 	CreateStandardArchetypes();
 }
 
-void UTypedElementDatabaseUi::Deinitialize()
+void UEditorDataStorageUi::Deinitialize()
 {
 }
 
-void UTypedElementDatabaseUi::RegisterWidgetPurpose(FName Purpose, EPurposeType Type, FText Description)
+void UEditorDataStorageUi::RegisterWidgetPurpose(FName Purpose, EPurposeType Type, FText Description)
 {
 	FPurposeInfo* PurposeInfo = WidgetPurposes.Find(Purpose);
 	if (!PurposeInfo)
@@ -53,10 +53,10 @@ void UTypedElementDatabaseUi::RegisterWidgetPurpose(FName Purpose, EPurposeType
 	}
 }
 
-bool UTypedElementDatabaseUi::RegisterWidgetFactory(FName Purpose, const UScriptStruct* Constructor)
+bool UEditorDataStorageUi::RegisterWidgetFactory(FName Purpose, const UScriptStruct* Constructor)
 {
 	checkf(Constructor->IsChildOf(FTypedElementWidgetConstructor::StaticStruct()),
-		TEXT("Attempting to register a Typed Elements widget constructor '%s' that isn't derived from FTypedElementWidgetConstructor."),
+		TEXT("Attempting to register a widget constructor '%s' that isn't derived from FTypedElementWidgetConstructor."),
 		*Constructor->GetFullName());
 	
 	if (FPurposeInfo* PurposeInfo = WidgetPurposes.Find(Purpose))
@@ -79,7 +79,7 @@ bool UTypedElementDatabaseUi::RegisterWidgetFactory(FName Purpose, const UScript
 			}
 			return true;
 		case ITypedElementDataStorageUiInterface::EPurposeType::UniqueByNameAndColumn:
-			UE_LOG(LogTypedElementDatabaseUI, Warning,
+			UE_LOG(LogEditorDataStorageUI, Warning,
 				TEXT("Unable to register widget factory '%s' as purpose '%s' requires at least one column for matching."), 
 				*Constructor->GetName(), *Purpose.ToString());
 			return false;
@@ -90,19 +90,19 @@ bool UTypedElementDatabaseUi::RegisterWidgetFactory(FName Purpose, const UScript
 	}
 	else
 	{
-		UE_LOG(LogTypedElementDatabaseUI, Warning, 
+		UE_LOG(LogEditorDataStorageUI, Warning, 
 			TEXT("Unable to register widget factory '%s' as purpose '%s' isn't registered."), *Constructor->GetName(), *Purpose.ToString());
 		return false;
 	}
 }
 
-bool UTypedElementDatabaseUi::RegisterWidgetFactory(
+bool UEditorDataStorageUi::RegisterWidgetFactory(
 	FName Purpose, const UScriptStruct* Constructor, TypedElementDataStorage::FQueryConditions Columns)
 {
 	if (!Columns.IsEmpty())
 	{
 		checkf(Constructor->IsChildOf(FTypedElementWidgetConstructor::StaticStruct()),
-			TEXT("Attempting to register a Typed Elements widget constructor '%s' that isn't deriving from FTypedElementWidgetConstructor."),
+			TEXT("Attempting to register a widget constructor '%s' that isn't deriving from FTypedElementWidgetConstructor."),
 			*Constructor->GetFullName());
 
 		if (FPurposeInfo* PurposeInfo = WidgetPurposes.Find(Purpose))
@@ -149,7 +149,7 @@ bool UTypedElementDatabaseUi::RegisterWidgetFactory(
 		}
 		else
 		{
-			UE_LOG(LogTypedElementDatabaseUI, Warning,
+			UE_LOG(LogEditorDataStorageUI, Warning,
 				TEXT("Unable to register widget factory '%s' as purpose '%s' isn't registered."), *Constructor->GetName(), *Purpose.ToString());
 			return false;
 		}
@@ -160,7 +160,7 @@ bool UTypedElementDatabaseUi::RegisterWidgetFactory(
 	}
 }
 
-bool UTypedElementDatabaseUi::RegisterWidgetFactory(FName Purpose, TUniquePtr<FTypedElementWidgetConstructor>&& Constructor)
+bool UEditorDataStorageUi::RegisterWidgetFactory(FName Purpose, TUniquePtr<FTypedElementWidgetConstructor>&& Constructor)
 {
 	checkf(Constructor->GetTypeInfo(), TEXT("Widget constructor being registered that doesn't have valid type information."));
 	
@@ -184,7 +184,7 @@ bool UTypedElementDatabaseUi::RegisterWidgetFactory(FName Purpose, TUniquePtr<FT
 			}
 			return true;
 		case ITypedElementDataStorageUiInterface::EPurposeType::UniqueByNameAndColumn:
-			UE_LOG(LogTypedElementDatabaseUI, Warning,
+			UE_LOG(LogEditorDataStorageUI, Warning,
 				TEXT("Unable to register widget factory '%s' as purpose '%s' requires at least one column for matching."),
 				*Constructor->GetTypeInfo()->GetName(), *Purpose.ToString());
 			return false;
@@ -195,13 +195,13 @@ bool UTypedElementDatabaseUi::RegisterWidgetFactory(FName Purpose, TUniquePtr<FT
 	}
 	else
 	{
-		UE_LOG(LogTypedElementDatabaseUI, Warning, 
+		UE_LOG(LogEditorDataStorageUI, Warning, 
 			TEXT("Unable to register widget factory as purpose '%s' isn't registered."), *Purpose.ToString());
 		return false;
 	}
 }
 
-bool UTypedElementDatabaseUi::RegisterWidgetFactory(FName Purpose, TUniquePtr<FTypedElementWidgetConstructor>&& Constructor, 
+bool UEditorDataStorageUi::RegisterWidgetFactory(FName Purpose, TUniquePtr<FTypedElementWidgetConstructor>&& Constructor, 
 	TypedElementDataStorage::FQueryConditions Columns)
 {
 	if (!Columns.IsEmpty())
@@ -252,7 +252,7 @@ bool UTypedElementDatabaseUi::RegisterWidgetFactory(FName Purpose, TUniquePtr<FT
 		}
 		else
 		{
-			UE_LOG(LogTypedElementDatabaseUI, Warning, TEXT("Unable to register widget factory '%s' as purpose '%s' isn't registered."), 
+			UE_LOG(LogEditorDataStorageUI, Warning, TEXT("Unable to register widget factory '%s' as purpose '%s' isn't registered."), 
 				*Constructor->GetTypeInfo()->GetName(), *Purpose.ToString());
 			return false;
 		}
@@ -263,7 +263,7 @@ bool UTypedElementDatabaseUi::RegisterWidgetFactory(FName Purpose, TUniquePtr<FT
 	}
 }
 
-void UTypedElementDatabaseUi::CreateWidgetConstructors(FName Purpose,
+void UEditorDataStorageUi::CreateWidgetConstructors(FName Purpose,
 	const TypedElementDataStorage::FMetaDataView& Arguments, const WidgetConstructorCallback& Callback)
 {
 	if (FPurposeInfo* PurposeInfo = WidgetPurposes.Find(Purpose))
@@ -278,7 +278,7 @@ void UTypedElementDatabaseUi::CreateWidgetConstructors(FName Purpose,
 	}
 }
 
-void UTypedElementDatabaseUi::CreateWidgetConstructors(FName Purpose, EMatchApproach MatchApproach, 
+void UEditorDataStorageUi::CreateWidgetConstructors(FName Purpose, EMatchApproach MatchApproach, 
 	TArray<TWeakObjectPtr<const UScriptStruct>>& Columns, const TypedElementDataStorage::FMetaDataView& Arguments,
 	const WidgetConstructorCallback& Callback)
 {
@@ -325,7 +325,7 @@ void UTypedElementDatabaseUi::CreateWidgetConstructors(FName Purpose, EMatchAppr
 	}
 }
 
-void UTypedElementDatabaseUi::ConstructWidgets(FName Purpose, const TypedElementDataStorage::FMetaDataView& Arguments,
+void UEditorDataStorageUi::ConstructWidgets(FName Purpose, const TypedElementDataStorage::FMetaDataView& Arguments,
 	const WidgetCreatedCallback& ConstructionCallback)
 {
 	if (FPurposeInfo* PurposeInfo = WidgetPurposes.Find(Purpose))
@@ -346,7 +346,7 @@ void UTypedElementDatabaseUi::ConstructWidgets(FName Purpose, const TypedElement
 						}
 						else
 						{
-							checkf(false, TEXT("Remaining stack space is too small to create a Typed Elements widget constructor from a description."));
+							checkf(false, TEXT("Remaining stack space is too small to create a widget constructor from a description."));
 						}
 					},
 					[this, &Arguments, &ConstructionCallback](const TUniquePtr<FTypedElementWidgetConstructor>& Constructor)
@@ -358,7 +358,7 @@ void UTypedElementDatabaseUi::ConstructWidgets(FName Purpose, const TypedElement
 	}
 }
 
-bool UTypedElementDatabaseUi::CreateSingleWidgetConstructor(
+bool UEditorDataStorageUi::CreateSingleWidgetConstructor(
 	const FWidgetFactory::ConstructorType& Constructor,
 	const TypedElementDataStorage::FMetaDataView& Arguments,
 	TArray<TWeakObjectPtr<const UScriptStruct>> MatchedColumnTypes,
@@ -416,7 +416,7 @@ bool UTypedElementDatabaseUi::CreateSingleWidgetConstructor(
 	return std::visit(Visitor(MoveTemp(MatchedColumnTypes), QueryConditions, Arguments, Callback), Constructor);
 }
 
-void UTypedElementDatabaseUi::CreateWidgetInstance(
+void UEditorDataStorageUi::CreateWidgetInstance(
 	FTypedElementWidgetConstructor& Constructor, 
 	const TypedElementDataStorage::FMetaDataView& Arguments,
 	const WidgetCreatedCallback& ConstructionCallback)
@@ -434,13 +434,13 @@ void UTypedElementDatabaseUi::CreateWidgetInstance(
 	}
 }
 
-TSharedPtr<SWidget> UTypedElementDatabaseUi::ConstructWidget(TypedElementRowHandle Row, FTypedElementWidgetConstructor& Constructor,
+TSharedPtr<SWidget> UEditorDataStorageUi::ConstructWidget(TypedElementRowHandle Row, FTypedElementWidgetConstructor& Constructor,
 	const TypedElementDataStorage::FMetaDataView& Arguments)
 {
 	return Constructor.ConstructFinalWidget(Row, Storage, this, Arguments);
 }
 
-void UTypedElementDatabaseUi::ListWidgetPurposes(const WidgetPurposeCallback& Callback) const
+void UEditorDataStorageUi::ListWidgetPurposes(const WidgetPurposeCallback& Callback) const
 {
 	for (auto&& It : WidgetPurposes)
 	{
@@ -448,16 +448,16 @@ void UTypedElementDatabaseUi::ListWidgetPurposes(const WidgetPurposeCallback& Ca
 	}
 }
 
-bool UTypedElementDatabaseUi::SupportsExtension(FName Extension) const
+bool UEditorDataStorageUi::SupportsExtension(FName Extension) const
 {
 	return false;
 }
 
-void UTypedElementDatabaseUi::ListExtensions(TFunctionRef<void(FName)> Callback) const
+void UEditorDataStorageUi::ListExtensions(TFunctionRef<void(FName)> Callback) const
 {
 }
 
-void UTypedElementDatabaseUi::CreateStandardArchetypes()
+void UEditorDataStorageUi::CreateStandardArchetypes()
 {
 	WidgetTable = Storage->RegisterTable(MakeArrayView(
 		{
@@ -466,7 +466,7 @@ void UTypedElementDatabaseUi::CreateStandardArchetypes()
 		}), FName(TEXT("Editor_WidgetTable")));
 }
 
-void UTypedElementDatabaseUi::CreateWidgetConstructors_LongestMatch(const TArray<FWidgetFactory>& WidgetFactories,
+void UEditorDataStorageUi::CreateWidgetConstructors_LongestMatch(const TArray<FWidgetFactory>& WidgetFactories,
 	TArray<TWeakObjectPtr<const UScriptStruct>>& Columns, const TypedElementDataStorage::FMetaDataView& Arguments,
 	const WidgetConstructorCallback& Callback)
 {
@@ -516,7 +516,7 @@ void UTypedElementDatabaseUi::CreateWidgetConstructors_LongestMatch(const TArray
 	}
 }
 
-void UTypedElementDatabaseUi::CreateWidgetConstructors_ExactMatch(const TArray<FWidgetFactory>& WidgetFactories,
+void UEditorDataStorageUi::CreateWidgetConstructors_ExactMatch(const TArray<FWidgetFactory>& WidgetFactories,
 	TArray<TWeakObjectPtr<const UScriptStruct>>& Columns, const TypedElementDataStorage::FMetaDataView& Arguments,
 	const WidgetConstructorCallback& Callback)
 {
@@ -547,7 +547,7 @@ void UTypedElementDatabaseUi::CreateWidgetConstructors_ExactMatch(const TArray<F
 	}
 }
 
-void UTypedElementDatabaseUi::CreateWidgetConstructors_SingleMatch(const TArray<FWidgetFactory>& WidgetFactories,
+void UEditorDataStorageUi::CreateWidgetConstructors_SingleMatch(const TArray<FWidgetFactory>& WidgetFactories,
 	TArray<TWeakObjectPtr<const UScriptStruct>>& Columns, const TypedElementDataStorage::FMetaDataView& Arguments,
 	const WidgetConstructorCallback& Callback)
 {
@@ -589,26 +589,26 @@ void UTypedElementDatabaseUi::CreateWidgetConstructors_SingleMatch(const TArray<
 // FWidgetFactory
 //
 
-UTypedElementDatabaseUi::FWidgetFactory::FWidgetFactory(const UScriptStruct* InConstructor)
+UEditorDataStorageUi::FWidgetFactory::FWidgetFactory(const UScriptStruct* InConstructor)
 	: Constructor(InConstructor)
 {
 }
 
-UTypedElementDatabaseUi::FWidgetFactory::FWidgetFactory(TUniquePtr<FTypedElementWidgetConstructor>&& InConstructor)
+UEditorDataStorageUi::FWidgetFactory::FWidgetFactory(TUniquePtr<FTypedElementWidgetConstructor>&& InConstructor)
 	: Constructor(MoveTemp(InConstructor))
 {
 	checkf(std::get<TUniquePtr<FTypedElementWidgetConstructor>>(Constructor)->GetTypeInfo(), 
 		TEXT("Widget constructor registered that didn't contain valid type information."));
 }
 
-UTypedElementDatabaseUi::FWidgetFactory::FWidgetFactory(const UScriptStruct* InConstructor, 
+UEditorDataStorageUi::FWidgetFactory::FWidgetFactory(const UScriptStruct* InConstructor, 
 	TypedElementDataStorage::FQueryConditions&& InColumns)
 	: Columns(MoveTemp(InColumns))
 	, Constructor(InConstructor)
 {
 }
 
-UTypedElementDatabaseUi::FWidgetFactory::FWidgetFactory(TUniquePtr<FTypedElementWidgetConstructor>&& InConstructor, 
+UEditorDataStorageUi::FWidgetFactory::FWidgetFactory(TUniquePtr<FTypedElementWidgetConstructor>&& InConstructor, 
 	TypedElementDataStorage::FQueryConditions&& InColumns)
 	: Columns(MoveTemp(InColumns))
 	, Constructor(MoveTemp(InConstructor))
diff --git a/Engine/Plugins/Experimental/EditorDataStorage/Source/TedsCore/Private/TypedElementDatabaseUI.h b/Engine/Plugins/Experimental/EditorDataStorage/Source/TedsCore/Private/TypedElementDatabaseUI.h
--- a/Engine/Plugins/Experimental/EditorDataStorage/Source/TedsCore/Private/TypedElementDatabaseUI.h
+++ b/Engine/Plugins/Experimental/EditorDataStorage/Source/TedsCore/Private/TypedElementDatabaseUI.h
@@ -14,17 +14,17 @@
 class ITypedElementDataStorageInterface;
 class ITypedElementDataStorageCompatibilityInterface;
 
-TEDSCORE_API DECLARE_LOG_CATEGORY_EXTERN(LogTypedElementDatabaseUI, Log, All);
+TEDSCORE_API DECLARE_LOG_CATEGORY_EXTERN(LogEditorDataStorageUI, Log, All);
 
 UCLASS()
-class TEDSCORE_API UTypedElementDatabaseUi final
+class TEDSCORE_API UEditorDataStorageUi final
 	: public UObject
 	, public ITypedElementDataStorageUiInterface
 {
 	GENERATED_BODY()
 
 public:
-	~UTypedElementDatabaseUi() override = default;
+	~UEditorDataStorageUi() override = default;
 
 	void Initialize(
 		ITypedElementDataStorageInterface* StorageInterface, 
diff --git a/Engine/Plugins/Experimental/EditorDataStorage/Source/TedsCore/Private/TypedElementSubsystems.cpp b/Engine/Plugins/Experimental/EditorDataStorage/Source/TedsCore/Private/TypedElementSubsystems.cpp
--- a/Engine/Plugins/Experimental/EditorDataStorage/Source/TedsCore/Private/TypedElementSubsystems.cpp
+++ b/Engine/Plugins/Experimental/EditorDataStorage/Source/TedsCore/Private/TypedElementSubsystems.cpp
@@ -5,79 +5,79 @@
 #include "Elements/Framework/TypedElementRegistry.h"
 
 //
-// UTypedElementDataStorageSubsystem
+// UEditorDataStorageSubsystem
 //
 
-UTypedElementDataStorageSubsystem::~UTypedElementDataStorageSubsystem()
+UEditorDataStorageSubsystem::~UEditorDataStorageSubsystem()
 {
 	DataStorage = nullptr;
 }
 
-ITypedElementDataStorageInterface* UTypedElementDataStorageSubsystem::Get()
+ITypedElementDataStorageInterface* UEditorDataStorageSubsystem::Get()
 {
 	if (!DataStorage)
 	{
 		UTypedElementRegistry* Registry = UTypedElementRegistry::GetInstance();
-		checkf(Registry, TEXT("UTypedElementDataStorageSubsystem created before the Typed Elements Registry is available."));
+		checkf(Registry, TEXT("UEditorDataStorageSubsystem created before the Typed Elements Registry is available."));
 		DataStorage = Registry->GetMutableDataStorage();
 	}
 	return DataStorage;
 }
 
-const ITypedElementDataStorageInterface* UTypedElementDataStorageSubsystem::Get() const
+const ITypedElementDataStorageInterface* UEditorDataStorageSubsystem::Get() const
 {
-	return const_cast<UTypedElementDataStorageSubsystem*>(this)->Get();
+	return const_cast<UEditorDataStorageSubsystem*>(this)->Get();
 }
 
 
 //
-// UTypedElementDataStorageUiSubsystem
+// UEditorDataStorageUiSubsystem
 //
 
-UTypedElementDataStorageUiSubsystem::~UTypedElementDataStorageUiSubsystem()
+UEditorDataStorageUiSubsystem::~UEditorDataStorageUiSubsystem()
 {
 	DataStorageUi = nullptr;
 }
 
-ITypedElementDataStorageUiInterface* UTypedElementDataStorageUiSubsystem::Get()
+ITypedElementDataStorageUiInterface* UEditorDataStorageUiSubsystem::Get()
 {
 	if (!DataStorageUi)
 	{
 		UTypedElementRegistry* Registry = UTypedElementRegistry::GetInstance();
-		checkf(Registry, TEXT("UTypedElementDataStorageUiSubsystem created before the Typed Elements Registry is available."));
+		checkf(Registry, TEXT("UEditorDataStorageUiSubsystem created before the Typed Elements Registry is available."));
 		DataStorageUi = Registry->GetMutableDataStorageUi();
 	}
 	return DataStorageUi;
 }
 
-const ITypedElementDataStorageUiInterface* UTypedElementDataStorageUiSubsystem::Get() const
+const ITypedElementDataStorageUiInterface* UEditorDataStorageUiSubsystem::Get() const
 {
-	return const_cast<UTypedElementDataStorageUiSubsystem*>(this)->Get();
+	return const_cast<UEditorDataStorageUiSubsystem*>(this)->Get();
 }
 
 
 //
-// UTypedElementDataStorageCompatibilitySubsystem
+// UEditorDataStorageCompatibilitySubsystem
 //
 
-UTypedElementDataStorageCompatibilitySubsystem::~UTypedElementDataStorageCompatibilitySubsystem()
+UEditorDataStorageCompatibilitySubsystem::~UEditorDataStorageCompatibilitySubsystem()
 {
 	DataStorageCompatibility = nullptr;
 }
 
-ITypedElementDataStorageCompatibilityInterface* UTypedElementDataStorageCompatibilitySubsystem::Get()
+ITypedElementDataStorageCompatibilityInterface* UEditorDataStorageCompatibilitySubsystem::Get()
 {
 	if (!DataStorageCompatibility)
 	{
 		UTypedElementRegistry* Registry = UTypedElementRegistry::GetInstance();
-		checkf(Registry, TEXT("UTypedElementDataStorageCompatibilitySubsystem created before the Typed Elements Registry is available."));
+		checkf(Registry, TEXT("UEditorDataStorageCompatibilitySubsystem created before the Typed Elements Registry is available."));
 
 		DataStorageCompatibility = Registry->GetMutableDataStorageCompatibility();
 	}
 	return DataStorageCompatibility;
 }
 
-const ITypedElementDataStorageCompatibilityInterface* UTypedElementDataStorageCompatibilitySubsystem::Get() const
+const ITypedElementDataStorageCompatibilityInterface* UEditorDataStorageCompatibilitySubsystem::Get() const
 {
-	return const_cast<UTypedElementDataStorageCompatibilitySubsystem*>(this)->Get();
+	return const_cast<UEditorDataStorageCompatibilitySubsystem*>(this)->Get();
 }
diff --git a/Engine/Plugins/Experimental/EditorDataStorage/Source/TedsCore/Private/TypedElementsDataStorage.cpp b/Engine/Plugins/Experimental/EditorDataStorage/Source/TedsCore/Private/TypedElementsDataStorage.cpp
--- a/Engine/Plugins/Experimental/EditorDataStorage/Source/TedsCore/Private/TypedElementsDataStorage.cpp
+++ b/Engine/Plugins/Experimental/EditorDataStorage/Source/TedsCore/Private/TypedElementsDataStorage.cpp
@@ -52,7 +52,7 @@ void ImpersonateMassTagsAndFragments()
 	FTedsSharedColumn::StaticStruct()->SetSuperStruct(FMassConstSharedFragment::StaticStruct());
 }
 
-void FTypedElementsDataStorageModule::StartupModule()
+void FEditorDataStorageModule::StartupModule()
 {
 	// Setup the editor settings;
 	if (ISettingsModule* SettingsModule = FModuleManager::GetModulePtr<ISettingsModule>("Settings"))
@@ -78,54 +78,54 @@ void FTypedElementsDataStorageModule::StartupModule()
 			{
 				UE_LOG(LogEditorDataStorage, Log, TEXT("Initializing"));
 				
-				Database = NewObject<UTypedElementDatabase>();
-				Database->Initialize();
+				DataStorage = NewObject<UEditorDataStorage>();
+				DataStorage->Initialize();
 
-				DatabaseCompatibility = NewObject<UTypedElementDatabaseCompatibility>();
-				DatabaseCompatibility->Initialize(Database.Get());
+				DataStorageCompatibility = NewObject<UEditorDataStorageCompatibility>();
+				DataStorageCompatibility->Initialize(DataStorage.Get());
 
-				DatabaseUi = NewObject<UTypedElementDatabaseUi>();
-				DatabaseUi->Initialize(Database.Get(), DatabaseCompatibility.Get());
+				DataStorageUi = NewObject<UEditorDataStorageUi>();
+				DataStorageUi->Initialize(DataStorage.Get(), DataStorageCompatibility.Get());
 
 				ObjectReinstancingManager = NewObject<UTedsObjectReinstancingManager>();
-				ObjectReinstancingManager->Initialize(*Database, *DatabaseCompatibility);
+				ObjectReinstancingManager->Initialize(*DataStorage, *DataStorageCompatibility);
 
-				// Register the various database instances.
+				// Register the various DataStorage instances.
 				UTypedElementRegistry* Registry = UTypedElementRegistry::GetInstance();
 				checkf(Registry, TEXT(
 					"FTypedElementsDataStorageModule tried to register itself, but there was no Typed Element Registry to register to."));
-				Registry->SetDataStorage(Database.Get());
-				Registry->SetDataStorageCompatibility(DatabaseCompatibility.Get());
-				Registry->SetDataStorageUi(DatabaseUi.Get());
+				Registry->SetDataStorage(DataStorage.Get());
+				Registry->SetDataStorageCompatibility(DataStorageCompatibility.Get());
+				Registry->SetDataStorageUi(DataStorageUi.Get());
 
 				// Allow any factories to register their content.
 				TArray<UClass*> FactoryClasses;
 				GetDerivedClasses(UTypedElementDataStorageFactory::StaticClass(), FactoryClasses);
 
-				Database->SetFactories(FactoryClasses);
+				DataStorage->SetFactories(FactoryClasses);
 				TArray<UTypedElementDataStorageFactory*> Factories;
 				Factories.Reserve(FactoryClasses.Num());
 
 				// First pass to call all registration without dependencies.
-				for (UTypedElementDatabase::FactoryIterator Iterator = Database->CreateFactoryIterator(); Iterator; ++Iterator)
+				for (UEditorDataStorage::FactoryIterator Iterator = DataStorage->CreateFactoryIterator(); Iterator; ++Iterator)
 				{
 					UTypedElementDataStorageFactory* Factory = *Iterator;
 					
-					Factory->RegisterTables(*Database);
-					Factory->RegisterTables(*Database, *DatabaseCompatibility);
-					Factory->RegisterTickGroups(*Database);
-					Factory->RegisterRegistrationFilters(*DatabaseCompatibility);
-					Factory->RegisterDealiaser(*DatabaseCompatibility);
-					Factory->RegisterWidgetPurposes(*DatabaseUi);
+					Factory->RegisterTables(*DataStorage);
+					Factory->RegisterTables(*DataStorage, *DataStorageCompatibility);
+					Factory->RegisterTickGroups(*DataStorage);
+					Factory->RegisterRegistrationFilters(*DataStorageCompatibility);
+					Factory->RegisterDealiaser(*DataStorageCompatibility);
+					Factory->RegisterWidgetPurposes(*DataStorageUi);
 				}
 
 				// Second pass to call all registration that would benefit or need the registration in the previous pass.
-				for (UTypedElementDatabase::FactoryIterator Iterator = Database->CreateFactoryIterator(); Iterator; ++Iterator)
+				for (UEditorDataStorage::FactoryIterator Iterator = DataStorage->CreateFactoryIterator(); Iterator; ++Iterator)
 				{
 					UTypedElementDataStorageFactory* Factory = *Iterator;
 					
-					Factory->RegisterQueries(*Database);
-					Factory->RegisterWidgetConstructors(*Database, *DatabaseUi);
+					Factory->RegisterQueries(*DataStorage);
+					Factory->RegisterWidgetConstructors(*DataStorage, *DataStorageUi);
 				}
 				
 				UE_LOG(LogEditorDataStorage, Log, TEXT("Initialized"));
@@ -133,16 +133,16 @@ void FTypedElementsDataStorageModule::StartupModule()
 				bInitialized = true;
 			}
 		});
-	FCoreDelegates::OnExit.AddRaw(this, &FTypedElementsDataStorageModule::ShutdownModule);
+	FCoreDelegates::OnExit.AddRaw(this, &FEditorDataStorageModule::ShutdownModule);
 }
 
-void FTypedElementsDataStorageModule::ShutdownModule()
+void FEditorDataStorageModule::ShutdownModule()
 {
 	if (bInitialized)
 	{
 		UE_LOG(LogEditorDataStorage, Log, TEXT("Deinitializing"));
 
-		Database->ResetFactories();
+		DataStorage->ResetFactories();
 
 		UTypedElementRegistry* Registry = UTypedElementRegistry::GetInstance();
 		if (Registry) // If the registry has already been destroyed there's no point in clearing the reference.
@@ -155,31 +155,31 @@ void FTypedElementsDataStorageModule::ShutdownModule()
 		if (UObjectInitialized())
 		{
 			ObjectReinstancingManager->Deinitialize();
-			DatabaseUi->Deinitialize();
-			DatabaseCompatibility->Deinitialize();
-			Database->Deinitialize();
+			DataStorageUi->Deinitialize();
+			DataStorageCompatibility->Deinitialize();
+			DataStorage->Deinitialize();
 		}
 
 		bInitialized = false;
 	}
 }
 
-void FTypedElementsDataStorageModule::AddReferencedObjects(FReferenceCollector& Collector)
+void FEditorDataStorageModule::AddReferencedObjects(FReferenceCollector& Collector)
 {
 	if (bInitialized)
 	{
-		Collector.AddReferencedObject(Database);
-		Collector.AddReferencedObject(DatabaseCompatibility);
-		Collector.AddReferencedObject(DatabaseUi);
+		Collector.AddReferencedObject(DataStorage);
+		Collector.AddReferencedObject(DataStorageCompatibility);
+		Collector.AddReferencedObject(DataStorageUi);
 		Collector.AddReferencedObject(ObjectReinstancingManager);
 	}
 }
 
-FString FTypedElementsDataStorageModule::GetReferencerName() const
+FString FEditorDataStorageModule::GetReferencerName() const
 {
-	return TEXT("Typed Elements: Data Storage Module");
+	return TEXT("TEDS: Editor Data Storage Core Module");
 }
 
 #undef LOCTEXT_NAMESPACE
 	
-IMPLEMENT_MODULE(FTypedElementsDataStorageModule, TedsCore)
\ No newline at end of file
+IMPLEMENT_MODULE(FEditorDataStorageModule, TedsCore)
\ No newline at end of file
diff --git a/Engine/Plugins/Experimental/EditorDataStorage/Source/TedsCore/Public/Memento/TypedElementMementoTranslators.h b/Engine/Plugins/Experimental/EditorDataStorage/Source/TedsCore/Public/Memento/TypedElementMementoTranslators.h
--- a/Engine/Plugins/Experimental/EditorDataStorage/Source/TedsCore/Public/Memento/TypedElementMementoTranslators.h
+++ b/Engine/Plugins/Experimental/EditorDataStorage/Source/TedsCore/Public/Memento/TypedElementMementoTranslators.h
@@ -1,4 +1,4 @@
-﻿// Copyright Epic Games, Inc. All Rights Reserved.
+// Copyright Epic Games, Inc. All Rights Reserved.
 
 #pragma once
 
@@ -10,7 +10,7 @@
  * a row.  This can be useful for reinstancing, undo/redo, LODs, etc.
  */
 UCLASS(Abstract)
-class UTypedElementMementoTranslatorBase : public UObject
+class UTedsMementoTranslatorBase : public UObject
 {
 	GENERATED_BODY()
 public:
@@ -19,7 +19,7 @@ public:
 	 */
 	virtual const UScriptStruct* GetColumnType() const
 	{
-		LowLevelFatalError(TEXT("Pure virtual not implemented (UTypedElementColumnMementoProcessorBase)"));
+		LowLevelFatalError(TEXT("Pure virtual not implemented (UTedsMementoTranslatorBase)"));
 		return nullptr;
 	}
 
@@ -28,18 +28,18 @@ public:
 	 */
 	virtual const UScriptStruct* GetMementoType() const
 	{
-		LowLevelFatalError(TEXT("Pure virtual not implemented (UTypedElementColumnMementoProcessorBase)"));
+		LowLevelFatalError(TEXT("Pure virtual not implemented (UTedsMementoTranslatorBase)"));
 		return nullptr;
 	}
 	
 	/**
 	 * Implements the logic to convert a column to a memento
 	 */
-	virtual void TranslateColumnToMemento(const void* TypeErasedColumn, void* TypeErasedMemento) const PURE_VIRTUAL(UTypedElementColumnMementoProcessorBase, )
+	virtual void TranslateColumnToMemento(const void* TypeErasedColumn, void* TypeErasedMemento) const PURE_VIRTUAL(UTedsMementoTranslatorBase, )
 	/**
 	 * Implements the logic to convert a memento to a column
 	 */
-	virtual void TranslateMementoToColumn(const void* TypeErasedMemento, void* TypeErasedColumn) const PURE_VIRTUAL(UTypedElementColumnMementoProcessorBase, )
+	virtual void TranslateMementoToColumn(const void* TypeErasedMemento, void* TypeErasedColumn) const PURE_VIRTUAL(UTedsMementoTranslatorBase, )
 };
 
 /**
@@ -71,7 +71,7 @@ public:
  * that have both FTypedElementMementoOnDelete and FPathColumn columns
  */
 UCLASS(Abstract)
-class UTypedElementDefaultMementoTranslator : public UTypedElementMementoTranslatorBase
+class UTedsDefaultMementoTranslator : public UTedsMementoTranslatorBase
 {
 	GENERATED_BODY()
 public:
diff --git a/Engine/Plugins/Experimental/EditorDataStorage/Source/TedsCore/Public/TypedElementSubsystems.h b/Engine/Plugins/Experimental/EditorDataStorage/Source/TedsCore/Public/TypedElementSubsystems.h
--- a/Engine/Plugins/Experimental/EditorDataStorage/Source/TedsCore/Public/TypedElementSubsystems.h
+++ b/Engine/Plugins/Experimental/EditorDataStorage/Source/TedsCore/Public/TypedElementSubsystems.h
@@ -12,11 +12,11 @@ class ITypedElementDataStorageUiInterface;
 class ITypedElementDataStorageCompatibilityInterface;
 
 /**
- * A subsystem to provide alternative access to the Typed Elements Data Storage. This can be used in situations where directly accessing
- * the Data Storage from the Typed Elements Registry is not recommended, such as for MASS.
+ * A subsystem to provide alternative access to the Editor Data Storage. This should be used in most situations instead of 
+ * directly Data Storage from the compatibility location in Typed Elements Registry.
  */
 UCLASS()
-class TEDSCORE_API UTypedElementDataStorageSubsystem : public UEditorSubsystem
+class TEDSCORE_API UEditorDataStorageSubsystem : public UEditorSubsystem
 {
 	GENERATED_BODY()
 
@@ -24,7 +24,7 @@ public:
 	static constexpr bool bRequiresGameThread = true;
 	static constexpr bool bIsHotReloadable = false;
 
-	~UTypedElementDataStorageSubsystem() override;
+	~UEditorDataStorageSubsystem() override;
 
 	ITypedElementDataStorageInterface* Get();
 	const ITypedElementDataStorageInterface* Get() const;
@@ -34,11 +34,11 @@ protected:
 };
 
 /**
- * A subsystem to provide alternative access to the Typed Elements Data Storage UI. This can be used in situations where directly 
- * accessing the UI from the Typed Elements Registry is not recommended, such as for MASS.
+ * A subsystem to provide alternative access to the Editor Data Storage UI. This should be used in most situations instead of
+ * directly Data Storage UI from the compatibility location in Typed Elements Registry.
  */
 UCLASS()
-class TEDSCORE_API UTypedElementDataStorageUiSubsystem : public UEditorSubsystem
+class TEDSCORE_API UEditorDataStorageUiSubsystem : public UEditorSubsystem
 {
 	GENERATED_BODY()
 
@@ -46,7 +46,7 @@ public:
 	static constexpr bool bRequiresGameThread = true;
 	static constexpr bool bIsHotReloadable = false;
 
-	~UTypedElementDataStorageUiSubsystem() override;
+	~UEditorDataStorageUiSubsystem() override;
 
 	ITypedElementDataStorageUiInterface* Get();
 	const ITypedElementDataStorageUiInterface* Get() const;
@@ -56,11 +56,11 @@ protected:
 };
 
 /**
- * A subsystem to provide alternative access to the Typed Elements Data Storage Compatibility. This can be used in situations where directly 
- * accessing the Compatibility extension from the Typed Elements Registry is not recommended, such as for MASS.
+ * A subsystem to provide alternative access to the Editor Data Storage Compatibility. This should be used in most situations instead of
+ * directly Data Storage Compatibility from the compatibility location in Typed Elements Registry.
  */
 UCLASS()
-class TEDSCORE_API UTypedElementDataStorageCompatibilitySubsystem : public UEditorSubsystem
+class TEDSCORE_API UEditorDataStorageCompatibilitySubsystem : public UEditorSubsystem
 {
 	GENERATED_BODY()
 
@@ -68,7 +68,7 @@ public:
 	static constexpr bool bRequiresGameThread = true;
 	static constexpr bool bIsHotReloadable = false;
 
-	~UTypedElementDataStorageCompatibilitySubsystem() override;
+	~UEditorDataStorageCompatibilitySubsystem() override;
 
 	ITypedElementDataStorageCompatibilityInterface* Get();
 	const ITypedElementDataStorageCompatibilityInterface* Get() const;
diff --git a/Engine/Plugins/Experimental/EditorDataStorage/Source/TedsCore/Public/TypedElementsDataStorage.h b/Engine/Plugins/Experimental/EditorDataStorage/Source/TedsCore/Public/TypedElementsDataStorage.h
--- a/Engine/Plugins/Experimental/EditorDataStorage/Source/TedsCore/Public/TypedElementsDataStorage.h
+++ b/Engine/Plugins/Experimental/EditorDataStorage/Source/TedsCore/Public/TypedElementsDataStorage.h
@@ -8,15 +8,15 @@
 #include "UObject/ObjectPtr.h"
 
 class FReferenceCollector;
-class UTypedElementDatabase;
-class UTypedElementDatabaseCompatibility;
-class UTypedElementDatabaseUi;
+class UEditorDataStorage;
+class UEditorDataStorageCompatibility;
+class UEditorDataStorageUi;
 class UTedsObjectReinstancingManager;
 
-class FTypedElementsDataStorageModule : public IModuleInterface, public FGCObject
+class FEditorDataStorageModule : public IModuleInterface, public FGCObject
 {
 public:
-	~FTypedElementsDataStorageModule() override = default;
+	~FEditorDataStorageModule() override = default;
 
 	//
 	// IModuleInterface
@@ -34,9 +34,9 @@ public:
 	FString GetReferencerName() const override;
 
 private:
-	TObjectPtr<UTypedElementDatabase> Database;
-	TObjectPtr<UTypedElementDatabaseCompatibility> DatabaseCompatibility;
-	TObjectPtr<UTypedElementDatabaseUi> DatabaseUi;
+	TObjectPtr<UEditorDataStorage> DataStorage;
+	TObjectPtr<UEditorDataStorageCompatibility> DataStorageCompatibility;
+	TObjectPtr<UEditorDataStorageUi> DataStorageUi;
 	TObjectPtr<UTedsObjectReinstancingManager> ObjectReinstancingManager;
 	bool bInitialized{ false };
 };
