CIS Fix for Mac ARM

#rnx

diff --git a/Engine/Plugins/Experimental/EditorDataStorage/Source/TedsCore/Private/Debug/TypedElementDatabaseDebugCommands.cpp b/Engine/Plugins/Experimental/EditorDataStorage/Source/TedsCore/Private/Debug/TypedElementDatabaseDebugCommands.cpp
--- a/Engine/Plugins/Experimental/EditorDataStorage/Source/TedsCore/Private/Debug/TypedElementDatabaseDebugCommands.cpp
+++ b/Engine/Plugins/Experimental/EditorDataStorage/Source/TedsCore/Private/Debug/TypedElementDatabaseDebugCommands.cpp
@@ -15,7 +15,7 @@
 #include "Misc/OutputDevice.h"
 #include "ProfilingDebugging/CpuProfilerTrace.h"
 
-namespace Private
+namespace UE::Editor::DataStorage::Private
 {
 	static const UStruct* GetTypeInfo(const TWeakObjectPtr<const UClass>& TypeInfo) { return TypeInfo.Get(); }
 	static const UStruct* GetTypeInfo(const TWeakObjectPtr<const UScriptStruct>& TypeInfo) { return TypeInfo.Get(); }
@@ -137,6 +137,8 @@ FAutoConsoleCommandWithOutputDevice PrintObjectTypeInformationConsoleCommand(
 	TEXT("Prints the type information of any rows that has a type information column."),
 	FConsoleCommandWithOutputDeviceDelegate::CreateLambda([](FOutputDevice& Output)
 		{
+			using namespace UE::Editor::DataStorage;
+
 			TRACE_CPUPROFILER_EVENT_SCOPE(TEDS.Debug.PrintObjectTypeInfo);
 
 			if (ITypedElementDataStorageInterface* DataStorage = UTypedElementRegistry::GetInstance()->GetMutableDataStorage())
@@ -155,6 +157,8 @@ FAutoConsoleCommandWithOutputDevice PrintAllUObjectsLabelsConsoleCommand(
 	TEXT("Prints out the labels for all UObjects found in the Typed Elements Data Storage."),
 	FConsoleCommandWithOutputDeviceDelegate::CreateLambda([](FOutputDevice& Output)
 		{
+			using namespace UE::Editor::DataStorage;
+
 			TRACE_CPUPROFILER_EVENT_SCOPE(TEDS.Debug.PrintAllUObjectsLabels);
 			Output.Log(TEXT("The Typed Elements Data Storage has the following UObjects:"));
 			Private::PrintObjectLabels(Output);
@@ -166,6 +170,8 @@ FAutoConsoleCommandWithOutputDevice PrintActorLabelsConsoleCommand(
 	TEXT("Prints out the labels for all actors found in the Typed Elements Data Storage."),
 	FConsoleCommandWithOutputDeviceDelegate::CreateLambda([](FOutputDevice& Output)
 		{
+			using namespace UE::Editor::DataStorage;
+
 			TRACE_CPUPROFILER_EVENT_SCOPE(TEDS.Debug.PrintActorLabels);
 			Output.Log(TEXT("The Typed Elements Data Storage has the following actors:"));
 			Private::PrintObjectLabels<FTypedElementActorTag>(Output);
