preparing code: adjusted object casts

diff --git a/Engine/Plugins/Experimental/EditorDataStorage/Source/TedsCore/Private/TypedElementDatabaseCompatibility.cpp b/Engine/Plugins/Experimental/EditorDataStorage/Source/TedsCore/Private/TypedElementDatabaseCompatibility.cpp
--- a/Engine/Plugins/Experimental/EditorDataStorage/Source/TedsCore/Private/TypedElementDatabaseCompatibility.cpp
+++ b/Engine/Plugins/Experimental/EditorDataStorage/Source/TedsCore/Private/TypedElementDatabaseCompatibility.cpp
@@ -1136,7 +1136,7 @@ void UEditorDataStorageCompatibility::OnPostGcUnreachableAnalysis()
 					{
 						if (Description->HasAnyFlags(EInternalObjectFlags::Garbage | EInternalObjectFlags::Unreachable))
 						{
-							if (UObject* Object = Cast<UObject>(Description->Object)) // No need to delete if this isn't a full UObject.
+							if (UObject* Object = static_cast<UObject*>(Description->Object)) // No need to delete if this isn't a full UObject.
 							{
 								QueuedCommands.AddCommand(FRemoveCompatibleUObject{ .Object = Object, .ObjectRow = Row });
 							}
@@ -1167,7 +1167,7 @@ void UEditorDataStorageCompatibility::OnPostGcUnreachableAnalysis()
 
 			for (TPair<FUObjectItem*, RowHandle>& ObjectItem : DeletedObjects)
 			{
-				if (UObject* Object = Cast<UObject>(ObjectItem.Key->Object)) // No need to delete if this isn't a full UObject.
+				if (UObject* Object = static_cast<UObject*>(ObjectItem.Key->Object)) // No need to delete if this isn't a full UObject.
 				{
 					RemoveCompatibleObjectExplicitTransactionable<false>(Object, ObjectItem.Value);
 				}
