[TEDS] Dynamic Columns

diff --git a/Engine/Plugins/Experimental/EditorDataStorage/Source/TedsCore/Private/Debug/TypedElementDatabaseDebugCommands.cpp b/Engine/Plugins/Experimental/EditorDataStorage/Source/TedsCore/Private/Debug/TypedElementDatabaseDebugCommands.cpp
--- a/Engine/Plugins/Experimental/EditorDataStorage/Source/TedsCore/Private/Debug/TypedElementDatabaseDebugCommands.cpp
+++ b/Engine/Plugins/Experimental/EditorDataStorage/Source/TedsCore/Private/Debug/TypedElementDatabaseDebugCommands.cpp
@@ -230,6 +230,353 @@ static FAutoConsoleCommand CVarCreateRow(
 		UE_LOG(LogEditorDataStorage, Warning, TEXT("Added Row %llu"), static_cast<uint64>(RowHandle));
 	}));
 
+static FAutoConsoleCommand CVarAddDynamicColumnTag(
+	TEXT("TEDS.Debug.DynamicColumn.AddTag"),
+	TEXT("Argument: Row, Identifier\n"),
+	FConsoleCommandWithArgsDelegate::CreateLambda([](const TArray<FString>& Args)
+	{
+		ITypedElementDataStorageInterface* DataStorage = UTypedElementRegistry::GetInstance()->GetMutableDataStorage();
+		
+		if (Args.Num() != 2)
+		{
+			UE_LOG(LogEditorDataStorage, Error, TEXT("Invalid number of arguments. Must be 2"));
+			return;
+		}
+
+		const uint64 RowAsU64 = FCString::Strtoui64(*Args[0], nullptr, 10);
+		const TypedElementDataStorage::RowHandle Row = RowAsU64;
+		const FName Identifier = FName(*Args[1]);
+
+		DataStorage->AddColumn<FTestDynamicTag>(Row, Identifier);
+		
+	}),
+	ECVF_Default);
+
+static FAutoConsoleCommand CVarAddDynamicColumn(
+	TEXT("TEDS.Debug.DynamicColumn.AddColumn"),
+	TEXT("Argument: Row, Identifier\n"),
+	FConsoleCommandWithArgsDelegate::CreateLambda([](const TArray<FString>& Args)
+	{
+		ITypedElementDataStorageInterface* DataStorage = UTypedElementRegistry::GetInstance()->GetMutableDataStorage();
+		
+		if (Args.Num() != 2)
+		{
+			UE_LOG(LogEditorDataStorage, Error, TEXT("Invalid number of arguments. Row, TagId"));
+			return;
+		}
+
+		const uint64 RowAsU64 = FCString::Strtoui64(*Args[0], nullptr, 10);
+		const TypedElementDataStorage::RowHandle Row = RowAsU64;
+		const FName Identifier = FName(*Args[1]);
+		
+		bool bUseDefaultApi = false;
+
+		if (bUseDefaultApi)
+		{
+			DataStorage->AddColumn<FTestDynamicColumn>(Row, Identifier);
+		}
+		else
+		{
+			FTestDynamicColumn TemplateColumn;
+			DataStorage->AddColumn(Row, Identifier, MoveTemp(TemplateColumn));
+		}		
+	}),
+	ECVF_Default);
+
+static FAutoConsoleCommand CVarRemoveDynamicColumn(
+	TEXT("TEDS.Debug.DynamicColumn.RemoveColumn"),
+	TEXT("Argument: Row, Identifier\n"),
+	FConsoleCommandWithArgsDelegate::CreateLambda([](const TArray<FString>& Args)
+	{
+		ITypedElementDataStorageInterface* DataStorage = UTypedElementRegistry::GetInstance()->GetMutableDataStorage();
+		
+		if (Args.Num() != 2)
+		{
+			UE_LOG(LogEditorDataStorage, Error, TEXT("Invalid number of arguments. Row, TagId, [optional] default=true/false"));
+			return;
+		}
+
+		const uint64 RowAsU64 = FCString::Strtoui64(*Args[0], nullptr, 10);
+		const TypedElementDataStorage::RowHandle Row = RowAsU64;
+		const FName Identifier = FName(*Args[1]);
+		
+		DataStorage->RemoveColumn<FTestDynamicColumn>(Row, Identifier);
+	}),
+	ECVF_Default);
+
+// Adds a value array stored in the dynamic column denoted with the given TagId
+static FAutoConsoleCommand CVarAddToDynamicColumn(
+	TEXT("TEDS.Debug.DynamicColumn.AddToColumn"),
+	TEXT("Argument: Row, TagId, Value, [optional] MethodId\n"),
+	FConsoleCommandWithArgsDelegate::CreateLambda([](const TArray<FString>& Args)
+	{
+		ITypedElementDataStorageInterface* DataStorage = UTypedElementRegistry::GetInstance()->GetMutableDataStorage();
+		
+		if (Args.Num() < 3 || Args.Num() > 4)
+		{
+			UE_LOG(LogEditorDataStorage, Error, TEXT("Invalid number of arguments. Row, TagId, Value"));
+			return;
+		}
+
+		const uint64 RowAsU64 = FCString::Strtoui64(*Args[0], nullptr, 10);
+		const TypedElementDataStorage::RowHandle Row = RowAsU64;
+		const FName TagId = FName(*Args[1]);
+		const uint64 Value = FCString::Strtoui64(*Args[2], nullptr, 10);
+		const uint64 MethodId = Args.Num() >= 4 ? FCString::Strtoui64(*Args[3], nullptr, 10) : 0;
+
+		// Two methods
+		// 0
+		if (MethodId == 0)
+		{
+			FTestDynamicColumn* Column = DataStorage->GetColumn<FTestDynamicColumn>(Row, TagId);
+			if (!Column)
+			{
+				UE_LOG(LogEditorDataStorage, Warning, TEXT("Row does not contain dynamic column"));
+				return;
+			}
+			Column->IntArray.Add(Value);
+		}
+		else if (MethodId == 1)
+		{
+			FTestDynamicColumn* Column = DataStorage->GetColumn<FTestDynamicColumn>(Row, TagId);
+			if (!Column)
+			{
+				UE_LOG(LogEditorDataStorage, Warning, TEXT("Row does not contain dynamic column. Creating one."));
+				FTestDynamicColumn TemplateColumn;
+				TemplateColumn.IntArray.Add(Value);
+				DataStorage->AddColumn(Row, TagId, MoveTemp(TemplateColumn));
+			}
+			else
+			{
+				// Move the column to a temporary... then mutate it, then move it back
+				FTestDynamicColumn TemplateColumn = MoveTemp(*Column);
+				TemplateColumn.IntArray.Add(Value);
+				DataStorage->AddColumn(Row, TagId, MoveTemp(TemplateColumn));
+			}
+		}
+	}),
+	ECVF_Default);
+
+static FAutoConsoleCommand CVarPrintDynamicColumn(
+	TEXT("TEDS.Debug.DynamicColumn.PrintColumn"),
+	TEXT("Argument: Row, TagId\n"),
+	FConsoleCommandWithArgsDelegate::CreateLambda([](const TArray<FString>& Args)
+	{
+		ITypedElementDataStorageInterface* DataStorage = UTypedElementRegistry::GetInstance()->GetMutableDataStorage();
+		
+		if (Args.Num() != 2)
+		{
+			UE_LOG(LogEditorDataStorage, Error, TEXT("Invalid number of arguments. Row, TagId"));
+			return;
+		}
+		const uint64 RowAsU64 = FCString::Strtoui64(*Args[0], nullptr, 10);
+		const TypedElementDataStorage::RowHandle Row = RowAsU64;
+		const FName TagId = FName(*Args[1]);
+
+		FTestDynamicColumn* Column = DataStorage->GetColumn<FTestDynamicColumn>(Row, TagId);
+		if (!Column)
+		{
+			UE_LOG(LogEditorDataStorage, Warning, TEXT("Row does not contain dynamic column"));
+			return;
+		}
+
+		TStringBuilder<512> StringBuilder;
+		StringBuilder.Append(TEXT("Array: \n"));
+		const int32 EndIndex = Column->IntArray.Num();
+		for (int32 Index = 0; Index < EndIndex - 1; ++Index)
+		{
+			StringBuilder.Appendf(TEXT("[%d] %llu\n"), Index, Column->IntArray[Index]);
+		}
+		if (const int32 LastIndex = EndIndex - 1; LastIndex >= 0)
+		{
+			StringBuilder.Appendf(TEXT("[%d] %llu\n"), LastIndex, Column->IntArray[LastIndex]);
+		}
+		else
+		{
+			StringBuilder.Append(TEXT("Empty"));
+		}
+
+		UE_LOG(LogEditorDataStorage, Warning, TEXT("%s"), StringBuilder.ToString());
+		
+	}),
+	ECVF_Default);
+
+/**
+ * A Command to illustrate building a query and callback to read a dynamic column
+ */
+static FAutoConsoleCommand CVarPrintDynamicColumnWithQuery(
+	TEXT("TEDS.Debug.DynamicColumn.PrintColumnWithQuery"),
+	TEXT("Argument: Identifier\n"),
+	FConsoleCommandWithArgsDelegate::CreateLambda([](const TArray<FString>& Args)
+	{
+		using namespace TypedElementQueryBuilder;
+		using DSI = ITypedElementDataStorageInterface;
+		using namespace UE::Editor::DataStorage;
+		
+		// Print column using query
+		if (Args.Num() != 1)
+		{
+			UE_LOG(LogEditorDataStorage, Error, TEXT("Invalid number of arguments. Identifier"));
+			return;
+		}
+		
+		ITypedElementDataStorageInterface* DataStorage = UTypedElementRegistry::GetInstance()->GetMutableDataStorage();
+		
+		const FName Identifier(*Args[0]);
+
+		const TypedElementQueryHandle Query = DataStorage->RegisterQuery(
+			Select().
+				// Specify ReadOnly access to a dynamic column of type FTestDynamicColumn specified by the Identifier
+				ReadOnly<FTestDynamicColumn>(Identifier).
+			Compile());
+
+		TStringBuilder<1024> StringBuilder;
+		TypedElementDataStorage::FQueryResult Result = DataStorage->RunQuery(Query,CreateDirectQueryCallbackBinding(
+			[Identifier, &StringBuilder](TypedElementDataStorage::IDirectQueryContext& Context, const TypedElementDataStorage::RowHandle* Rows)
+			{
+				const TArrayView<const TypedElementDataStorage::RowHandle> RowView = MakeConstArrayView(Rows, Context.GetRowCount());
+				// Get pointer to the start of the range of columns to process
+				const FTestDynamicColumn* DynamicColumnsRangeStart = Context.GetColumn<FTestDynamicColumn>(Identifier);
+				const TArrayView<const FTestDynamicColumn> DynamicColumnView = MakeConstArrayView(DynamicColumnsRangeStart, Context.GetRowCount());
+				
+				for (int32 Index = 0, End = RowView.Num(); Index < End; ++Index)
+				{
+					StringBuilder.Appendf(TEXT("%llu: {"), static_cast<uint64>(RowView[Index]));
+					const FTestDynamicColumn& DynamicColumn = DynamicColumnView[Index];
+					if (!DynamicColumn.IntArray.IsEmpty())
+					{
+						for (int32 Value : DynamicColumn.IntArray)
+						{
+							StringBuilder.Appendf(TEXT("%d, "), Value);
+						}
+						// Trim the last ', '
+						StringBuilder.RemoveSuffix(2);
+					}
+					StringBuilder.Append(TEXT("}\n"));
+				}
+			}));
+
+		const int32 RowCount = Result.Count;
+		StringBuilder.Appendf(TEXT("Processed '%d' items."), RowCount);
+		
+		DataStorage->UnregisterQuery(Query);
+
+		UE_LOG(LogEditorDataStorage, Warning, TEXT("%s"), StringBuilder.ToString());
+		
+	}),
+	ECVF_Default);
+
+static FAutoConsoleCommand CVarCountDynamicTagWithQuery(
+	TEXT("TEDS.Debug.DynamicColumn.CountDynamicTagWithQuery"),
+	TEXT("Argument: Identifier\n"),
+	FConsoleCommandWithArgsDelegate::CreateLambda([](const TArray<FString>& Args)
+	{
+		using namespace TypedElementQueryBuilder;
+		using DSI = ITypedElementDataStorageInterface;
+		using namespace UE::Editor::DataStorage;
+				
+		// Print column using query
+		if (Args.Num() != 1)
+		{
+			UE_LOG(LogEditorDataStorage, Error, TEXT("Invalid number of arguments. Identifier"));
+			return;
+		}
+				
+		ITypedElementDataStorageInterface* DataStorage = UTypedElementRegistry::GetInstance()->GetMutableDataStorage();
+				
+		const FName Identifier(*Args[0]);
+
+		const TypedElementQueryHandle Query = DataStorage->RegisterQuery(
+			Select().
+				// Specify Any access to a dynamically created tag
+				Where().
+					Any<FTestDynamicTag>(Identifier).
+			Compile());
+
+		int32 Count = 0;
+		DataStorage->RunQuery(Query, CreateDirectQueryCallbackBinding([Identifier, &Count](TypedElementDataStorage::IDirectQueryContext& Context, const TypedElementDataStorage::RowHandle* Rows)
+		{
+			Count += Context.GetRowCount();
+		}));
+
+		DataStorage->UnregisterQuery(Query);
+		
+		UE_LOG(LogEditorDataStorage, Warning, TEXT("Processed '%d' items."), Count);
+		
+	}),
+	ECVF_Default);
+
+/**
+ * Registers an activatable query which will run against rows that have the FTestDynamicTag::<Identifier> column
+ */
+static FAutoConsoleCommand CVarRegisterListDynamicColumnQuery(
+	TEXT("TEDS.Debug.DynamicColumn.RegisterListDynamicColumnQuery"),
+	TEXT("Argument: Identifier ActivationGroup\n"),
+	FConsoleCommandWithArgsDelegate::CreateLambda([](const TArray<FString>& Args)
+	{
+		using namespace TypedElementQueryBuilder;
+		using DSI = ITypedElementDataStorageInterface;
+		using namespace UE::Editor::DataStorage;
+				
+		// Print column using query
+		if (Args.Num() != 2)
+		{
+			UE_LOG(LogEditorDataStorage, Error, TEXT("Invalid number of arguments. Identifier ActivationGroup"));
+			return;
+		}
+				
+		ITypedElementDataStorageInterface* DataStorage = UTypedElementRegistry::GetInstance()->GetMutableDataStorage();
+				
+		const FName Identifier(*Args[0]);
+		const FName ActivationGroup(*Args[1]);
+		
+		using namespace TypedElementQueryBuilder;
+		using DSI = ITypedElementDataStorageInterface;
+		using namespace UE::Editor::DataStorage;
+		using namespace TypedElementDataStorage;
+
+		// Lists the rows processed that have 
+		DataStorage->RegisterQuery(
+				Select(
+					TEXT("ProcessDynamicTagColumns"),
+					FProcessor(EQueryTickPhase::FrameEnd, DataStorage->GetQueryTickGroupName(EQueryTickGroups::Default))
+						.MakeActivatable(ActivationGroup),
+					[](IQueryContext& Context, const RowHandle* Rows)
+					{
+						auto RowView = MakeConstArrayView(Rows, Context.GetRowCount());
+						for (RowHandle Row : RowView)
+						{
+							UE_LOG(LogEditorDataStorage, Warning, TEXT("- '%llu'\n"), Row);
+						}
+					})
+					.Where()
+						.All<FTestDynamicTag>(Identifier)
+				.Compile());
+
+		UE_LOG(LogEditorDataStorage, Warning, TEXT("Query registered for Dynamic Column FTestDynamicTag::%s with activation group '%s'"), *Identifier.ToString(), *ActivationGroup.ToString());
+	}),
+	ECVF_Default);
+
+static FAutoConsoleCommand CVarActivateListDynamicColumnQuery(
+	TEXT("TEDS.Debug.DynamicColumn.ActivateListDynamicColumnQuery"),
+	TEXT("Argument: ActivationGroup\n"),
+	FConsoleCommandWithArgsDelegate::CreateLambda([](const TArray<FString>& Args)
+	{
+		// Print column using query
+		if (Args.Num() != 1)
+		{
+			UE_LOG(LogEditorDataStorage, Error, TEXT("Invalid number of arguments. ActivationGroup"));
+			return;
+		}
+						
+		ITypedElementDataStorageInterface* DataStorage = UTypedElementRegistry::GetInstance()->GetMutableDataStorage();
+		
+		const FName ActivationGroup(*Args[0]);
+
+		DataStorage->ActivateQueries(ActivationGroup);
+	}),
+	ECVF_Default);
+
 static FAutoConsoleCommand CVarAddDynamicTag(
 	TEXT("TEDS.Debug.DynamicTag.AddColumn"),
 	TEXT("Argument: Row, Tag, Value\n"),
diff --git a/Engine/Plugins/Experimental/EditorDataStorage/Source/TedsCore/Private/Debug/TypedElementDatabaseDebugTypes.h b/Engine/Plugins/Experimental/EditorDataStorage/Source/TedsCore/Private/Debug/TypedElementDatabaseDebugTypes.h
--- a/Engine/Plugins/Experimental/EditorDataStorage/Source/TedsCore/Private/Debug/TypedElementDatabaseDebugTypes.h
+++ b/Engine/Plugins/Experimental/EditorDataStorage/Source/TedsCore/Private/Debug/TypedElementDatabaseDebugTypes.h
@@ -2,6 +2,9 @@
 
 #pragma once
 
+#include "Elements/Common/TypedElementCommonTypes.h"
+#include "Elements/Interfaces/TypedElementDataStorageFactory.h"
+#include "StructUtils/InstancedStruct.h"
 #include "UObject/ObjectMacros.h"
 
 #include "TypedElementDatabaseDebugTypes.generated.h"
@@ -18,3 +21,19 @@ enum class ETedsDebugEnum : int8
 	Orange,
 	Purple
 };
+
+USTRUCT()
+struct FTestDynamicColumn : public FEditorDataStorageColumn
+{
+	GENERATED_BODY()
+
+	// Note: Not trivial type to check that non-default move/copying works
+	// Should work regardless of UPROPERTY
+	TArray<int32> IntArray;
+};
+
+USTRUCT()
+struct FTestDynamicTag : public FEditorDataStorageTag
+{
+	GENERATED_BODY()
+};
diff --git a/Engine/Plugins/Experimental/EditorDataStorage/Source/TedsCore/Private/DynamicColumnGenerator.cpp b/Engine/Plugins/Experimental/EditorDataStorage/Source/TedsCore/Private/DynamicColumnGenerator.cpp
--- a/Engine/Plugins/Experimental/EditorDataStorage/Source/TedsCore/Private/DynamicColumnGenerator.cpp
+++ b/Engine/Plugins/Experimental/EditorDataStorage/Source/TedsCore/Private/DynamicColumnGenerator.cpp
@@ -1,26 +1,39 @@
-// Copyright Epic Games, Inc. All Rights Reserved.
+﻿// Copyright Epic Games, Inc. All Rights Reserved.
 
 #include "DynamicColumnGenerator.h"
 
 #include "AssetDefinitionDefault.h"
 #include "TypedElementDataStorageSharedColumn.h"
 #include "Elements/Common/TypedElementCommonTypes.h"
+#include "Kismet2/ReloadUtilities.h"
 
 namespace UE::Editor::DataStorage
 {
-	FDynamicColumnGeneratorInfo FDynamicColumnGenerator::GenerateColumn(const FName& ColumnName, const UScriptStruct& Template)
+	FDynamicColumnGenerator::~FDynamicColumnGenerator()
+	{
+		FReload ReloadContext(EActiveReloadType::Reinstancing, TEXT("REINSTANCING"), *GLog);
+		ReloadContext.SetEnableReinstancing(true);
+		ReloadContext.SetSendReloadCompleteNotification(false);
+		
+		for (FGeneratedColumnRecord& GeneratedColumnRecord : GeneratedColumnData)
+		{
+			GeneratedColumnRecord.Type->DeferCppStructOps(GeneratedColumnRecord.AssetPath, nullptr);
+		}
+	}
+
+	FDynamicColumnGeneratorInfo FDynamicColumnGenerator::GenerateColumn(const UScriptStruct& Template, const FName& Identifier)
 	{
 		const FGeneratedColumnKey Key
 		{
-			.Name = ColumnName,
-			.Template = Template
+			.Template = Template,
+			.Identifier = Identifier,
 		};
 	
 		FDynamicColumnGeneratorInfo GeneratedColumnInfo;
 		UE_MT_SCOPED_WRITE_ACCESS(AccessDetector);
 		
 		{
-			const int32* GeneratedColumnIndex = GenerationParamsLookup.Find(Key);
+			const int32* GeneratedColumnIndex = GeneratedColumnLookup.Find(Key);
 			if (GeneratedColumnIndex != nullptr)
 			{
 				const FGeneratedColumnRecord& GeneratedColumnRecord = GeneratedColumnData[*GeneratedColumnIndex];
@@ -30,49 +43,73 @@ namespace UE::Editor::DataStorage
 			}		
 		}
 		
-		checkf(
-			Template.IsChildOf(UE::Editor::DataStorage::FColumn::StaticStruct()) ||
-			Template.IsChildOf(UE::Editor::DataStorage::FTag::StaticStruct()) ||
-			Template.IsChildOf(FTedsSharedColumn::StaticStruct()),
-			TEXT("Template struct must derive from Column, Tag or SharedColumn"));
+		auto IsValidType = [](const UScriptStruct& Template)
+		{
+			return
+				Template.IsChildOf(UE::Editor::DataStorage::FColumn::StaticStruct()) ||
+				Template.IsChildOf(UE::Editor::DataStorage::FTag::StaticStruct()) ||
+				Template.IsChildOf(FTedsSharedColumn::StaticStruct());
+		};
+		if (!ensureMsgf(IsValidType(Template), TEXT("Template struct [%s] must derive from Column, Tag or SharedColumn"), *Template.GetName()))
+		{
+			return FDynamicColumnGeneratorInfo
+			{
+				.Type = nullptr,
+				.bNewlyGenerated = false
+			};
+		}
 	
 		{
 			checkf(Template.GetCppStructOps() != nullptr && Template.IsNative(), TEXT("Can only create column from native struct"));
+
+			TStringBuilder<256> ObjectNameBuilder;
+			ObjectNameBuilder.Append(Template.GetName());
+			ObjectNameBuilder.Append(TEXT("::"));
+			ObjectNameBuilder.Append(Identifier.ToString());
+
+			const FName ObjectName = FName(ObjectNameBuilder);
+			const FTopLevelAssetPath AssetPath(GetTransientPackage()->GetFName(), ObjectName);
 			
-			UScriptStruct* NewScriptStruct = NewObject<UScriptStruct>(GetTransientPackage(), ColumnName);
+			UScriptStruct* NewScriptStruct = NewObject<UScriptStruct>(GetTransientPackage(), ObjectName);
+			// Ensure it is not garbage collected
+			// FDynamicColumnGenerator is not a UObject and thus does not participate in GC
 			NewScriptStruct->AddToRoot();
 	
+			// New struct subclasses the template to allow for casting back to template and usage of CppStructOps
+			// for copy/move.
 			NewScriptStruct->SetSuperStruct(&const_cast<UScriptStruct&>(Template));
 			
-			NewScriptStruct->DeferCppStructOps(FTopLevelAssetPath(GetTransientPackage()->GetFName(), ColumnName), Template.GetCppStructOps());
+			NewScriptStruct->DeferCppStructOps(AssetPath, Template.GetCppStructOps());
 			NewScriptStruct->Bind();
 			NewScriptStruct->PrepareCppStructOps();
-			NewScriptStruct->StaticLink(true);
-	
-			
+			NewScriptStruct->StaticLink(true);			
 			const int32 Index = GeneratedColumnData.Emplace(FGeneratedColumnRecord
 			{
-				.Name = ColumnName,
+				.Identifier = Identifier,
 				.Template = &Template,
-				.Type = NewScriptStruct
+				.Type = NewScriptStruct,
+				.AssetPath = AssetPath
 			});
 			
-			GenerationParamsLookup.Add(Key, Index);
-			NameLookup.Add(ColumnName, Index);
-			
+			GeneratedColumnLookup.Add(Key, Index);
+						
 			GeneratedColumnInfo.Type = NewScriptStruct;
 			GeneratedColumnInfo.bNewlyGenerated = true;
 			return GeneratedColumnInfo;
 		}
 	}
 	
-	const UScriptStruct* FDynamicColumnGenerator::LookupColumn(const FName& ColumnName) const
+	const UScriptStruct* FDynamicColumnGenerator::FindColumn(const UScriptStruct& Template, const FName& Identifier) const
 	{
 		UE_MT_SCOPED_READ_ACCESS(AccessDetector);
-		
-		if (const int32* Index = NameLookup.Find(ColumnName))
+	
+		const int32* IndexPtr = GeneratedColumnLookup.Find(FGeneratedColumnKey{
+			.Template = Template,
+			.Identifier = Identifier
+			});
+		if (IndexPtr)
 		{
-			return GeneratedColumnData[*Index].Type;
+			return GeneratedColumnData[*IndexPtr].Type;
 		}
 		return nullptr;
 	}
@@ -100,7 +137,7 @@ namespace UE::Editor::DataStorage
 		{
 			const UScriptStruct* ColumnType = GenerateColumnType(InTag);
 	
-			const UE::Editor::DataStorage::FDynamicTagColumn Overlay
+			const FTedsValueTagColumn Overlay
 			{
 				.Value = InValue
 			};
@@ -115,7 +152,7 @@ namespace UE::Editor::DataStorage
 	
 	const UScriptStruct* FDynamicTagManager::GenerateColumnType(const FDynamicTag& Tag)
 	{
-		const FDynamicColumnGeneratorInfo GeneratedColumnType = ColumnGenerator.GenerateColumn(Tag.GetName(), *FDynamicTagColumn::StaticStruct());
+		const FDynamicColumnGeneratorInfo GeneratedColumnType = ColumnGenerator.GenerateColumn(*FTedsValueTagColumn::StaticStruct(), Tag.GetName());
 		
 		return GeneratedColumnType.Type;
 	}
diff --git a/Engine/Plugins/Experimental/EditorDataStorage/Source/TedsCore/Private/DynamicColumnGenerator.h b/Engine/Plugins/Experimental/EditorDataStorage/Source/TedsCore/Private/DynamicColumnGenerator.h
--- a/Engine/Plugins/Experimental/EditorDataStorage/Source/TedsCore/Private/DynamicColumnGenerator.h
+++ b/Engine/Plugins/Experimental/EditorDataStorage/Source/TedsCore/Private/DynamicColumnGenerator.h
@@ -18,7 +18,7 @@ class UScriptStruct;
 // The template struct that is used to generate the DynamicTag column
 // It is safe to reinterpret a DynamicTag column to this template to access the Value
 USTRUCT()
-struct FTedsDynamicTagColumn : public FTedsSharedColumn
+struct FTedsValueTagColumn : public FTedsSharedColumn
 {
 	GENERATED_BODY()
 	
@@ -28,7 +28,7 @@ struct FTedsDynamicTagColumn : public FTedsSharedColumn
 
 namespace UE::Editor::DataStorage
 {
-	using FDynamicTagColumn = FTedsDynamicTagColumn;
+	using FValueTagColumn = FTedsValueTagColumn;
 
 	struct FDynamicColumnInfo
 	{
@@ -47,48 +47,46 @@ namespace UE::Editor::DataStorage
 	class FDynamicColumnGenerator
 	{
 	public:
+		~FDynamicColumnGenerator();
 		/**
 		 * Generates a dynamic TEDS column type based on a Template type (if it hasn't been generated before)
 		 */
-		FDynamicColumnGeneratorInfo GenerateColumn(const FName& ColumnName, const UScriptStruct& Template);
+		FDynamicColumnGeneratorInfo GenerateColumn(const UScriptStruct& Template, const FName& Identifier);
 
-		/**
-		 * Looks up a column based on the name given 
-		 */
-		const UScriptStruct* LookupColumn(const FName& ColumnName) const;
+		const UScriptStruct* FindColumn(const UScriptStruct& Template, const FName& Identifier) const;
 	private:
 
 		struct FGeneratedColumnRecord
 		{
-			FName Name;
+			FName Identifier;
 			const UScriptStruct* Template;
 			const UScriptStruct* Type;
+			FTopLevelAssetPath AssetPath;
 		};
 
-		UE_MT_DECLARE_RW_ACCESS_DETECTOR(AccessDetector);
-
-		TArray<FGeneratedColumnRecord> GeneratedColumnData;
-
 		struct FGeneratedColumnKey
 		{
-			FName Name;
 			const UScriptStruct& Template;
+			FName Identifier;
 
 			friend bool operator==(const FGeneratedColumnKey& Lhs, const FGeneratedColumnKey& Rhs)
 			{
-				return Lhs.Name == Rhs.Name && &Lhs.Template == &Rhs.Template;
+				return Lhs.Identifier == Rhs.Identifier && &Lhs.Template == &Rhs.Template;
 			}
 
 			friend uint32 GetTypeHash(const FGeneratedColumnKey& Key)
 			{
-				return HashCombineFast(GetTypeHash(Key.Name), PointerHash(&Key.Template));
+				return HashCombineFast(GetTypeHash(Key.Identifier), PointerHash(&Key.Template));
 			}
 		};
+
+		UE_MT_DECLARE_RW_ACCESS_DETECTOR(AccessDetector);
+
+		TArray<FGeneratedColumnRecord> GeneratedColumnData;
+		
 		// Looks up generated column index by the parameters used to generate it
 		// Used to de-duplicate
-		TMap<FGeneratedColumnKey, int32> GenerationParamsLookup;
-		// Looks up generated column index by name
-		TMap<FName, int32> NameLookup;
+		TMap<FGeneratedColumnKey, int32> GeneratedColumnLookup;
 	};
 
 	class FDynamicTagManager
diff --git a/Engine/Plugins/Experimental/EditorDataStorage/Source/TedsCore/Private/Processors/TypedElementProcessorAdaptors.cpp b/Engine/Plugins/Experimental/EditorDataStorage/Source/TedsCore/Private/Processors/TypedElementProcessorAdaptors.cpp
--- a/Engine/Plugins/Experimental/EditorDataStorage/Source/TedsCore/Private/Processors/TypedElementProcessorAdaptors.cpp
+++ b/Engine/Plugins/Experimental/EditorDataStorage/Source/TedsCore/Private/Processors/TypedElementProcessorAdaptors.cpp
@@ -9,6 +9,7 @@
 #include "Queries/TypedElementExtendedQueryStore.h"
 #include "TypedElementDatabase.h"
 #include "TypedElementDatabaseEnvironment.h"
+#include "Algo/Copy.h"
 
 namespace UE::Editor::DataStorage
 {
@@ -17,9 +18,11 @@ namespace UE::Editor::DataStorage
 		struct FMassContextCommon
 		{
 			FMassExecutionContext& Context;
+			FEnvironment& Environment;
 	
-			explicit FMassContextCommon(FMassExecutionContext& InContext)
+			FMassContextCommon(FMassExecutionContext& InContext, FEnvironment& InEnvironment)
 				: Context(InContext)
+				, Environment(InEnvironment)
 			{}
 	
 			uint32 GetRowCount() const
@@ -111,6 +114,12 @@ namespace UE::Editor::DataStorage
 				checkf(bIsTagOrFragment, TEXT("Attempting to check for a column type that is not a column or tag."));
 				return false;
 			}
+
+			const UScriptStruct* FindDynamicColumnType(const UE::Editor::DataStorage::FDynamicColumnDescription& Description) const
+			{
+				const UScriptStruct* DynamicColumnType = Environment.FindDynamicColumn(*Description.TemplateType, Description.Identifier);
+				return DynamicColumnType;
+			}
 		};
 
 		struct FMassWithEnvironmentContextCommon : public FMassContextCommon
@@ -166,8 +175,7 @@ namespace UE::Editor::DataStorage
 			using ObjectCopyOrMove = void (*)(const UScriptStruct& TypeInfo, void* Destination, void* Source);
 	
 			FMassWithEnvironmentContextCommon(FMassExecutionContext& InContext, FEnvironment& InEnvironment)
-				: FMassContextCommon(InContext)
-				, Environment(InEnvironment)
+				: FMassContextCommon(InContext, InEnvironment)
 			{}
 
 			uint64 GetUpdateCycleId() const 
@@ -194,6 +202,72 @@ namespace UE::Editor::DataStorage
 					});
 			}
 
+			template<typename InputT, typename OutputT>
+			void CopyArrayViews(const InputT Input, OutputT Output)
+			{
+				for (int32 Index = 0, End = Input.Num(); Index < End; ++Index)
+				{
+					Output[Index] = Input[Index];
+				}
+			}
+
+			void AddColumns(TConstArrayView<TypedElementDataStorage::RowHandle> Rows, TConstArrayView<UE::Editor::DataStorage::FDynamicColumnDescription> DynamicColumnDescriptions)
+			{
+				struct FAddDynamicColumns
+				{
+					TConstArrayView<TypedElementDataStorage::RowHandle> Rows;
+					TConstArrayView<FDynamicColumnDescription> Descriptions;
+					TArrayView<const UScriptStruct*> ResolvedTypes;
+				};
+
+				FScratchBuffer& ScratchBuffer = Environment.GetScratchBuffer();
+				
+				FAddDynamicColumns* CommandData = ScratchBuffer.Emplace<FAddDynamicColumns>();
+				TArrayView<TypedElementDataStorage::RowHandle> ScratchRows = MakeArrayView(
+					ScratchBuffer.EmplaceArray<TypedElementDataStorage::RowHandle>(Rows.Num()),
+					DynamicColumnDescriptions.Num());
+				TArrayView<FDynamicColumnDescription> ScratchDescriptions = MakeArrayView(
+					ScratchBuffer.EmplaceArray<FDynamicColumnDescription>(DynamicColumnDescriptions.Num()),
+					DynamicColumnDescriptions.Num());
+				
+				TArrayView<const UScriptStruct*> ScratchTypes = MakeArrayView(
+					ScratchBuffer.EmplaceArray<const UScriptStruct*>(DynamicColumnDescriptions.Num()),
+					DynamicColumnDescriptions.Num());
+				
+				CopyArrayViews(Rows, ScratchRows);
+				CopyArrayViews(DynamicColumnDescriptions, ScratchDescriptions);
+
+				*CommandData = FAddDynamicColumns
+				{
+					.Rows = ScratchRows,
+					.Descriptions = ScratchDescriptions,
+					.ResolvedTypes = ScratchTypes
+				};
+				
+				this->Context.Defer().template PushCommand<FMassDeferredAddCommand>(
+					[CommandData, this](FMassEntityManager& System)
+					{
+						for (int32 DynamicTypeIndex = 0, DynamicTypeIndexEnd = CommandData->Descriptions.Num(); DynamicTypeIndex < DynamicTypeIndexEnd; ++DynamicTypeIndex)
+						{
+							const FDynamicColumnDescription& Description = CommandData->Descriptions[DynamicTypeIndex];
+							const UScriptStruct* DynamicColumnType = Environment.GenerateDynamicColumn(*Description.TemplateType, Description.Identifier);
+							CommandData->ResolvedTypes[DynamicTypeIndex] = DynamicColumnType;
+						}
+
+						FMassArchetypeCompositionDescriptor AddDescriptor;
+						TedsColumnsToMassDescriptor(AddDescriptor, CommandData->ResolvedTypes);
+						
+						for (TypedElementDataStorage::RowHandle Row : CommandData->Rows)
+						{
+							FMassEntityHandle Entity = FMassEntityHandle::FromNumber(Row);
+							if (System.IsEntityValid(Entity))
+							{
+								System.AddCompositionToEntity_GetDelta(Entity, AddDescriptor);
+							}
+						}
+					});
+			}
+
 			void* AddColumnUninitialized(TypedElementDataStorage::RowHandle Row, const UScriptStruct* ObjectType) 
 			{
 				return AddColumnUninitialized(Row, ObjectType,
@@ -263,6 +337,80 @@ namespace UE::Editor::DataStorage
 				return ColumnData;
 			}
 
+			void* AddColumnUninitialized(TypedElementDataStorage::RowHandle Row, const FDynamicColumnDescription& Description)
+			{
+				return AddColumnUninitialized(Row, Description, [](const UScriptStruct& TypeInfo, void* Destination, void* Source)
+				{
+					TypeInfo.CopyScriptStruct(Destination, Source);
+				});
+			}
+
+			void* AddColumnUninitialized(TypedElementDataStorage::RowHandle Row, const FDynamicColumnDescription& Description, ObjectCopyOrMove Relocator)
+			{
+				struct FAddDynamicColumn
+				{
+					ObjectCopyOrMove Relocator;
+					FDynamicColumnDescription Description;
+					FMassEntityHandle Entity;
+					void* Object;
+					bool bNeedsDestruction;
+
+					FAddDynamicColumn() = default;
+					FAddDynamicColumn(ObjectCopyOrMove InRelocator, const FDynamicColumnDescription& InDescription, FMassEntityHandle InEntity, void* InObject)
+						: Relocator(InRelocator)
+						, Description(InDescription)
+						, Entity(InEntity)
+						, Object(InObject)
+					{
+						// Check here and cache off the result to avoid command buffer needing to dereference UScriptStruct to check if anything
+						// needs to be done. In many cases, this is expected to be false
+						bNeedsDestruction = (this->Description.TemplateType->StructFlags & (STRUCT_IsPlainOldData | STRUCT_NoDestructor)) == 0;
+					}
+
+					~FAddDynamicColumn()
+					{
+						if (bNeedsDestruction)
+						{
+							this->Description.TemplateType->DestroyStruct(this->Object);
+						}
+					}
+				};
+
+				FScratchBuffer& ScratchBuffer = Environment.GetScratchBuffer();
+				// DynamicColumn types are derivations from their template that add no new members.  The size and alignment will be the same
+				void* ColumnData = ScratchBuffer.Allocate(Description.TemplateType->GetStructureSize(), Description.TemplateType->GetMinAlignment());
+				FAddDynamicColumn* AddedColumn =
+					ScratchBuffer.Emplace<FAddDynamicColumn>(Relocator, Description, FMassEntityHandle::FromNumber(Row), ColumnData);
+		
+				this->Context.Defer().template PushCommand<FMassDeferredAddCommand>(
+					[AddedColumn, this](FMassEntityManager& System)
+					{
+						// Check entity before proceeding. It's possible it may have been invalidated before this deferred call fired.
+						if (System.IsEntityActive(AddedColumn->Entity))
+						{
+							const UScriptStruct* DynamicStructType = Environment.GenerateDynamicColumn(*AddedColumn->Description.TemplateType, AddedColumn->Description.Identifier);
+
+							FStructView Fragment = System.GetFragmentDataStruct(AddedColumn->Entity, DynamicStructType);
+							// Check before adding.  Mass's AddFragmentToEntity is not idempotent and will assert if adding
+							// column to a row that already has one
+							if (!Fragment.IsValid())
+							{
+								System.AddFragmentToEntity(AddedColumn->Entity, DynamicStructType, 
+									[AddedColumn](void* Fragment, const UScriptStruct& FragmentType)
+									{
+										AddedColumn->Relocator(FragmentType, Fragment, AddedColumn->Object);
+									});
+							}
+							else
+							{
+								AddedColumn->Relocator(*DynamicStructType, Fragment.GetMemory(), AddedColumn->Object);
+							}
+						}
+					});
+		
+				return ColumnData;
+			}
+	
 			void AddColumns(TypedElementDataStorage::RowHandle Row, TConstArrayView<const UScriptStruct*> ColumnTypes) 
 			{
 				struct FAddedColumns
@@ -464,15 +612,12 @@ namespace UE::Editor::DataStorage
 				this->Context.Defer().DestroyEntities(
 					TConstArrayView<FMassEntityHandle>(reinterpret_cast<const FMassEntityHandle*>(Rows.begin()), Rows.Num()));
 			}
-
-		protected:
-			FEnvironment& Environment;
 		};
 
 		struct FMassDirectContextForwarder final : public ITypedElementDataStorageInterface::IDirectQueryContext
 		{
-			explicit FMassDirectContextForwarder(FMassExecutionContext& InContext)
-				: Implementation(InContext)
+			FMassDirectContextForwarder(FMassExecutionContext& InContext, FEnvironment& InEnvironment)
+				: Implementation(InContext, InEnvironment)
 			{}
 	
 			uint32 GetRowCount() const override { return Implementation.GetRowCount(); }
@@ -483,6 +628,7 @@ namespace UE::Editor::DataStorage
 			void GetColumnsUnguarded(int32 TypeCount, char** RetrievedAddresses, const TWeakObjectPtr<const UScriptStruct>* ColumnTypes, const TypedElementDataStorage::EQueryAccessType* AccessTypes) override { return Implementation.GetColumnsUnguarded(TypeCount, RetrievedAddresses, ColumnTypes, AccessTypes); }
 			bool HasColumn(const UScriptStruct* ColumnType) const override { return Implementation.HasColumn(ColumnType); }
 			bool HasColumn(TypedElementDataStorage::RowHandle Row, const UScriptStruct* ColumnType) const override { return Implementation.HasColumn(Row, ColumnType); }
+			virtual const UScriptStruct* FindDynamicColumnType(const UE::Editor::DataStorage::FDynamicColumnDescription& Description) const override { return Implementation.FindDynamicColumnType(Description); }
 
 			FMassContextCommon Implementation;
 		};
@@ -511,10 +657,15 @@ namespace UE::Editor::DataStorage
 			void RemoveRows(TConstArrayView<TypedElementDataStorage::RowHandle> Rows) override  { return Implementation.RemoveRows(Rows); }
 			void AddColumns(TypedElementDataStorage::RowHandle Row, TConstArrayView<const UScriptStruct*> ColumnTypes) override { return Implementation.AddColumns(Row, ColumnTypes); }
 			void AddColumns(TConstArrayView<TypedElementDataStorage::RowHandle> Rows, TConstArrayView<const UScriptStruct*> ColumnTypes) override  { return Implementation.AddColumns(Rows, ColumnTypes); }
+			void AddColumns(TConstArrayView<TypedElementDataStorage::RowHandle> Rows, TConstArrayView<UE::Editor::DataStorage::FDynamicColumnDescription> DynamicColumnDescriptions) override { Implementation.AddColumns(Rows, DynamicColumnDescriptions); }
 			void* AddColumnUninitialized(TypedElementDataStorage::RowHandle Row, const UScriptStruct* ColumnType) override  { return Implementation.AddColumnUninitialized(Row, ColumnType); }
 			void* AddColumnUninitialized(TypedElementDataStorage::RowHandle Row, const UScriptStruct* ObjectType, ObjectCopyOrMove Relocator) override { return Implementation.AddColumnUninitialized(Row, ObjectType, Relocator); }
+			void* AddColumnUninitialized(TypedElementDataStorage::RowHandle Row, const UE::Editor::DataStorage::FDynamicColumnDescription& DynamicColumnDescription) override { return Implementation.AddColumnUninitialized(Row, DynamicColumnDescription); }
+			void* AddColumnUninitialized(TypedElementDataStorage::RowHandle Row, const UE::Editor::DataStorage::FDynamicColumnDescription& DynamicColumnDescription, ObjectCopyOrMove Relocator) override { return Implementation.AddColumnUninitialized(Row, DynamicColumnDescription, Relocator); }
 			void RemoveColumns(TypedElementDataStorage::RowHandle Row, TConstArrayView<const UScriptStruct*> ColumnTypes) override { Implementation.RemoveColumns(Row, ColumnTypes); }
 			void RemoveColumns(TConstArrayView<TypedElementDataStorage::RowHandle> Rows, TConstArrayView<const UScriptStruct*> ColumnTypes) override { Implementation.RemoveColumns(Rows, ColumnTypes); }
+			const UScriptStruct* FindDynamicColumnType(const UE::Editor::DataStorage::FDynamicColumnDescription& Description) const override { return Implementation.FindDynamicColumnType(Description); }
+			
 
 			FMassWithEnvironmentContextCommon Implementation;
 		};
@@ -654,11 +805,15 @@ namespace UE::Editor::DataStorage
 			void RemoveRows(TConstArrayView<TypedElementDataStorage::RowHandle> Rows) override  { return Implementation.RemoveRows(Rows); }
 			void AddColumns(TypedElementDataStorage::RowHandle Row, TConstArrayView<const UScriptStruct*> ColumnTypes) override { return Implementation.AddColumns(Row, ColumnTypes); }
 			void AddColumns(TConstArrayView<TypedElementDataStorage::RowHandle> Rows, TConstArrayView<const UScriptStruct*> ColumnTypes) override  { return Implementation.AddColumns(Rows, ColumnTypes); }
+			void AddColumns(TConstArrayView<TypedElementDataStorage::RowHandle> Rows, TConstArrayView<UE::Editor::DataStorage::FDynamicColumnDescription> DynamicColumnDescriptions) override { return Implementation.AddColumns(Rows, DynamicColumnDescriptions); }
 			void* AddColumnUninitialized(TypedElementDataStorage::RowHandle Row, const UScriptStruct* ColumnType) override  { return Implementation.AddColumnUninitialized(Row, ColumnType); }
 			void* AddColumnUninitialized(TypedElementDataStorage::RowHandle Row, const UScriptStruct* ObjectType, ObjectCopyOrMove Relocator) override { return Implementation.AddColumnUninitialized(Row, ObjectType, Relocator); }
+			void* AddColumnUninitialized(TypedElementDataStorage::RowHandle Row, const FDynamicColumnDescription& DynamicColumnDescription, ObjectCopyOrMove Relocator) override { return Implementation.AddColumnUninitialized(Row, DynamicColumnDescription, Relocator); }
+			void* AddColumnUninitialized(TypedElementDataStorage::RowHandle Row, const FDynamicColumnDescription& DynamicColumnDescription) override { return Implementation.AddColumnUninitialized(Row, DynamicColumnDescription); };
 			void RemoveColumns(TypedElementDataStorage::RowHandle Row, TConstArrayView<const UScriptStruct*> ColumnTypes) override { Implementation.RemoveColumns(Row, ColumnTypes); }
 			void RemoveColumns(TConstArrayView<TypedElementDataStorage::RowHandle> Rows, TConstArrayView<const UScriptStruct*> ColumnTypes) override { Implementation.RemoveColumns(Rows, ColumnTypes); }
-	
+			const UScriptStruct* FindDynamicColumnType(const UE::Editor::DataStorage::FDynamicColumnDescription& Description) const override { return Implementation.FindDynamicColumnType(Description); }
+
 			const UObject* GetDependency(const UClass* DependencyClass) override { return Implementation.GetDependency(DependencyClass); }
 			UObject* GetMutableDependency(const UClass* DependencyClass) override { return Implementation.GetMutableDependency(DependencyClass); }
 			void GetDependencies(TArrayView<UObject*> RetrievedAddresses, TConstArrayView<TWeakObjectPtr<const UClass>> DependencyTypes, TConstArrayView<TypedElementDataStorage::EQueryAccessType> AccessTypes) override { return Implementation.GetDependencies(RetrievedAddresses, DependencyTypes, AccessTypes); }
@@ -668,6 +823,7 @@ namespace UE::Editor::DataStorage
 			TypedElementDataStorage::FQueryResult RunSubquery(int32 SubqueryIndex, TypedElementDataStorage::SubqueryCallbackRef Callback) override { return Implementation.RunSubquery(SubqueryIndex, Callback); }
 			TypedElementDataStorage::FQueryResult RunSubquery(int32 SubqueryIndex, TypedElementDataStorage::RowHandle Row, TypedElementDataStorage::SubqueryCallbackRef Callback) override { return Implementation.RunSubquery(SubqueryIndex, Row, Callback); }
 	
+
 			FMassQueryContextImplementation Implementation;
 		};
 	} // namespace Processors::Private
@@ -906,10 +1062,10 @@ TypedElementDataStorage::FQueryResult FTypedElementQueryProcessorData::Execute(
 		if (EnumHasAnyFlags(ExecutionFlags, EDirectQueryExecutionFlags::IgnoreActivationCount) || Description.Callback.ActivationCount > 0)
 		{
 			FMassExecutionContext Context(EntityManager);
-			auto ExecuteFunction = [&Result, &Callback, &Description](FMassExecutionContext& Context)
+			auto ExecuteFunction = [&Result, &Callback, &Description, &Environment](FMassExecutionContext& Context)
 				{
 					// No need to cache any subsystem dependencies as these are not accessible from a direct query.
-					UE::Editor::DataStorage::Processors::Private::FMassDirectContextForwarder QueryContext(Context);
+					UE::Editor::DataStorage::Processors::Private::FMassDirectContextForwarder QueryContext(Context, Environment);
 					Callback(Description, QueryContext);
 					Result.Count += Context.GetNumEntities();
 				};
diff --git a/Engine/Plugins/Experimental/EditorDataStorage/Source/TedsCore/Private/Queries/TypedElementExtendedQueryStore.cpp b/Engine/Plugins/Experimental/EditorDataStorage/Source/TedsCore/Private/Queries/TypedElementExtendedQueryStore.cpp
--- a/Engine/Plugins/Experimental/EditorDataStorage/Source/TedsCore/Private/Queries/TypedElementExtendedQueryStore.cpp
+++ b/Engine/Plugins/Experimental/EditorDataStorage/Source/TedsCore/Private/Queries/TypedElementExtendedQueryStore.cpp
@@ -24,7 +24,8 @@ namespace UE::Editor::DataStorage
 		StoredQuery.Description = MoveTemp(Query);
 
 		FMassEntityQuery& NativeQuery = SetupNativeQuery(StoredQuery.Description, StoredQuery);
-		bool bContinueSetup = SetupSelectedColumns(StoredQuery.Description, NativeQuery);
+		bool bContinueSetup =               SetupDynamicColumns(StoredQuery.Description, Environment);
+		bContinueSetup = bContinueSetup &&	SetupSelectedColumns(StoredQuery.Description, NativeQuery);
 		bContinueSetup = bContinueSetup && SetupChunkFilters(Result, StoredQuery.Description, Environment, NativeQuery);
 		bContinueSetup = bContinueSetup && SetupConditions(StoredQuery.Description, NativeQuery);
 		bContinueSetup = bContinueSetup && SetupDependencies(StoredQuery.Description, NativeQuery);
@@ -370,6 +371,51 @@ namespace UE::Editor::DataStorage
 		Output.Log(TEXT("End of Typed Elements Data Storage query callback list."));
 	}
 
+	bool FExtendedQueryStore::SetupDynamicColumns(ITypedElementDataStorageInterface::FQueryDescription& Query, FEnvironment& Environment)
+	{
+		using namespace TypedElementDataStorage;
+	
+		const int32 SelectionCount = Query.DynamicSelectionTypes.Num();
+
+		for (int32 Index = 0, End = SelectionCount; Index < End; ++Index)
+		{
+			// Convert the dynamic elements into concrete types
+			const FDynamicColumnDescription& Description = Query.DynamicSelectionTypes[Index];
+			if (!ensureMsgf(Description.TemplateType, TEXT("Null template type for dynamic column")))
+			{
+				continue;
+			}
+			const UScriptStruct* DynamicColumnType = Environment.GenerateDynamicColumn(*Description.TemplateType, Description.Identifier);
+
+			const EQueryAccessType AccessType = Query.DynamicSelectionAccessTypes[Index];
+			const FColumnMetaData::EFlags MetadataFlags = Query.DynamicSelectionMetaData[Index];
+			if (ensureMsgf(DynamicColumnType, TEXT("Provided query selection type can not be null.")))
+			{
+				Query.SelectionTypes.Emplace(DynamicColumnType);
+				Query.SelectionAccessTypes.Emplace(AccessType);
+				Query.SelectionMetaData.Emplace(TypedElementDataStorage::FColumnMetaData(DynamicColumnType, MetadataFlags));
+			}
+		}
+
+		for (int32 Index = 0, End = Query.DynamicConditionDescriptions.Num(); Index < End; ++Index)
+		{
+			const UE::Editor::DataStorage::FDynamicColumnDescription& Description = Query.DynamicConditionDescriptions[Index];
+			if (!ensureMsgf(Description.TemplateType, TEXT("Null template type for dynamic column")))
+			{
+				continue;
+			}
+			const UScriptStruct* DynamicColumnType = Environment.GenerateDynamicColumn(*Description.TemplateType, Description.Identifier);
+			TypedElementDataStorage::FQueryDescription::EOperatorType Operation = Query.DynamicConditionOperations[Index];
+			Query.ConditionTypes.Add(Operation);
+			Query.ConditionOperators.Add(TypedElementDataStorage::FQueryDescription::FOperator
+			{
+				.Type = DynamicColumnType
+			});
+		}
+
+		return true;
+	}
+
 	FMassEntityQuery& FExtendedQueryStore::SetupNativeQuery(
 		ITypedElementDataStorageInterface::FQueryDescription& Query, FExtendedQuery& StoredQuery)
 	{
@@ -583,32 +629,30 @@ namespace UE::Editor::DataStorage
 		return false;
 	}
 
-	static const FName DynamicTagDataParamName = TEXT("TagName");
-
 	bool FExtendedQueryStore::SetupChunkFilters(
 		Handle QueryHandle,
 		ITypedElementDataStorageInterface::FQueryDescription& Query,
 		FEnvironment& Environment,
 		FMassEntityQuery& NativeQuery)
 	{
-		if (Query.DynamicTags.IsEmpty())
+		if (Query.ValueTags.IsEmpty())
 		{
 			return true;
 		}
 
-		Algo::SortBy(Query.DynamicTags, [](const TypedElementDataStorage::FQueryDescription::FDynamicTagData& DynamicTagData)
+		Algo::SortBy(Query.ValueTags, [](const TypedElementDataStorage::FQueryDescription::FDynamicTagData& DynamicTagData)
 			{
 				return DynamicTagData.Tag.GetName();
 			}, FNameFastLess());
 		// Check if there are any duplicate groups. Not yet supported until we can match multiple MatchTags
-		UE::Editor::DataStorage::FDynamicTag PreviousTag = Query.DynamicTags[0].Tag;
-		for (int32 Index = 1, End = Query.DynamicTags.Num(); Index < End; ++Index)
+		UE::Editor::DataStorage::FDynamicTag PreviousTag = Query.ValueTags[0].Tag;
+		for (int32 Index = 1, End = Query.ValueTags.Num(); Index < End; ++Index)
 		{
-			if (Query.DynamicTags[Index].Tag == PreviousTag)
+			if (Query.ValueTags[Index].Tag == PreviousTag)
 			{
 				return false;
 			}
-			PreviousTag = Query.DynamicTags[Index].Tag;
+			PreviousTag = Query.ValueTags[Index].Tag;
 		}
 
 		struct FGroupTagPair
@@ -618,14 +662,14 @@ namespace UE::Editor::DataStorage
 		};
 
 		TArray<FGroupTagPair> GroupTagPairsTemp;
-		GroupTagPairsTemp.Reserve(Query.DynamicTags.Num());
-		for (int32 Index = 0, End = Query.DynamicTags.Num(); Index < End; ++Index)
+		GroupTagPairsTemp.Reserve(Query.ValueTags.Num());
+		for (int32 Index = 0, End = Query.ValueTags.Num(); Index < End; ++Index)
 		{
-			const UScriptStruct* ColumnType = Environment.GenerateColumnType(Query.DynamicTags[Index].Tag);
+			const UScriptStruct* ColumnType = Environment.GenerateColumnType(Query.ValueTags[Index].Tag);
 			GroupTagPairsTemp.Emplace(FGroupTagPair
 				{
 					.ColumnType = ColumnType,
-					.Value = Query.DynamicTags[Index].MatchValue
+					.Value = Query.ValueTags[Index].MatchValue
 				});
 		}
 
@@ -644,9 +688,7 @@ namespace UE::Editor::DataStorage
 
 				if (SharedFragmentData)
 				{
-					// TODO: Use reflection / knowledge of key offset to read the tag and compare it
-					// For now... we know it is at offset 0
-					const FDynamicTagColumn* TagOverlay = static_cast<const FDynamicTagColumn*>(SharedFragmentData);
+					const FValueTagColumn* TagOverlay = static_cast<const FValueTagColumn*>(SharedFragmentData);
 					// NAME_None will match any presence of the shared fragment
 					// otherwise, match the specific tag only
 					if (GroupTagPair.Value != NAME_None && TagOverlay->Value != GroupTagPair.Value)
diff --git a/Engine/Plugins/Experimental/EditorDataStorage/Source/TedsCore/Private/Queries/TypedElementExtendedQueryStore.h b/Engine/Plugins/Experimental/EditorDataStorage/Source/TedsCore/Private/Queries/TypedElementExtendedQueryStore.h
--- a/Engine/Plugins/Experimental/EditorDataStorage/Source/TedsCore/Private/Queries/TypedElementExtendedQueryStore.h
+++ b/Engine/Plugins/Experimental/EditorDataStorage/Source/TedsCore/Private/Queries/TypedElementExtendedQueryStore.h
@@ -165,6 +165,7 @@ namespace UE::Editor::DataStorage
 			CallbackReference Callback);
 
 		FMassEntityQuery& SetupNativeQuery(ITypedElementDataStorageInterface::FQueryDescription& Query, FExtendedQuery& StoredQuery);
+		bool SetupDynamicColumns(ITypedElementDataStorageInterface::FQueryDescription& Query, FEnvironment& Environment);
 		bool SetupSelectedColumns(ITypedElementDataStorageInterface::FQueryDescription& Query, FMassEntityQuery& NativeQuery);
 		bool SetupConditions(ITypedElementDataStorageInterface::FQueryDescription& Query, FMassEntityQuery& NativeQuery);
 		bool SetupChunkFilters(Handle QueryHandle, ITypedElementDataStorageInterface::FQueryDescription& Query, FEnvironment& Environment, FMassEntityQuery& NativeQuery);
diff --git a/Engine/Plugins/Experimental/EditorDataStorage/Source/TedsCore/Private/TypedElementDatabase.cpp b/Engine/Plugins/Experimental/EditorDataStorage/Source/TedsCore/Private/TypedElementDatabase.cpp
--- a/Engine/Plugins/Experimental/EditorDataStorage/Source/TedsCore/Private/TypedElementDatabase.cpp
+++ b/Engine/Plugins/Experimental/EditorDataStorage/Source/TedsCore/Private/TypedElementDatabase.cpp
@@ -956,6 +956,16 @@ bool UEditorDataStorage::MatchesColumns(TypedElementDataStorage::RowHandle Row,
 	return false;
 }
 
+const UScriptStruct* UEditorDataStorage::FindDynamicColumn(const UE::Editor::DataStorage::FDynamicColumnDescription& Description) const
+{
+	return Environment->FindDynamicColumn(*Description.TemplateType, Description.Identifier);
+}
+
+const UScriptStruct* UEditorDataStorage::GenerateDynamicColumn(const UE::Editor::DataStorage::FDynamicColumnDescription& Description)
+{
+	return Environment->GenerateDynamicColumn(*Description.TemplateType, Description.Identifier);
+}
+
 void UEditorDataStorage::RegisterTickGroup(
 	FName GroupName, EQueryTickPhase Phase, FName BeforeGroup, FName AfterGroup, TypedElementDataStorage::EExecutionMode ExecutionMode)
 {
diff --git a/Engine/Plugins/Experimental/EditorDataStorage/Source/TedsCore/Private/TypedElementDatabase.h b/Engine/Plugins/Experimental/EditorDataStorage/Source/TedsCore/Private/TypedElementDatabase.h
--- a/Engine/Plugins/Experimental/EditorDataStorage/Source/TedsCore/Private/TypedElementDatabase.h
+++ b/Engine/Plugins/Experimental/EditorDataStorage/Source/TedsCore/Private/TypedElementDatabase.h
@@ -133,6 +133,9 @@ public:
 	virtual void ListColumns(TypedElementDataStorage::RowHandle Row, ColumnListWithDataCallbackRef Callback);
 	virtual bool MatchesColumns(TypedElementDataStorage::RowHandle Row, const TypedElementDataStorage::FQueryConditions& Conditions) const override;
 
+	const UScriptStruct* FindDynamicColumn(const UE::Editor::DataStorage::FDynamicColumnDescription& Description) const override;
+	const UScriptStruct* GenerateDynamicColumn(const UE::Editor::DataStorage::FDynamicColumnDescription& Description) override;
+
 	void RegisterTickGroup(FName GroupName, EQueryTickPhase Phase, FName BeforeGroup, FName AfterGroup, TypedElementDataStorage::EExecutionMode ExecutionMode);
 	void UnregisterTickGroup(FName GroupName, EQueryTickPhase Phase);
 
diff --git a/Engine/Plugins/Experimental/EditorDataStorage/Source/TedsCore/Private/TypedElementDatabaseEnvironment.cpp b/Engine/Plugins/Experimental/EditorDataStorage/Source/TedsCore/Private/TypedElementDatabaseEnvironment.cpp
--- a/Engine/Plugins/Experimental/EditorDataStorage/Source/TedsCore/Private/TypedElementDatabaseEnvironment.cpp
+++ b/Engine/Plugins/Experimental/EditorDataStorage/Source/TedsCore/Private/TypedElementDatabaseEnvironment.cpp
@@ -91,6 +91,16 @@ namespace UE::Editor::DataStorage
 		return MassPhaseManager;
 	}
 
+	const UScriptStruct* FEnvironment::FindDynamicColumn(const UScriptStruct& Template, const FName& Identifier)
+	{
+		return DynamicColumnGenerator.FindColumn(Template, Identifier);
+	}
+
+	const UScriptStruct* FEnvironment::GenerateDynamicColumn(const UScriptStruct& Template, const FName& Identifier)
+	{
+		return DynamicColumnGenerator.GenerateColumn(Template, Identifier).Type;
+	}
+
 	FConstSharedStruct FEnvironment::GenerateDynamicTag(const FDynamicTag& Tag, const FName& Value)
 	{
 		return DynamicTagManager.GenerateDynamicTag(Tag, Value);
diff --git a/Engine/Plugins/Experimental/EditorDataStorage/Source/TedsCore/Private/TypedElementDatabaseEnvironment.h b/Engine/Plugins/Experimental/EditorDataStorage/Source/TedsCore/Private/TypedElementDatabaseEnvironment.h
--- a/Engine/Plugins/Experimental/EditorDataStorage/Source/TedsCore/Private/TypedElementDatabaseEnvironment.h
+++ b/Engine/Plugins/Experimental/EditorDataStorage/Source/TedsCore/Private/TypedElementDatabaseEnvironment.h
@@ -44,8 +44,18 @@ namespace UE::Editor::DataStorage
 		FMassProcessingPhaseManager& GetMassPhaseManager();
 		const FMassProcessingPhaseManager& GetMassPhaseManager() const;
 
-		FConstSharedStruct GenerateDynamicTag(const FDynamicTag& Tag, const FName& Value);
+		// Finds the type information for the dynamic column
+		// Dynamic columns are specified by a template layout and a FName Identifier
+		const UScriptStruct* FindDynamicColumn(const UScriptStruct& Template, const FName& Identifier);
+		// Generates or returns an existing type for the dynamic column
+		// Dynamic columns are specified by a template layout and a FName Identifier
+		const UScriptStruct* GenerateDynamicColumn(const UScriptStruct& Template, const FName& Identifier);
+		
+		// Creates or Finds the column type associated with the value tag
 		const UScriptStruct* GenerateColumnType(const FDynamicTag& Tag);
+		
+		// Creates an instance of a value tag
+		FConstSharedStruct GenerateDynamicTag(const FDynamicTag& Tag, const FName& Value);
 
 		void NextUpdateCycle();
 		uint64 GetUpdateCycleId() const;
diff --git a/Engine/Plugins/Experimental/EditorDataStorage/Source/TedsCore/Tests/DynamicColumnTests.cpp b/Engine/Plugins/Experimental/EditorDataStorage/Source/TedsCore/Tests/DynamicColumnTests.cpp
new file mode 100644
--- /dev/null
+++ b/Engine/Plugins/Experimental/EditorDataStorage/Source/TedsCore/Tests/DynamicColumnTests.cpp
@@ -0,0 +1,617 @@
+﻿// Copyright Epic Games, Inc. All Rights Reserved.
+
+#include "Elements/Framework/TypedElementRegistry.h"
+#include "Elements/Framework/TypedElementTestColumns.h"
+#include "Elements/Interfaces/TypedElementDataStorageInterface.h"
+#include "Debug/TypedElementDatabaseDebugTypes.h"
+#include "Misc/CoreDelegates.h"
+#if WITH_TESTS
+
+#include "Misc/AutomationTest.h"
+
+namespace UE::Editor::DataStorage::Tests
+{
+	BEGIN_DEFINE_SPEC(DynamicColumnTestFixture, "Editor.DataStorage.DynamicColumns", EAutomationTestFlags::EditorContext | EAutomationTestFlags::EngineFilter)
+		ITypedElementDataStorageInterface* TedsInterface = nullptr;
+		const FName TestTableName = TEXT("TestTable_DynamicColumnsTest");
+
+		TypedElementDataStorage::TableHandle TestTable;
+		TArray<TypedElementDataStorage::RowHandle> Rows;
+		TArray<TypedElementDataStorage::QueryHandle> QueryHandles;
+
+		TArray<FName> Identifiers;
+
+		TypedElementDataStorage::TableHandle RegisterTestTable() const
+		{
+			const TypedElementDataStorage::TableHandle Table = TedsInterface->FindTable(TestTableName);
+				
+			if (Table != TypedElementDataStorage::InvalidTableHandle)
+			{
+				return Table;
+			}
+				
+			return TedsInterface->RegisterTable(
+			{
+				FTestColumnA::StaticStruct()
+			},
+			TestTableName);
+		}
+
+		TypedElementDataStorage::RowHandle CreateTestRow(TypedElementDataStorage::TableHandle InTableHandle)
+		{
+			TypedElementDataStorage::RowHandle RowHandle = TedsInterface->AddRow(InTableHandle);
+			Rows.Add(RowHandle);
+			return RowHandle;
+		}
+
+		TypedElementDataStorage::QueryHandle RegisterQuery(TypedElementDataStorage::FQueryDescription&& Query)
+		{
+			TypedElementDataStorage::QueryHandle QueryHandle = TedsInterface->RegisterQuery(MoveTemp(Query));
+			QueryHandles.Add(QueryHandle);
+			return QueryHandle;
+		}
+
+		// Wait at least FrameCount gamethread ticks
+		void WaitFrames(int32 FrameCount = 2)
+		{
+			if(ensureMsgf(!IsInGameThread(), TEXT("Must not be on async thread")))
+			{
+				TPromise<void> Promise;
+				TFuture<void> Future = Promise.GetFuture();
+				bool bFutureSet = false;
+
+				FDelegateHandle DelegateHandle;
+
+				Async(EAsyncExecution::TaskGraphMainThread, [&Promise, FrameCount, &DelegateHandle, &bFutureSet]()
+				{
+					const uint64 LastFrameCount = GFrameCounter;
+					const uint64 UnblockAt = LastFrameCount + FrameCount;
+					DelegateHandle = FCoreDelegates::OnEndFrame.AddLambda([UnblockAt, &Promise, &bFutureSet]()
+					{
+						const uint64 CurrentFrameCount = GFrameCounter;
+						if (UnblockAt < CurrentFrameCount)
+						{
+							// Handle case where test thread hasn't yet unregistered the delegate
+							if (!bFutureSet)
+							{
+								Promise.SetValue();
+								bFutureSet = true;
+							}
+						}
+					});
+				});
+
+				Future.Wait();
+				
+				Async(EAsyncExecution::TaskGraphMainThread, [&DelegateHandle]()
+				{
+					FCoreDelegates::OnEndFrame.Remove(DelegateHandle);
+					DelegateHandle = FDelegateHandle();
+				}).Wait();
+			}
+		}
+
+	
+	END_DEFINE_SPEC(DynamicColumnTestFixture)
+
+	void DynamicColumnTestFixture::Define()
+	{
+		BeforeEach([this]()
+		{
+			UTypedElementRegistry* TypedElementRegistry = UTypedElementRegistry::GetInstance();
+			TedsInterface = TypedElementRegistry->GetMutableDataStorage();
+			TestTable = RegisterTestTable();
+			Identifiers = {TEXT("StaticMesh"), TEXT("Animation"), TEXT("AudioClip")};
+			TestTrue("", TedsInterface != nullptr);
+		});
+			
+		Describe("", [this]
+		{
+			It(TEXT("Tags"), EAsyncExecution::ThreadPool, [this]()
+			{
+				// Add dynamic columns that are actually tags (ie. dataless)
+				Async(EAsyncExecution::TaskGraphMainThread,[this]()
+				{
+					for (int32 Index = 0; Index <= 2; ++Index)
+					{
+						CreateTestRow(TestTable);
+					}
+
+					TedsInterface->AddColumn<FTestDynamicTag>(Rows[0], Identifiers[0]);
+					TedsInterface->AddColumn<FTestDynamicTag>(Rows[0], Identifiers[1]);
+
+					TedsInterface->AddColumn<FTestDynamicTag>(Rows[1], Identifiers[0]);
+
+					TedsInterface->AddColumn<FTestDynamicTag>(Rows[2], Identifiers[1]);
+
+					// Check they were added
+					// Note: There is no HasColumn function for syntactic sugar to get a dynamic column type
+					TArray<const UScriptStruct*> DynamicTagTypes;
+					
+					DynamicTagTypes.Emplace(TedsInterface->FindDynamicColumn(FDynamicColumnDescription
+						{
+							.TemplateType = FTestDynamicTag::StaticStruct(),
+							.Identifier = Identifiers[0]
+						}));
+					DynamicTagTypes.Emplace(TedsInterface->FindDynamicColumn(FDynamicColumnDescription
+						{
+							.TemplateType = FTestDynamicTag::StaticStruct(),
+							.Identifier = Identifiers[1]
+						}));
+					TestTrue("Expected columns not found", TedsInterface->HasColumns(Rows[0], MakeConstArrayView({DynamicTagTypes[0], DynamicTagTypes[1]})));
+					
+					TestTrue("Expected columns not found", TedsInterface->HasColumns(Rows[1], MakeConstArrayView({DynamicTagTypes[0]})));
+					TestFalse("Unexpected columns found", TedsInterface->HasColumns(Rows[1], MakeConstArrayView({DynamicTagTypes[1]})));
+
+					TestFalse("Unexpected columns found", TedsInterface->HasColumns(Rows[2], MakeConstArrayView({DynamicTagTypes[0]})));
+					TestTrue("Expected columns not found", TedsInterface->HasColumns(Rows[2], MakeConstArrayView({DynamicTagTypes[1]})));
+				}).Wait();
+				
+				// Direct Query
+				{
+					Async(EAsyncExecution::TaskGraphMainThread,[this]()
+					{
+						using namespace TypedElementQueryBuilder;
+						using namespace UE::Editor::DataStorage;
+
+						TArray<TypedElementDataStorage::RowHandle> RowsToMatch;
+						TBitArray WasMatched;
+						auto SetExpectedMatches = [&RowsToMatch, &WasMatched](TConstArrayView<TypedElementDataStorage::RowHandle> Expectation)
+						{
+							RowsToMatch.Empty(RowsToMatch.Num());
+							RowsToMatch.Append(Expectation);
+							WasMatched.Reset();
+							WasMatched.SetNum(Expectation.Num(), false);
+						};
+						auto GetMatchCount = [&WasMatched]()
+						{
+							return WasMatched.CountSetBits();
+						};
+						auto Callback = CreateDirectQueryCallbackBinding([this, &RowsToMatch, &WasMatched] (TypedElementDataStorage::IDirectQueryContext& Context, const TypedElementDataStorage::RowHandle* CallbackRows)
+						{
+							TConstArrayView<TypedElementDataStorage::RowHandle> Rows = MakeConstArrayView(CallbackRows, Context.GetRowCount());
+							for (TypedElementDataStorage::RowHandle Row : Rows)
+							{
+								int32 Index = RowsToMatch.Find(Row);
+								TestTrue(TEXT("Returned row in query is within expected match array"), Index != INDEX_NONE);
+								if (Index != INDEX_NONE)
+								{
+									TestFalse("Returned row was not duplicated in the callback", WasMatched[Index]);
+									WasMatched[Index] = true;
+								}
+							}
+						});					
+
+						// Should match Rows[0]
+						{
+							TypedElementDataStorage::QueryHandle Query = RegisterQuery(
+								Select().
+								Where().
+									All<FTestDynamicTag>(Identifiers[0]).
+									All<FTestDynamicTag>(Identifiers[1]).
+								Compile());
+						
+							SetExpectedMatches({Rows[0]});
+							TypedElementDataStorage::FQueryResult Result = TedsInterface->RunQuery(Query, Callback);
+							TestEqual("Match Row[0]", Result.Count, GetMatchCount());
+						}
+						{
+							// Should match Rows 0 and 1
+							TypedElementDataStorage::QueryHandle Query = RegisterQuery(
+								Select().
+									Where().
+										All<FTestDynamicTag>(Identifiers[0]).
+								Compile());
+							SetExpectedMatches({Rows[0], Rows[1]});
+							TypedElementDataStorage::FQueryResult Result = TedsInterface->RunQuery(Query, Callback);
+							TestEqual("Match Row[0] and Row[1]", Result.Count, GetMatchCount());
+						}
+						{
+							// Should match Rows 1
+							TypedElementDataStorage::QueryHandle Query = RegisterQuery(
+								Select().
+									Where().
+										All<FTestDynamicTag>(Identifiers[0]).
+										None<FTestDynamicTag>(Identifiers[1]).
+								Compile());
+							SetExpectedMatches({Rows[1]});
+							TypedElementDataStorage::FQueryResult Result = TedsInterface->RunQuery(Query, Callback);
+							TestEqual("Match Row[1]", Result.Count, GetMatchCount());
+						}
+						{
+							// Should match Rows 0 and 2
+							TypedElementDataStorage::QueryHandle Query = RegisterQuery(
+								Select().
+									Where().
+										All<FTestDynamicTag>(Identifiers[1]).
+								Compile());
+							SetExpectedMatches({Rows[0], Rows[2]});
+							TypedElementDataStorage::FQueryResult Result = TedsInterface->RunQuery(Query, Callback);
+							TestEqual("Match Row[0] and Row[2]", Result.Count, GetMatchCount());
+						}
+						{
+							// Should match Rows 0 and 2
+							TypedElementDataStorage::QueryHandle Query = RegisterQuery(
+								Select().
+									Where().
+										None<FTestDynamicTag>(Identifiers[0]).
+										All<FTestDynamicTag>(Identifiers[1]).
+								Compile());
+							SetExpectedMatches({Rows[2]});
+							TypedElementDataStorage::FQueryResult Result = TedsInterface->RunQuery(Query, Callback);
+							TestEqual("Match Row[2]", Result.Count, GetMatchCount());
+						}
+						{
+							// Should match Rows 0, 1 and 2
+							TypedElementDataStorage::QueryHandle Query = RegisterQuery(
+								Select().
+									Where().
+										Any<FTestDynamicTag>(Identifiers[0]).
+										Any<FTestDynamicTag>(Identifiers[1]).
+								Compile());
+							SetExpectedMatches({Rows[0], Rows[1], Rows[2]});
+							TypedElementDataStorage::FQueryResult Result = TedsInterface->RunQuery(Query, Callback);
+							TestEqual("Match All Rows", Result.Count, GetMatchCount());
+						}
+					}).Wait();
+				}
+
+				// Processor Query
+				{
+					TArray<TypedElementDataStorage::RowHandle> RowsToMatch;
+					TArray<int32> MatchCount;
+					int32 UnexpectedRowCount;
+					auto SetExpectedMatches = [&RowsToMatch, &MatchCount, &UnexpectedRowCount](TConstArrayView<TypedElementDataStorage::RowHandle> Expectation)
+					{
+						RowsToMatch.Empty(RowsToMatch.Num());
+						RowsToMatch.Append(Expectation);
+						MatchCount.Empty(MatchCount.Num());
+						MatchCount.SetNum(Expectation.Num());
+						for (int32 I = 0; I < MatchCount.Num(); ++I)
+						{
+							MatchCount[I] = 0;
+						}
+						UnexpectedRowCount = 0;
+					};
+
+					TArray<FName> ActivationKeys;
+					TArray<TArray<TypedElementDataStorage::RowHandle>> QueryExpectedMatchRows;
+					
+					Async(EAsyncExecution::TaskGraphMainThread,[this, &RowsToMatch, &MatchCount, &UnexpectedRowCount, &ActivationKeys, &QueryExpectedMatchRows]()
+					{
+						auto Callback = [this, &RowsToMatch, &MatchCount, &UnexpectedRowCount] (TypedElementDataStorage::IQueryContext& Context, const TypedElementDataStorage::RowHandle* CallbackRows)
+						{
+							TConstArrayView<TypedElementDataStorage::RowHandle> Rows = MakeConstArrayView(CallbackRows, Context.GetRowCount());
+							for (auto Row : Rows)
+							{
+								int32 Index = RowsToMatch.Find(Row);
+								TestTrue(TEXT("Returned row in query is within expected match array"), Index != INDEX_NONE);
+								if (Index != INDEX_NONE)
+								{
+									++MatchCount[Index];
+								}
+								else
+								{
+									++UnexpectedRowCount;
+								}
+							}
+						};
+						
+						using namespace TypedElementQueryBuilder;
+						using namespace TypedElementDataStorage;
+						// Setup an activatable processor for the test
+						ActivationKeys.Emplace(TEXT("TEST: Match Row[0]"));
+						QueryExpectedMatchRows.Emplace(TArray({Rows[0]}));
+						RegisterQuery(
+							Select(
+								ActivationKeys.Last(),
+								FProcessor(EQueryTickPhase::FrameEnd, TedsInterface->GetQueryTickGroupName(EQueryTickGroups::SyncDataStorageToExternal))
+									.MakeActivatable(ActivationKeys.Last()),
+									CopyTemp(Callback)
+							)
+							.Where()
+								.All<FTestDynamicTag>(Identifiers[0])
+								.All<FTestDynamicTag>(Identifiers[1])
+							.Compile());
+
+						ActivationKeys.Emplace(TEXT("TEST: Match Row[0] and Row[1]"));
+						QueryExpectedMatchRows.Emplace(TArray({Rows[0], Rows[1]}));
+						RegisterQuery(
+							Select(
+								ActivationKeys.Last(),
+								FProcessor(EQueryTickPhase::FrameEnd, TedsInterface->GetQueryTickGroupName(EQueryTickGroups::SyncDataStorageToExternal))
+									.MakeActivatable(ActivationKeys.Last()),
+									CopyTemp(Callback)
+							)
+							.Where()
+								.All<FTestDynamicTag>(Identifiers[0])
+							.Compile());
+
+						ActivationKeys.Emplace(TEXT("TEST: Match Row[1]"));
+						QueryExpectedMatchRows.Emplace(TArray({Rows[1]}));
+						RegisterQuery(
+							Select(
+								ActivationKeys.Last(),
+								FProcessor(EQueryTickPhase::FrameEnd, TedsInterface->GetQueryTickGroupName(EQueryTickGroups::SyncDataStorageToExternal))
+									.MakeActivatable(ActivationKeys.Last()),
+									CopyTemp(Callback)
+							)
+							.Where()
+								.All<FTestDynamicTag>(Identifiers[0])
+								.None<FTestDynamicTag>(Identifiers[1])
+							.Compile());
+
+						ActivationKeys.Emplace(TEXT("TEST: Match Row[0] and Row[2]"));
+						QueryExpectedMatchRows.Emplace(TArray({Rows[0], Rows[2]}));
+						RegisterQuery(
+							Select(
+								ActivationKeys.Last(),
+								FProcessor(EQueryTickPhase::FrameEnd, TedsInterface->GetQueryTickGroupName(EQueryTickGroups::SyncDataStorageToExternal))
+									.MakeActivatable(ActivationKeys.Last()),
+									CopyTemp(Callback)
+							)
+							.Where()
+								.All<FTestDynamicTag>(Identifiers[1])
+							.Compile());
+
+						ActivationKeys.Emplace(TEXT("TEST: Match Row[2]"));
+						QueryExpectedMatchRows.Emplace(TArray({Rows[2]}));
+						RegisterQuery(
+							Select(
+								ActivationKeys.Last(),
+								FProcessor(EQueryTickPhase::FrameEnd, TedsInterface->GetQueryTickGroupName(EQueryTickGroups::SyncDataStorageToExternal))
+									.MakeActivatable(ActivationKeys.Last()),
+									CopyTemp(Callback)
+							)
+							.Where()
+								.None<FTestDynamicTag>(Identifiers[0])
+								.All<FTestDynamicTag>(Identifiers[1])
+							.Compile());
+
+						ActivationKeys.Emplace(TEXT("TEST: Match all rows"));
+						QueryExpectedMatchRows.Emplace(TArray({Rows[0], Rows[1], Rows[2]}));
+						RegisterQuery(
+							Select(
+								ActivationKeys.Last(),
+								FProcessor(EQueryTickPhase::FrameEnd, TedsInterface->GetQueryTickGroupName(EQueryTickGroups::SyncDataStorageToExternal))
+									.MakeActivatable(ActivationKeys.Last()),
+									CopyTemp(Callback)
+							)
+							.Where()
+								.Any<FTestDynamicTag>(Identifiers[0])
+								.Any<FTestDynamicTag>(Identifiers[1])
+							.Compile());
+					}).Wait();
+
+					for (int32 TestIndex = 0; TestIndex < ActivationKeys.Num(); ++TestIndex)
+					{
+						SetExpectedMatches(QueryExpectedMatchRows[TestIndex]);
+						// Kick off the activation
+						TFuture<void> Future = Async(EAsyncExecution::TaskGraphMainThread, [this, &ActivationKeys, TestIndex]()
+						{
+							TedsInterface->ActivateQueries(ActivationKeys[TestIndex]);
+						});
+						Future.Wait();
+
+						// Wait for execution to occur
+						WaitFrames();
+						
+						// Check results
+						TStringBuilder<256> TestMessage;
+						TestMessage.Appendf(TEXT("'%s': Check that no rows processed that weren't expected"), *ActivationKeys[TestIndex].ToString());
+						TestEqual(TestMessage.ToString(), UnexpectedRowCount, static_cast<int32>(0));
+
+						TestMessage.Reset();
+						for (int32 RowIndex = 0; RowIndex < MatchCount.Num(); ++RowIndex)
+						{
+							TestMessage.Appendf(TEXT("'%s': Check row matched 1 times"), *ActivationKeys[TestIndex].ToString());
+							TestEqual(TestMessage.ToString(), MatchCount[RowIndex], 1);
+						}
+					}
+				}
+
+				WaitFrames();
+			});
+
+			It(TEXT("Columns"), EAsyncExecution::ThreadPool, [this]()
+			{
+				// Add dynamic columns that have data
+				Async(EAsyncExecution::TaskGraphMainThread,[this]()
+				{
+					for (int32 Index = 0; Index <= 2; ++Index)
+					{
+						CreateTestRow(TestTable);
+					}
+
+					TedsInterface->AddColumn(Rows[0], Identifiers[0], FTestDynamicColumn
+					{
+						.IntArray = {1, 2, 3},
+					});
+					
+					TedsInterface->AddColumn(Rows[0], Identifiers[1], FTestDynamicColumn
+					{
+						.IntArray = {10, 11, 12, 13},
+					});
+
+					TedsInterface->AddColumn(Rows[1], Identifiers[0], FTestDynamicColumn
+					{
+						.IntArray = {14, 15, 16},
+					});
+
+					TedsInterface->AddColumn(Rows[2], Identifiers[1], FTestDynamicColumn
+					{
+						.IntArray = {11, 22, 33, 44},
+					});
+
+					// Check they were added
+					// Note: There is no HasColumn function for syntactic sugar to get a dynamic column type
+					TArray<const UScriptStruct*> DynamicTagTypes;
+								
+					DynamicTagTypes.Emplace(TedsInterface->FindDynamicColumn(FDynamicColumnDescription
+						{
+							.TemplateType = FTestDynamicColumn::StaticStruct(),
+							.Identifier = Identifiers[0]
+						}));
+					DynamicTagTypes.Emplace(TedsInterface->FindDynamicColumn(FDynamicColumnDescription
+						{
+							.TemplateType = FTestDynamicColumn::StaticStruct(),
+							.Identifier = Identifiers[1]
+						}));
+					TestTrue("Expected columns not found", TedsInterface->HasColumns(Rows[0], MakeConstArrayView({DynamicTagTypes[0], DynamicTagTypes[1]})));
+								
+					TestTrue("Expected columns not found", TedsInterface->HasColumns(Rows[1], MakeConstArrayView({DynamicTagTypes[0]})));
+					TestFalse("Unexpected columns found", TedsInterface->HasColumns(Rows[1], MakeConstArrayView({DynamicTagTypes[1]})));
+
+					TestFalse("Unexpected columns found", TedsInterface->HasColumns(Rows[2], MakeConstArrayView({DynamicTagTypes[0]})));
+					TestTrue("Expected columns not found", TedsInterface->HasColumns(Rows[2], MakeConstArrayView({DynamicTagTypes[1]})));
+				}).Wait();
+				
+				// Processor Query
+				{
+					TArray<FName> ActivationKeys;
+					
+					struct FExpectation
+					{
+						TypedElementDataStorage::RowHandle Row;
+						TArray<int32> ColumnValues;
+
+						int32 MatchCount = 0;
+						bool bValuesMatch = false;
+					};
+					
+					TArray<TArray<FExpectation>> AllTestExpectations;
+
+					// Expectation for the current test
+					int32 UnexpectedRowCount;
+					
+					Async(EAsyncExecution::TaskGraphMainThread,[this, &ActivationKeys, &AllTestExpectations, &UnexpectedRowCount]()
+					{
+						using namespace TypedElementQueryBuilder;
+						using namespace TypedElementDataStorage;
+
+						auto RunTest = [this, &AllTestExpectations, &UnexpectedRowCount](IQueryContext& Context, const RowHandle* Rows, int32 TestIndex)
+						{
+							TConstArrayView<RowHandle> RowView = MakeConstArrayView(Rows, Context.GetRowCount());
+							TConstArrayView<FTestDynamicColumn> ColumnView = MakeConstArrayView(
+								Context.GetColumn<FTestDynamicColumn>(Identifiers[0]),
+								Context.GetRowCount());
+
+							TArray<FExpectation>& Expectations = AllTestExpectations[TestIndex];
+							for (int32 RowIndex = 0, End = Context.GetRowCount(); RowIndex < End; ++RowIndex)
+							{
+								RowHandle Row = RowView[RowIndex];
+																		
+								FExpectation* Expectation = Expectations.FindByPredicate([Row](const FExpectation& Expectation)
+								{
+									return Expectation.Row == Row;
+								});
+								if (Expectation)
+								{
+									Expectation->MatchCount++;
+									// Check values
+									const FTestDynamicColumn& Column = ColumnView[RowIndex];
+									if (Column.IntArray.Num() == Expectation->ColumnValues.Num())
+									{
+										bool bAllMatch = true;
+										for (int32 ColumnDataIndex = 0; ColumnDataIndex < Column.IntArray.Num(); ++ColumnDataIndex)
+										{
+											if (Column.IntArray[ColumnDataIndex] != Expectation->ColumnValues[ColumnDataIndex])
+											{
+												bAllMatch = false;
+												break;
+											}
+										}
+										Expectation->bValuesMatch = bAllMatch;
+									}
+								}
+								else
+								{
+									++UnexpectedRowCount;
+								}
+							}
+						};
+						
+						// Setup an activatable processor for the test
+						ActivationKeys.Emplace(TEXT("TEST: Match Row[0]"));
+						AllTestExpectations.Add(
+							{
+								{
+									.Row = Rows[0],
+									.ColumnValues = {1, 2, 3}
+								},
+								{
+									.Row = Rows[1],
+									.ColumnValues = {14, 15, 16}
+								}
+							});
+
+						RegisterQuery(
+							Select(
+								ActivationKeys.Last(),
+								FProcessor(EQueryTickPhase::FrameEnd, TedsInterface->GetQueryTickGroupName(EQueryTickGroups::SyncDataStorageToExternal))
+									.MakeActivatable(ActivationKeys.Last()),
+									[this, TestIndex = AllTestExpectations.Num() - 1, &RunTest](IQueryContext& Context, const RowHandle* Rows)
+									{
+										RunTest(Context, Rows, TestIndex);
+									}
+							)
+							.ReadOnly<FTestDynamicColumn>(Identifiers[0])
+							.Compile());
+					});
+
+					for (int32 TestIndex = 0; TestIndex < ActivationKeys.Num(); ++TestIndex)
+					{
+						// Kick off the activation
+						UnexpectedRowCount = 0;
+						TFuture<void> Future = Async(EAsyncExecution::TaskGraphMainThread,
+							[this, &ActivationKeys, TestIndex]()
+							{
+								TedsInterface->ActivateQueries(ActivationKeys[TestIndex]);
+							});
+						Future.Wait();
+						
+						WaitFrames();
+						
+						// Check results
+						TStringBuilder<256> TestMessage;
+						TestMessage.Appendf(TEXT("'%s': Check that no rows processed that weren't expected"), *ActivationKeys[TestIndex].ToString());
+						TestEqual(TestMessage.ToString(), UnexpectedRowCount, static_cast<int32>(0));
+
+						TestMessage.Reset();
+						
+						TArray<FExpectation>& TestExpectations = AllTestExpectations[TestIndex];
+						for (int32 RowIndex = 0; RowIndex < TestExpectations.Num(); ++RowIndex)
+						{
+							TestMessage.Appendf(TEXT("'%s': Each row matched 1 times"), *ActivationKeys[TestIndex].ToString());
+							TestEqual(TestMessage.ToString(), TestExpectations[RowIndex].MatchCount, 1);
+
+							TestMessage.Reset();
+							TestMessage.Appendf(TEXT("'%s': Column values are expected"));
+							TestEqual(TestMessage.ToString(), TestExpectations[RowIndex].bValuesMatch, true);
+						}
+					}
+				}
+			});
+		});
+		
+		AfterEach([this]()
+		{
+			Identifiers.Empty(Identifiers.Num());
+			for (TypedElementDataStorage::RowHandle Row : Rows)
+			{
+				TedsInterface->RemoveRow(Row);
+			}
+			Rows.Empty(Rows.Num());
+
+			for (TypedElementDataStorage::QueryHandle QueryHandle : QueryHandles)
+			{
+				TedsInterface->UnregisterQuery(QueryHandle);
+			}
+			QueryHandles.Empty(QueryHandles.Num());
+
+			TedsInterface = nullptr;
+		});
+	}
+}
+
+#endif // WITH_TESTS
\ No newline at end of file
diff --git a/Engine/Source/Runtime/TypedElementFramework/Private/Elements/Common/TypedElementCommonTypes.cpp b/Engine/Source/Runtime/TypedElementFramework/Private/Elements/Common/TypedElementCommonTypes.cpp
new file mode 100644
--- /dev/null
+++ b/Engine/Source/Runtime/TypedElementFramework/Private/Elements/Common/TypedElementCommonTypes.cpp
@@ -0,0 +1,28 @@
+// Copyright Epic Games, Inc. All Rights Reserved.
+
+#include "Elements/Common/TypedElementCommonTypes.h"
+
+namespace UE
+{
+	namespace Editor::DataStorage
+	{
+		const FName& FDynamicTag::GetName() const
+		{
+			return Name;
+		}
+		
+		FDynamicTag::FDynamicTag(const FName& InTypeName)
+			: Name(InTypeName)
+		{}
+		
+		uint32 GetTypeHash(const FDynamicTag& InName)
+		{
+			return GetTypeHash(InName.Name);
+		}
+
+		uint32 GetTypeHash(const FDynamicColumnDescription& Descriptor)
+		{
+			return HashCombineFast(PointerHash(Descriptor.TemplateType), GetTypeHash(Descriptor.Identifier));
+		}
+	}
+}
diff --git a/Engine/Source/Runtime/TypedElementFramework/Private/Elements/Framework/TypedElementQueryBuilder.cpp b/Engine/Source/Runtime/TypedElementFramework/Private/Elements/Framework/TypedElementQueryBuilder.cpp
--- a/Engine/Source/Runtime/TypedElementFramework/Private/Elements/Framework/TypedElementQueryBuilder.cpp
+++ b/Engine/Source/Runtime/TypedElementFramework/Private/Elements/Framework/TypedElementQueryBuilder.cpp
@@ -145,7 +145,7 @@ namespace TypedElementQueryBuilder
 
 	FSimpleQuery& FSimpleQuery::All(const UE::Editor::DataStorage::FDynamicTag& Tag, const FName& Value)
 	{
-		Query->DynamicTags.Emplace(
+		Query->ValueTags.Emplace(
 			TypedElementDataStorage::FQueryDescription::FDynamicTagData
 			{
 				.Tag = Tag,
@@ -171,17 +171,12 @@ namespace TypedElementQueryBuilder
 		using namespace UE::Editor::DataStorage;
 		return All(FDynamicTag(Enum.GetFName()), ValueName);
 	}
-
-	template <>
-	FSimpleQuery& FSimpleQuery::All<UE::Editor::DataStorage::FDynamicTag>(const FName& Tag)
-	{
-		return All(UE::Editor::DataStorage::FDynamicTag(Tag));
-	}
-
-	template <>
-	FSimpleQuery& FSimpleQuery::All<UE::Editor::DataStorage::FDynamicTag>(const FName& Tag, const FName& Value)
+	
+	FSimpleQuery& FSimpleQuery::All(const UE::Editor::DataStorage::FDynamicColumnDescription& Description)
 	{
-		return All(UE::Editor::DataStorage::FDynamicTag(Tag), Value);
+		Query->DynamicConditionDescriptions.Add(Description);
+		Query->DynamicConditionOperations.Add(ITypedElementDataStorageInterface::FQueryDescription::EOperatorType::SimpleAll);
+		return *this;
 	}
 
 	FSimpleQuery& FSimpleQuery::All(const UE::Editor::DataStorage::FDynamicTag& Tag)
@@ -212,6 +207,13 @@ namespace TypedElementQueryBuilder
 		return *this;
 	}
 
+	FSimpleQuery& FSimpleQuery::Any(const UE::Editor::DataStorage::FDynamicColumnDescription& Description)
+	{
+		Query->DynamicConditionDescriptions.Add(Description);
+		Query->DynamicConditionOperations.Add(TypedElementDataStorage::FQueryDescription::EOperatorType::SimpleAny);
+		return *this;
+	}
+
 	FSimpleQuery& FSimpleQuery::None(const UScriptStruct* Target)
 	{
 		if (Target)
@@ -235,6 +237,13 @@ namespace TypedElementQueryBuilder
 		return *this;
 	}
 
+	FSimpleQuery& FSimpleQuery::None(const UE::Editor::DataStorage::FDynamicColumnDescription& Description)
+	{
+		Query->DynamicConditionDescriptions.Add(Description);
+		Query->DynamicConditionOperations.Add(TypedElementDataStorage::FQueryDescription::EOperatorType::SimpleNone);
+		return *this;
+	}
+
 	FDependency FSimpleQuery::DependsOn()
 	{
 		return FDependency{ Query };
@@ -253,6 +262,11 @@ namespace TypedElementQueryBuilder
 		Query->SelectionMetaData.Shrink();
 		Query->ConditionTypes.Shrink();
 		Query->ConditionOperators.Shrink();
+		Query->DynamicConditionDescriptions.Shrink();
+		Query->DynamicConditionOperations.Shrink();
+		Query->DynamicSelectionAccessTypes.Shrink();
+		Query->DynamicSelectionMetaData.Shrink();
+		Query->DynamicSelectionTypes.Shrink();
 		Query->DependencyTypes.Shrink();
 		Query->DependencyFlags.Shrink();
 		Query->CachedDependencies.Shrink();
@@ -408,6 +422,14 @@ namespace TypedElementQueryBuilder
 		return *this;
 	}
 
+	Select& Select::ReadOnly(const UE::Editor::DataStorage::FDynamicColumnDescription& Description)
+	{
+		Query.DynamicSelectionTypes.Emplace(Description);
+		Query.DynamicSelectionAccessTypes.Emplace(ITypedElementDataStorageInterface::EQueryAccessType::ReadOnly);
+		Query.DynamicSelectionMetaData.Emplace(TypedElementDataStorage::FColumnMetaData::EFlags::None);
+		return *this;
+	}
+
 	Select& Select::ReadOnly(const UScriptStruct* Target, EOptional Optional)
 	{
 		checkf(Target, TEXT("The Select section in the Typed Elements query builder doesn't support nullptrs as Read-Only input."));
@@ -457,6 +479,14 @@ namespace TypedElementQueryBuilder
 		return *this;
 	}
 
+	Select& Select::ReadWrite(const UE::Editor::DataStorage::FDynamicColumnDescription& Descirption)
+	{
+		Query.DynamicSelectionTypes.Emplace(Descirption);
+		Query.DynamicSelectionAccessTypes.Emplace(ITypedElementDataStorageInterface::EQueryAccessType::ReadWrite);
+		Query.DynamicSelectionMetaData.Emplace(TypedElementDataStorage::FColumnMetaData::EFlags::IsMutable);
+		return *this;
+	}
+
 	FSimpleQuery Select::Where()
 	{
 		return FSimpleQuery{ &Query };
diff --git a/Engine/Source/Runtime/TypedElementFramework/Public/Elements/Common/TypedElementCommonTypes.h b/Engine/Source/Runtime/TypedElementFramework/Public/Elements/Common/TypedElementCommonTypes.h
--- a/Engine/Source/Runtime/TypedElementFramework/Public/Elements/Common/TypedElementCommonTypes.h
+++ b/Engine/Source/Runtime/TypedElementFramework/Public/Elements/Common/TypedElementCommonTypes.h
@@ -56,32 +56,26 @@ namespace UE
 		class FDynamicTag
 		{
 		public:
-			explicit FDynamicTag(const FName& InName);
+			TYPEDELEMENTFRAMEWORK_API explicit FDynamicTag(const FName& InName);
 			
-			const FName& GetName() const;
-			bool operator==(const FDynamicTag& Other) const;
+			TYPEDELEMENTFRAMEWORK_API const FName& GetName() const;
+			TYPEDELEMENTFRAMEWORK_API bool operator==(const FDynamicTag& Other) const = default;
 		private:
-			friend uint32 GetTypeHash(const FDynamicTag& InName)
-			{
-				return GetTypeHash(InName.Name);
-			}
+			TYPEDELEMENTFRAMEWORK_API friend uint32 GetTypeHash(const FDynamicTag& InName);
 			FName Name;
 		};
 
-		inline FDynamicTag::FDynamicTag(const FName& InTypeName)
-			: Name(InTypeName)
-		{}
-
-		inline const FName& FDynamicTag::GetName() const
-		{
-			return Name;
-		}
+		template<typename T>
+		concept TValueTagType = std::is_same_v<T, FDynamicTag>;
 
-		inline bool FDynamicTag::operator==(const FDynamicTag& Other) const
+		struct FDynamicColumnDescription
 		{
-			return Other.Name == Name;
-		}
+			const UScriptStruct* TemplateType;
+			FName Identifier;
 
+			TYPEDELEMENTFRAMEWORK_API friend uint32 GetTypeHash(const FDynamicColumnDescription& Descriptor);
+			TYPEDELEMENTFRAMEWORK_API bool operator==(const FDynamicColumnDescription&) const = default;
+		};
 		// Standard callbacks.
 
 		using RowCreationCallbackRef = TFunctionRef<void(TypedElementDataStorage::RowHandle Row)>;
diff --git a/Engine/Source/Runtime/TypedElementFramework/Public/Elements/Common/TypedElementQueryDescription.h b/Engine/Source/Runtime/TypedElementFramework/Public/Elements/Common/TypedElementQueryDescription.h
--- a/Engine/Source/Runtime/TypedElementFramework/Public/Elements/Common/TypedElementQueryDescription.h
+++ b/Engine/Source/Runtime/TypedElementFramework/Public/Elements/Common/TypedElementQueryDescription.h
@@ -97,7 +97,14 @@ namespace TypedElementDataStorage
 		TArray<EOperatorType, TInlineAllocator<NumInlineConditions>> ConditionTypes;
 		TArray<FOperator, TInlineAllocator<NumInlineConditions>> ConditionOperators;
 
-		TArray<FDynamicTagData> DynamicTags;
+		TArray<UE::Editor::DataStorage::FDynamicColumnDescription, TInlineAllocator<NumInlineSelections>> DynamicSelectionTypes;
+		TArray<EQueryAccessType, TInlineAllocator<NumInlineSelections>> DynamicSelectionAccessTypes;
+		TArray<FColumnMetaData::EFlags, TInlineAllocator<NumInlineSelections>> DynamicSelectionMetaData;
+		
+		TArray<EOperatorType, TInlineAllocator<NumInlineConditions>> DynamicConditionOperations;
+		TArray<UE::Editor::DataStorage::FDynamicColumnDescription, TInlineAllocator<NumInlineConditions>> DynamicConditionDescriptions;
+
+		TArray<FDynamicTagData> ValueTags;
 
 		TArray<TWeakObjectPtr<const UClass>, TInlineAllocator<NumInlineDependencies>> DependencyTypes;
 		TArray<EQueryDependencyFlags, TInlineAllocator<NumInlineDependencies>> DependencyFlags;
diff --git a/Engine/Source/Runtime/TypedElementFramework/Public/Elements/Framework/TypedElementQueryBuilder.h b/Engine/Source/Runtime/TypedElementFramework/Public/Elements/Framework/TypedElementQueryBuilder.h
--- a/Engine/Source/Runtime/TypedElementFramework/Public/Elements/Framework/TypedElementQueryBuilder.h
+++ b/Engine/Source/Runtime/TypedElementFramework/Public/Elements/Framework/TypedElementQueryBuilder.h
@@ -161,11 +161,35 @@ namespace TypedElementQueryBuilder
 		FSimpleQuery& Any();
 		TYPEDELEMENTFRAMEWORK_API FSimpleQuery& Any(const UScriptStruct* Target);
 		TYPEDELEMENTFRAMEWORK_API FSimpleQuery& Any(TConstArrayView<const UScriptStruct*> Targets);
+
+		// Dynamic Column Support
+		TYPEDELEMENTFRAMEWORK_API FSimpleQuery& Any(const UE::Editor::DataStorage::FDynamicColumnDescription& Description);
+		template<UE::Editor::DataStorage::TColumnType T>
+		FSimpleQuery& Any(const FName&);
+
+		// Dynamic Tags not yet supported for Any
+		template<UE::Editor::DataStorage::TValueTagType>
+		FSimpleQuery& Any(const FName&) = delete;
+		template<UE::Editor::DataStorage::TValueTagType>
+		FSimpleQuery& Any(const FName&, const FName&) = delete;
+		
 		template<UE::Editor::DataStorage::TColumnType... TargetTypes>
 		FSimpleQuery& None();
+		
 		TYPEDELEMENTFRAMEWORK_API FSimpleQuery& None(const UScriptStruct* Target);
 		TYPEDELEMENTFRAMEWORK_API FSimpleQuery& None(TConstArrayView<const UScriptStruct*> Targets);
 		
+		// Dynamic Tags not yet supported for None
+		template<UE::Editor::DataStorage::TValueTagType>
+		FSimpleQuery& None(const FName& Tag) = delete;
+		template<UE::Editor::DataStorage::TValueTagType>
+		FSimpleQuery& None(const FName& Tag, const FName& Value) = delete;
+
+		// Dynamic Column Support
+		TYPEDELEMENTFRAMEWORK_API FSimpleQuery& None(const UE::Editor::DataStorage::FDynamicColumnDescription& Description);
+		template<UE::Editor::DataStorage::TColumnType T>
+		FSimpleQuery& None(const FName&);
+		
 		// Dynamic Tags
 		// ============
 		// Adds a filter to the query which must match a given DynamicTag.  The value of the DynamicTag will not be checked.
@@ -177,20 +201,18 @@ namespace TypedElementQueryBuilder
 		TYPEDELEMENTFRAMEWORK_API FSimpleQuery& All(const UEnum& Enum);
 		TYPEDELEMENTFRAMEWORK_API FSimpleQuery& All(const UEnum& Enum, int64 Value);
 
-		// Disabled generic template for DynamicTags
-		template<typename... T>
-		FSimpleQuery& All(const FName&) = delete;
-		// Disabled generic template for DynamicTags
-		template<typename... T>
-		FSimpleQuery& All(const FName&, const FName&) = delete;
+		// Dynamic Column support
+		TYPEDELEMENTFRAMEWORK_API FSimpleQuery& All(const UE::Editor::DataStorage::FDynamicColumnDescription& Description);
+		template<UE::Editor::DataStorage::TColumnType T>
+		FSimpleQuery& All(const FName&);
 
 		// Adds a filter to the query which must match a given DynamicTag.  The value of the DynamicTag will not be checked.
-		template<>
-		TYPEDELEMENTFRAMEWORK_API FSimpleQuery& All<UE::Editor::DataStorage::FDynamicTag>(const FName& Tag);
+		template<UE::Editor::DataStorage::TValueTagType>
+		FSimpleQuery& All(const FName& Tag);
 		// Adds a filter to the query which must match a given DynamicTag.  The value of the DynamicTag must also match.
 		// Note: The query can only match a single value.  Multiple value queries are not supported at this time.
-		template<>
-		TYPEDELEMENTFRAMEWORK_API FSimpleQuery& All<UE::Editor::DataStorage::FDynamicTag>(const FName& Tag, const FName& Value);
+		template<UE::Editor::DataStorage::TValueTagType>
+		FSimpleQuery& All(const FName& Tag, const FName& Value);
 		
 		template<UE::Editor::DataStorage::TEnumType EnumT>
 		FSimpleQuery& All();
@@ -353,6 +375,7 @@ namespace TypedElementQueryBuilder
 		TYPEDELEMENTFRAMEWORK_API Select& ReadOnly(const UScriptStruct* Target);
 		/** Request read-only access to the listed columns. */
 		TYPEDELEMENTFRAMEWORK_API Select& ReadOnly(TConstArrayView<const UScriptStruct*> Targets);
+		TYPEDELEMENTFRAMEWORK_API Select& ReadOnly(const UE::Editor::DataStorage::FDynamicColumnDescription& Description);
 		/** 
 		 * Request read-only access to the listed columns. If optional is true read access will be given if the column is in the table but
 		 * it will not be used for finding matching tables. Columns bound with optional can not be bound to a query callback argument.
@@ -368,7 +391,9 @@ namespace TypedElementQueryBuilder
 		 * Request read-only access to the listed columns. If optional is true read access will be given if the column is in the table but
 		 * it will not be used for finding matching tables. Columns bound with optional can not be bound to a query callback argument.
 		 */
-		TYPEDELEMENTFRAMEWORK_API Select& ReadOnly(TConstArrayView<const UScriptStruct*> Targets, EOptional Optionall);
+		TYPEDELEMENTFRAMEWORK_API Select& ReadOnly(TConstArrayView<const UScriptStruct*> Targets, EOptional Optional);
+		template<UE::Editor::DataStorage::TDataColumnType Target>
+		Select& ReadOnly(const FName& Identifier);
 		/** Request read and write access to the listed columns. */
 		template<UE::Editor::DataStorage::TDataColumnType... TargetTypes>
 		Select& ReadWrite();
@@ -376,6 +401,9 @@ namespace TypedElementQueryBuilder
 		TYPEDELEMENTFRAMEWORK_API Select& ReadWrite(const UScriptStruct* Target);
 		/** Request read and write access to the listed columns. */
 		TYPEDELEMENTFRAMEWORK_API Select& ReadWrite(TConstArrayView<const UScriptStruct*> Targets);
+		TYPEDELEMENTFRAMEWORK_API Select& ReadWrite(const UE::Editor::DataStorage::FDynamicColumnDescription& Description);
+		template<UE::Editor::DataStorage::TDataColumnType Target>
+		Select& ReadWrite(const FName& Identifier);
 
 		TYPEDELEMENTFRAMEWORK_API ITypedElementDataStorageInterface::FQueryDescription&& Compile();
 		TYPEDELEMENTFRAMEWORK_API FSimpleQuery Where();
diff --git a/Engine/Source/Runtime/TypedElementFramework/Public/Elements/Framework/TypedElementQueryBuilder.inl b/Engine/Source/Runtime/TypedElementFramework/Public/Elements/Framework/TypedElementQueryBuilder.inl
--- a/Engine/Source/Runtime/TypedElementFramework/Public/Elements/Framework/TypedElementQueryBuilder.inl
+++ b/Engine/Source/Runtime/TypedElementFramework/Public/Elements/Framework/TypedElementQueryBuilder.inl
@@ -1133,6 +1133,17 @@ e.g. void(FCachedQueryContext<Subsystem1, const Subsystem2>& Context, TypedEleme
 		return *this;
 	}
 
+	template<UE::Editor::DataStorage::TDataColumnType Target>
+	Select& Select::ReadOnly(const FName& Identifier)
+	{
+		ReadOnly(UE::Editor::DataStorage::FDynamicColumnDescription
+			{
+				.TemplateType = Target::StaticStruct(),
+				.Identifier = Identifier
+			});
+		return *this;
+	}
+
 	template<TDataColumnType... TargetTypes>
 	Select& Select::ReadWrite()
 	{
@@ -1140,6 +1151,17 @@ e.g. void(FCachedQueryContext<Subsystem1, const Subsystem2>& Context, TypedEleme
 		return *this;
 	}
 
+	template<TDataColumnType Target>
+	Select& Select::ReadWrite(const FName& Identifier)
+	{
+		ReadWrite(UE::Editor::DataStorage::FDynamicColumnDescription
+			{
+				.TemplateType = Target::StaticStruct(),
+				.Identifier = Identifier
+			});
+		return *this;
+	}
+
 
 	//
 	// FSimpleQuery
@@ -1159,13 +1181,55 @@ e.g. void(FCachedQueryContext<Subsystem1, const Subsystem2>& Context, TypedEleme
 		return *this;
 	}
 
+	template <TColumnType DynamicColumnTemplate>
+	FSimpleQuery& FSimpleQuery::Any(const FName& Identifier)
+	{
+		return Any(UE::Editor::DataStorage::FDynamicColumnDescription
+		{
+			.TemplateType = DynamicColumnTemplate::StaticStruct(),
+			.Identifier = Identifier
+		});
+	}
+
 	template<TColumnType... TargetTypes>
 	FSimpleQuery& FSimpleQuery::None()
 	{
 		None({ TargetTypes::StaticStruct()... });
 		return *this;
 	}
+
+	template<UE::Editor::DataStorage::TValueTagType>
+	FSimpleQuery& FSimpleQuery::All(const FName& Tag)
+	{
+		return All(UE::Editor::DataStorage::FDynamicTag(Tag));
+	}
+
+	template<UE::Editor::DataStorage::TValueTagType>
+	FSimpleQuery& FSimpleQuery::All(const FName& Tag, const FName& Value)
+	{
+		return All(UE::Editor::DataStorage::FDynamicTag(Tag), Value);
+	}
 	
+	template <TColumnType DynamicColumnTemplate>
+	FSimpleQuery& FSimpleQuery::None(const FName& Identifier)
+	{
+		return None(UE::Editor::DataStorage::FDynamicColumnDescription
+			{
+				.TemplateType = DynamicColumnTemplate::StaticStruct(),
+				.Identifier = Identifier
+			});
+	}
+
+	template <TColumnType DynamicColumnTemplate>
+	FSimpleQuery& FSimpleQuery::All(const FName& Identifier)
+	{
+		return All(UE::Editor::DataStorage::FDynamicColumnDescription
+			{
+				.TemplateType = DynamicColumnTemplate::StaticStruct(),
+				.Identifier = Identifier
+			});
+	}
+
 	template <TEnumType EnumT>
     FSimpleQuery& FSimpleQuery::All()
     {
diff --git a/Engine/Source/Runtime/TypedElementFramework/Public/Elements/Interfaces/TypedElementDataStorageInterface.h b/Engine/Source/Runtime/TypedElementFramework/Public/Elements/Interfaces/TypedElementDataStorageInterface.h
--- a/Engine/Source/Runtime/TypedElementFramework/Public/Elements/Interfaces/TypedElementDataStorageInterface.h
+++ b/Engine/Source/Runtime/TypedElementFramework/Public/Elements/Interfaces/TypedElementDataStorageInterface.h
@@ -203,6 +203,9 @@ public:
 	 */
 	virtual void AddColumn(TypedElementDataStorage::RowHandle Row, const UE::Editor::DataStorage::FDynamicTag& Tag, const FName& Value) = 0;
 
+	template<typename T>
+	void AddColumn(TypedElementDataStorage::RowHandle Row, const FName& Tag) = delete;
+	
 	template<typename T>
 	void AddColumn(TypedElementDataStorage::RowHandle Row, const FName& Tag, const FName& Value) = delete;
 	
@@ -215,6 +218,12 @@ public:
 	template<auto Value, UE::Editor::DataStorage::TEnumType EnumT = decltype(Value)>
 	void AddColumn(TypedElementDataStorage::RowHandle Row);
 
+	template<UE::Editor::DataStorage::TColumnType DynamicColumnTemplate>
+	void AddColumn(TypedElementDataStorage::RowHandle Row, const FName& Identifier);
+	
+	template<UE::Editor::DataStorage::TColumnType DynamicColumnTemplate>
+	void AddColumn(TypedElementDataStorage::RowHandle Row, const FName& Identifier, DynamicColumnTemplate&& TemplateInstance);
+
 	/**
 	 * Adds multiple columns from a row. This is typically more efficient than adding columns one 
 	 * at a time.
@@ -239,6 +248,9 @@ public:
 
 	template<typename T>
 	void RemoveColumn(TypedElementDataStorage::RowHandle Row, const FName& Tag) = delete;
+
+	template<UE::Editor::DataStorage::TColumnType DynamicColumnTemplateType>
+	void RemoveColumn(TypedElementDataStorage::RowHandle Row, const FName& Identifier);
 	
 	template<>
 	void RemoveColumn<UE::Editor::DataStorage::FDynamicTag>(TypedElementDataStorage::RowHandle Row, const FName& Tag);
@@ -273,6 +285,11 @@ public:
 	ColumnType* GetColumn(TypedElementRowHandle Row);
 	template<UE::Editor::DataStorage::TDataColumnType ColumnType>
 	const ColumnType* GetColumn(TypedElementRowHandle Row) const;
+	// Gets a dynamic column identified by the ColumnTypeTemplate and Identifier
+	template<UE::Editor::DataStorage::TDataColumnType ColumnTypeTemplate>
+	ColumnTypeTemplate* GetColumn(TypedElementRowHandle Row, const FName& Identifer);
+	template<UE::Editor::DataStorage::TDataColumnType ColumnTypeTemplate>
+	const ColumnTypeTemplate* GetColumn(TypedElementRowHandle Row, const FName& Identifer) const;
 	
 	/** Determines if the provided row contains the collection of columns and tags. */
 	virtual bool HasColumns(TypedElementRowHandle Row, TConstArrayView<const UScriptStruct*> ColumnTypes) const = 0;
@@ -291,6 +308,20 @@ public:
 
 	/** Determines if the columns in the row match the query conditions. */
 	virtual bool MatchesColumns(TypedElementDataStorage::RowHandle Row, const TypedElementDataStorage::FQueryConditions& Conditions) const = 0;
+
+	/**
+	 * Finds the type information for a dynamic column.
+	 * If the dynamic column has not been generated, then return nullptr
+	 * The TemplateType may be a typed derived from either FTypedElementDataStorageColumn or FTypedElementDataStorageTag, anything else will return nullptr
+	 */
+	virtual const UScriptStruct* FindDynamicColumn(const UE::Editor::DataStorage::FDynamicColumnDescription& Description) const = 0;
+
+	/**
+	 * Generates a new dynamic column from a Template.  A dynamic column is uniquely identified using the given template and an Identifier
+	 * This function is idempotent - multiple calls with the same parameters will result in subsequent calls returning the same type
+	 * The TemplateType may be a typed derived from either FTypedElementDataStorageColumn or FTypedElementDataStorageTag
+	 */
+	virtual const UScriptStruct* GenerateDynamicColumn(const UE::Editor::DataStorage::FDynamicColumnDescription& Description) = 0;
 	
 	/**
 	 * @section Query
@@ -468,6 +499,18 @@ inline void ITypedElementDataStorageInterface::RemoveColumn<UE::Editor::DataStor
 	RemoveColumn(Row, FDynamicTag(Tag));
 }
 
+template<UE::Editor::DataStorage::TColumnType DynamicColumnTemplate>
+void ITypedElementDataStorageInterface::RemoveColumn(TypedElementDataStorage::RowHandle Row, const FName& Identifier)
+{
+	const UE::Editor::DataStorage::FDynamicColumnDescription Description
+	{
+		.TemplateType = DynamicColumnTemplate::StaticStruct(),
+		.Identifier = Identifier
+	};
+	const UScriptStruct* StructInfo = FindDynamicColumn(Description);
+	RemoveColumn(Row, StructInfo);
+}
+
 template<UE::Editor::DataStorage::TEnumType EnumT>
 void ITypedElementDataStorageInterface::AddColumn(TypedElementDataStorage::RowHandle Row, EnumT Value)
 {
@@ -485,6 +528,57 @@ void ITypedElementDataStorageInterface::AddColumn(TypedElementDataStorage::RowHa
 	AddColumn<EnumT>(Row, Value);
 }
 
+template <UE::Editor::DataStorage::TColumnType DynamicColumnTemplate>
+void ITypedElementDataStorageInterface::AddColumn(TypedElementDataStorage::RowHandle Row, const FName& Identifier)
+{
+	static_assert(UE::Editor::DataStorage::TDataColumnType<DynamicColumnTemplate> || UE::Editor::DataStorage::TTagColumnType<DynamicColumnTemplate>,
+		"DynamicColumnTemplate must be derived from either a Tag or Column");
+	const UE::Editor::DataStorage::FDynamicColumnDescription Description
+	{
+		.TemplateType = DynamicColumnTemplate::StaticStruct(),
+		.Identifier = Identifier
+	};
+	const UScriptStruct* StructInfo = GenerateDynamicColumn(Description);
+	AddColumn(Row, StructInfo);
+}
+
+template <UE::Editor::DataStorage::TColumnType DynamicColumnTemplate>
+void ITypedElementDataStorageInterface::AddColumn(TypedElementDataStorage::RowHandle Row, const FName& Identifier, DynamicColumnTemplate&& TemplateInstance)
+{
+	static_assert(UE::Editor::DataStorage::TDataColumnType<DynamicColumnTemplate> || UE::Editor::DataStorage::TTagColumnType<DynamicColumnTemplate>,
+		"DynamicColumnTemplate must be derived from either a Tag or Column");
+	const UE::Editor::DataStorage::FDynamicColumnDescription Description
+	{
+		.TemplateType = DynamicColumnTemplate::StaticStruct(),
+		.Identifier = Identifier
+	};
+	const UScriptStruct* StructInfo = GenerateDynamicColumn(Description);
+	AddColumnData(Row, StructInfo,
+		[&TemplateInstance](void* ColumnData, const UScriptStruct&)
+		{
+			if constexpr (std::is_move_constructible_v<DynamicColumnTemplate>)
+			{
+				new(ColumnData) DynamicColumnTemplate(MoveTemp(TemplateInstance));
+			}
+			else
+			{
+				new(ColumnData) DynamicColumnTemplate(TemplateInstance);
+			}
+		},
+		[](const UScriptStruct&, void* Destination, void* Source)
+		{
+			if constexpr (std::is_move_assignable_v<DynamicColumnTemplate>)
+			{
+				*static_cast<DynamicColumnTemplate*>(Destination) = MoveTemp(*static_cast<DynamicColumnTemplate*>(Source));
+			}
+			else
+			{
+				*static_cast<DynamicColumnTemplate*>(Destination) = *static_cast<DynamicColumnTemplate*>(Source);
+			}
+		});
+	
+}
+
 template<UE::Editor::DataStorage::TEnumType EnumT>
 void ITypedElementDataStorageInterface::RemoveColumn(TypedElementDataStorage::RowHandle Row)
 {
@@ -538,6 +632,38 @@ const ColumnType* ITypedElementDataStorageInterface::GetColumn(TypedElementRowHa
 	return reinterpret_cast<const ColumnType*>(GetColumnData(Row, ColumnType::StaticStruct()));
 }
 
+template <UE::Editor::DataStorage::TDataColumnType DynamicColumnTemplate>
+DynamicColumnTemplate* ITypedElementDataStorageInterface::GetColumn(TypedElementRowHandle Row, const FName& Identifier)
+{
+	const UE::Editor::DataStorage::FDynamicColumnDescription Description
+	{
+		.TemplateType = DynamicColumnTemplate::StaticStruct(),
+		.Identifier = Identifier
+	};
+	const UScriptStruct* StructInfo = GenerateDynamicColumn(Description);
+	if (StructInfo)
+	{
+		return static_cast<DynamicColumnTemplate*>(GetColumnData(Row, StructInfo));
+	}
+	return nullptr;
+}
+
+template <UE::Editor::DataStorage::TDataColumnType DynamicColumnTemplate>
+const DynamicColumnTemplate* ITypedElementDataStorageInterface::GetColumn(TypedElementRowHandle Row, const FName& Identifier) const
+{
+	const UE::Editor::DataStorage::FDynamicColumnDescription Description
+	{
+		.TemplateType = DynamicColumnTemplate::StaticStruct(),
+		.Identifier = Identifier
+	};
+	const UScriptStruct* StructInfo = FindDynamicColumn(Description);
+	if (StructInfo)
+	{
+		return static_cast<DynamicColumnTemplate*>(GetColumnData(Row, StructInfo));
+	}
+	return nullptr;
+}
+
 template<UE::Editor::DataStorage::TColumnType... ColumnType>
 bool ITypedElementDataStorageInterface::HasColumns(TypedElementRowHandle Row) const
 {
diff --git a/Engine/Source/Runtime/TypedElementFramework/Public/Elements/Interfaces/TypedElementQueryStorageInterfaces.h b/Engine/Source/Runtime/TypedElementFramework/Public/Elements/Interfaces/TypedElementQueryStorageInterfaces.h
--- a/Engine/Source/Runtime/TypedElementFramework/Public/Elements/Interfaces/TypedElementQueryStorageInterfaces.h
+++ b/Engine/Source/Runtime/TypedElementFramework/Public/Elements/Interfaces/TypedElementQueryStorageInterfaces.h
@@ -3,6 +3,7 @@
 #pragma once
 
 #include "Containers/ContainersFwd.h"
+#include "Elements/Common/TypedElementCommonTypes.h"
 #include "Elements/Common/TypedElementHandles.h"
 #include "Elements/Common/TypedElementQueryTypes.h"
 #include "UObject/Class.h"
@@ -41,11 +42,15 @@ namespace TypedElementDataStorage
 		/** Return the address of a immutable column matching the requested type or a nullptr if not found. */
 		template<typename Column>
 		const Column* GetColumn() const;
+		template<UE::Editor::DataStorage::TDataColumnType TemplateType>
+		const TemplateType* GetColumn(const FName& Identifier) const;
 		/** Return the address of a mutable column matching the requested type or a nullptr if not found. */
 		virtual void* GetMutableColumn(const UScriptStruct* ColumnType) = 0;
 		/** Return the address of a mutable column matching the requested type or a nullptr if not found. */
 		template<typename Column>
 		Column* GetMutableColumn();
+		template<UE::Editor::DataStorage::TDataColumnType TemplateType>
+		TemplateType* GetMutableColumn(const FName& Identifier);
 
 		/**
 		 * Get a list of columns or nullptrs if the column type wasn't found. Mutable addresses are returned and it's up to
@@ -86,12 +91,59 @@ namespace TypedElementDataStorage
 		 */
 		template<typename Column>
 		bool HasColumn(TypedElementDataStorage::RowHandle Row) const;
+
+		template<UE::Editor::DataStorage::TColumnType DynamicColumnTemplate>
+		bool HasColumn(TypedElementDataStorage::RowHandle Row, const FName& Identifier) const;
+
+		/**
+		 * Finds the type of a dynamic column
+		 * @return The UScriptStruct for a previously created Dynamic Column.  If no column exists, then nullptr
+		 */
+		virtual const UScriptStruct* FindDynamicColumnType(const UE::Editor::DataStorage::FDynamicColumnDescription& Description) const = 0;
+		template<UE::Editor::DataStorage::TColumnType TemplateType>
+		const UScriptStruct* FindDynamicColumnType(const FName& Identifier) const;
 	};
 
 	struct ICommonQueryWithEnvironmentContext : public ICommonQueryContext
 	{
 		using ObjectCopyOrMove = void (*)(const UScriptStruct& TypeInfo, void* Destination, void* Source);
 
+		using ICommonQueryContext::GetColumn;
+		using ICommonQueryContext::HasColumn;
+		using ICommonQueryContext::GetMutableColumn;
+		
+		/**
+		 * Helper template that returns a const reference to the column array data for the range of rows being processed
+		 * Returns nullptr if type not found or column not assigned to row
+		 * @warning It is advised to only use this function in a query callback that handles multiple rows
+		 *          the returned reference will be to the first column in the table chunk, not necessarily
+		 *          the column corresponding to the row that the singular callback form is using.
+		 */
+		template<typename DynamicColumnTemplate>
+		const DynamicColumnTemplate* GetColumn(const FName& Identifier) const;
+
+		/**
+		 * Helper template that returns the reference to the column array data for the range of rows being processed
+		 * Returns nullptr if type not found or column not assigned to row
+		 * @warning It is advised to only use this function in a query callback that handles multiple rows
+		 *          the returned reference will be to the first column in the table chunk, not necessarily
+		 *          the column corresponding to the row that the singular callback form is using.
+		 */
+		template<typename DynamicColumnTemplate>
+		DynamicColumnTemplate* GetMutableColumn(const FName& Identifier);
+
+		/*
+		 * Helper template to return whether a row has a given dynamic column.
+		 */
+		template<typename DynamicColumnTemplate>
+		bool HasColumn(const FName& Identifier) const;
+		
+		/*
+		 * Helper template to return whether a row has a given dynamic column.
+		 */
+		template<typename DynamicColumnTemplate>
+		bool HasColumn(RowHandle Row, const FName& Identifier) const;
+
 		/**
 		 * Returns the id for the current update cycle. Every time TEDS goes through a cycle of running query callbacks, this is
 		 * incremented by one. This guarantees that all query callbacks in the same run see the same cycle id. This can be useful
@@ -135,6 +187,24 @@ namespace TypedElementDataStorage
 		 */
 		template<typename ColumnType>
 		ColumnType& AddColumn(RowHandle Row, ColumnType&& Column);
+		/**
+		 * Helper template for adding a dynamic column to the requested row.  The resulting column will be default initialized.
+		 * @tparam ColumnTypeTemplate The template layout of the column
+		 * @param Row The row to add the column to
+		 * @param Identifier Combined with the template type, the dynamic column is uniquely identified with this parameter
+		 * @return An ephemeral reference to the data of the column.  The returned column will have been constructed.
+		 */
+		template<typename ColumnTypeTemplate>
+		ColumnTypeTemplate* AddColumn(RowHandle Row, const FName& Identifier);
+		/**
+		 * Helper template for adding a dynamic column to the requested row
+		 * @tparam ColumnTypeTemplate The template layout of the column
+		 * @param Row The row to add the column to
+		 * @param Identifier Combined with the template type, the dynamic column is uniquely identified with this parameter
+		 * @return An ephemeral reference to the data of the column.  The returned column will have been constructed.
+		 */
+		template<typename ColumnTypeTemplate>
+		ColumnTypeTemplate* AddColumn(RowHandle Row, const FName& Identifier, ColumnTypeTemplate&& Column);
 		/**
 		 * Adds new empty columns to a row of the provided type. The addition will not be immediately done but delayed until the end of the
 		 * tick group.
@@ -157,6 +227,8 @@ namespace TypedElementDataStorage
 		 * tick group.
 		 */
 		virtual void AddColumns(TConstArrayView<RowHandle> Rows, TConstArrayView<const UScriptStruct*> ColumnTypes) = 0;
+		
+		virtual void AddColumns(TConstArrayView<RowHandle> Rows, TConstArrayView<UE::Editor::DataStorage::FDynamicColumnDescription> DynamicColumnDescriptions) = 0;
 		/**
 		 * Add a new uninitialized column of the provided type if one does not exist.
 		 * Returns a staged column which is used to copy into the database at a later time via the UStructScript Copy operator at the end
@@ -176,6 +248,10 @@ namespace TypedElementDataStorage
 		 */
 		virtual void* AddColumnUninitialized(RowHandle Row, const UScriptStruct* ObjectType, ObjectCopyOrMove Relocator) = 0;
 
+		virtual void* AddColumnUninitialized(RowHandle Row, const UE::Editor::DataStorage::FDynamicColumnDescription& DynamicColumnDescription, ObjectCopyOrMove Relocator) = 0;
+		virtual void* AddColumnUninitialized(RowHandle Row, const UE::Editor::DataStorage::FDynamicColumnDescription& DynamicColumnDescription) = 0;
+		
+
 		/**
 		 * Removes columns of the provided types from a row. The removal will not be immediately done but delayed until the end of the
 		 * tick group.
@@ -292,12 +368,26 @@ namespace TypedElementDataStorage
 		return reinterpret_cast<const Column*>(GetColumn(Column::StaticStruct()));
 	}
 
+	template <UE::Editor::DataStorage::TDataColumnType TemplateType>
+	const TemplateType* ICommonQueryContext::GetColumn(const FName& Identifier) const
+	{
+		const UScriptStruct* DynamicColumnType = FindDynamicColumnType<TemplateType>(Identifier);
+		return static_cast<const TemplateType*>(GetColumn(DynamicColumnType));
+	}
+
 	template<typename Column>
 	Column* ICommonQueryContext::GetMutableColumn()
 	{
 		return reinterpret_cast<Column*>(GetMutableColumn(Column::StaticStruct()));
 	}
 
+	template <UE::Editor::DataStorage::TDataColumnType TemplateType>
+	TemplateType* ICommonQueryContext::GetMutableColumn(const FName& Identifier)
+	{
+		const UScriptStruct* DynamicColumnType = FindDynamicColumnType<TemplateType>(Identifier);
+		return static_cast<TemplateType*>(GetMutableColumn(DynamicColumnType));
+	}
+
 	template <typename Column>
 	bool ICommonQueryContext::HasColumn() const
 	{
@@ -310,7 +400,95 @@ namespace TypedElementDataStorage
 		return HasColumn(Row, Column::StaticStruct());
 	}
 
+	template <UE::Editor::DataStorage::TColumnType DynamicColumnTemplate>
+	bool ICommonQueryContext::HasColumn(TypedElementDataStorage::RowHandle Row, const FName& Identifier) const
+	{
+		if (const UScriptStruct* DynamicColumnType = FindDynamicColumnType<DynamicColumnTemplate>(Identifier))
+		{
+			return HasColumn(DynamicColumnType);
+		}
+		return false;
+	}
+
+	template <UE::Editor::DataStorage::TColumnType TemplateType>
+	const UScriptStruct* ICommonQueryContext::FindDynamicColumnType(const FName& Identifier) const
+	{
+		return FindDynamicColumnType(UE::Editor::DataStorage::FDynamicColumnDescription
+			{
+				.TemplateType = TemplateType::StaticStruct(),
+				.Identifier = Identifier
+			});
+	}
 
+	template <typename DynamicColumnTemplate>
+	const DynamicColumnTemplate* ICommonQueryWithEnvironmentContext::GetColumn(const FName& Identifier) const
+	{
+		using namespace UE::Editor::DataStorage;
+		const FDynamicColumnDescription Description
+		{
+			.TemplateType = DynamicColumnTemplate::StaticStruct(),
+			.Identifier = Identifier
+		};
+		const UScriptStruct* StructInfo = FindDynamicColumnType(Description);
+		if (StructInfo)
+		{
+			const void* ColumnData = GetColumn(StructInfo);
+			return static_cast<const DynamicColumnTemplate*>(ColumnData);
+		}
+		return nullptr;
+	}
+
+	template <typename DynamicColumnTemplate>
+	DynamicColumnTemplate* ICommonQueryWithEnvironmentContext::GetMutableColumn(const FName& Identifier)
+	{
+		using namespace UE::Editor::DataStorage;
+		const FDynamicColumnDescription Description
+		{
+			.TemplateType = DynamicColumnTemplate::StaticStruct(),
+			.Identifier = Identifier
+		};
+		const UScriptStruct* StructInfo = FindDynamicColumnType(Description);
+		if (StructInfo)
+		{
+			const void* ColumnData = GetMutableColumn(StructInfo);
+			return static_cast<const DynamicColumnTemplate*>(ColumnData);
+		}
+		return nullptr;
+	}
+
+	template <typename DynamicColumnTemplate>
+	bool ICommonQueryWithEnvironmentContext::HasColumn(const FName& Identifier) const
+	{
+		using namespace UE::Editor::DataStorage;
+		const FDynamicColumnDescription Description
+		{
+			.TemplateType = DynamicColumnTemplate::StaticStruct(),
+			.Identifier = Identifier
+		};
+		const UScriptStruct* StructInfo = FindDynamicColumnType(Description);
+		if (StructInfo)
+		{
+			return ICommonQueryContext::HasColumn(StructInfo);
+		}
+		return false;
+	}
+
+	template <typename DynamicColumnTemplate>
+	bool ICommonQueryWithEnvironmentContext::HasColumn(RowHandle Row, const FName& Identifier) const
+	{
+		using namespace UE::Editor::DataStorage;
+		const FDynamicColumnDescription Description
+		{
+			.TemplateType = DynamicColumnTemplate::StaticStruct(),
+			.Identifier = Identifier
+		};
+		const UScriptStruct* StructInfo = FindDynamicColumnType(Description);
+		if (StructInfo)
+		{
+			return ICommonQueryContext::HasColumn(Row, StructInfo);
+		}
+		return false;
+	}
 
 	//
 	// ICommonQueryWithEnvironmentContext
@@ -338,6 +516,58 @@ namespace TypedElementDataStorage
 		}
 	}
 
+	template <typename ColumnTypeTemplate>
+	ColumnTypeTemplate* ICommonQueryWithEnvironmentContext::AddColumn(RowHandle Row, const FName& Identifier)
+	{
+		UScriptStruct* TemplateType = ColumnTypeTemplate::StaticStruct();
+		const UE::Editor::DataStorage::FDynamicColumnDescription Description
+		{
+			.TemplateType = TemplateType,
+			.Identifier = Identifier
+		};
+
+		if constexpr (UE::Editor::DataStorage::TDataColumnType<ColumnTypeTemplate>)
+		{
+			ColumnTypeTemplate* ColumnData = static_cast<ColumnTypeTemplate*>(AddColumnUninitialized(Row, Description));
+			new (ColumnData) ColumnTypeTemplate();
+			return ColumnData;
+		}
+		if constexpr (UE::Editor::DataStorage::TTagColumnType<ColumnTypeTemplate>)
+		{
+			AddColumns({Row}, {Description});
+			return nullptr;
+		}
+
+		return  nullptr;
+	}
+
+	template <typename ColumnTypeTemplate>
+	ColumnTypeTemplate* ICommonQueryWithEnvironmentContext::AddColumn(RowHandle Row, const FName& Identifier, ColumnTypeTemplate&& Column)
+	{
+		UScriptStruct* TemplateType = ColumnTypeTemplate::StaticStruct();
+		const UE::Editor::DataStorage::FDynamicColumnDescription Description
+		{
+			.TemplateType = TemplateType,
+			.Identifier = Identifier
+		};
+
+		if constexpr (std::is_move_constructible_v<ColumnTypeTemplate>)
+		{
+			void* Address = AddColumnUninitialized(Row, Description,
+				[](const UScriptStruct&, void* Destination, void* Source)
+				{
+					*static_cast<ColumnTypeTemplate*>(Destination) = MoveTemp(*static_cast<ColumnTypeTemplate*>(Source));
+				});
+			return *(new(Address) ColumnTypeTemplate(Forward<ColumnTypeTemplate>(Column)));
+		}
+		else
+		{
+			void* Address = AddColumnUninitialized(Row, Description);
+			new (Address) ColumnTypeTemplate(Forward<ColumnTypeTemplate>(Column));
+			return *static_cast<ColumnTypeTemplate*>(Address);
+		}
+	}
+
 	template<typename... Columns>
 	void ICommonQueryWithEnvironmentContext::AddColumns(RowHandle Row)
 	{
