[TEDS] Compile warning fixes for MacArm64

* Declaration shadowing

diff --git a/Engine/Plugins/Experimental/EditorDataStorage/Source/TedsCore/Tests/DynamicColumnTests.cpp b/Engine/Plugins/Experimental/EditorDataStorage/Source/TedsCore/Tests/DynamicColumnTests.cpp
--- a/Engine/Plugins/Experimental/EditorDataStorage/Source/TedsCore/Tests/DynamicColumnTests.cpp
+++ b/Engine/Plugins/Experimental/EditorDataStorage/Source/TedsCore/Tests/DynamicColumnTests.cpp
@@ -169,8 +169,8 @@ namespace UE::Editor::DataStorage::Tests
 						};
 						auto Callback = CreateDirectQueryCallbackBinding([this, &RowsToMatch, &WasMatched] (TypedElementDataStorage::IDirectQueryContext& Context, const TypedElementDataStorage::RowHandle* CallbackRows)
 						{
-							TConstArrayView<TypedElementDataStorage::RowHandle> Rows = MakeConstArrayView(CallbackRows, Context.GetRowCount());
-							for (TypedElementDataStorage::RowHandle Row : Rows)
+							TConstArrayView<TypedElementDataStorage::RowHandle> RowsView = MakeConstArrayView(CallbackRows, Context.GetRowCount());
+							for (TypedElementDataStorage::RowHandle Row : RowsView)
 							{
 								int32 Index = RowsToMatch.Find(Row);
 								TestTrue(TEXT("Returned row in query is within expected match array"), Index != INDEX_NONE);
@@ -281,8 +281,8 @@ namespace UE::Editor::DataStorage::Tests
 					{
 						auto Callback = [this, &RowsToMatch, &MatchCount, &UnexpectedRowCount] (TypedElementDataStorage::IQueryContext& Context, const TypedElementDataStorage::RowHandle* CallbackRows)
 						{
-							TConstArrayView<TypedElementDataStorage::RowHandle> Rows = MakeConstArrayView(CallbackRows, Context.GetRowCount());
-							for (auto Row : Rows)
+							TConstArrayView<TypedElementDataStorage::RowHandle> RowsView = MakeConstArrayView(CallbackRows, Context.GetRowCount());
+							for (auto Row : RowsView)
 							{
 								int32 Index = RowsToMatch.Find(Row);
 								TestTrue(TEXT("Returned row in query is within expected match array"), Index != INDEX_NONE);
@@ -489,9 +489,9 @@ namespace UE::Editor::DataStorage::Tests
 						using namespace TypedElementQueryBuilder;
 						using namespace TypedElementDataStorage;
 
-						auto RunTest = [this, &AllTestExpectations, &UnexpectedRowCount](IQueryContext& Context, const RowHandle* Rows, int32 TestIndex)
+						auto RunTest = [this, &AllTestExpectations, &UnexpectedRowCount](IQueryContext& Context, const RowHandle* RowsPtr, int32 TestIndex)
 						{
-							TConstArrayView<RowHandle> RowView = MakeConstArrayView(Rows, Context.GetRowCount());
+							TConstArrayView<RowHandle> RowView = MakeConstArrayView(RowsPtr, Context.GetRowCount());
 							TConstArrayView<FTestDynamicColumn> ColumnView = MakeConstArrayView(
 								Context.GetColumn<FTestDynamicColumn>(Identifiers[0]),
 								Context.GetRowCount());
@@ -550,9 +550,9 @@ namespace UE::Editor::DataStorage::Tests
 								ActivationKeys.Last(),
 								FProcessor(EQueryTickPhase::FrameEnd, TedsInterface->GetQueryTickGroupName(EQueryTickGroups::SyncDataStorageToExternal))
 									.MakeActivatable(ActivationKeys.Last()),
-									[this, TestIndex = AllTestExpectations.Num() - 1, &RunTest](IQueryContext& Context, const RowHandle* Rows)
+									[this, TestIndex = AllTestExpectations.Num() - 1, &RunTest](IQueryContext& Context, const RowHandle* RowsPtr)
 									{
-										RunTest(Context, Rows, TestIndex);
+										RunTest(Context, RowsPtr, TestIndex);
 									}
 							)
 							.ReadOnly<FTestDynamicColumn>(Identifiers[0])
