[Teds] Update naming for RowHandle, QueryHandle, and TableHandle

Some more namespacing adjustments

diff --git a/Engine/Plugins/Developer/PerforceSourceControl/Source/PerforceSourceControl/Private/PerforceSourceControlOperations.cpp b/Engine/Plugins/Developer/PerforceSourceControl/Source/PerforceSourceControl/Private/PerforceSourceControlOperations.cpp
--- a/Engine/Plugins/Developer/PerforceSourceControl/Source/PerforceSourceControl/Private/PerforceSourceControlOperations.cpp
+++ b/Engine/Plugins/Developer/PerforceSourceControl/Source/PerforceSourceControl/Private/PerforceSourceControlOperations.cpp
@@ -2023,6 +2023,7 @@ bool FPerforceUpdateStatusWorker::UpdateStates() const
 		using namespace TypedElementQueryBuilder;
 		using namespace TypedElementDataStorage;
 		using DSI = ITypedElementDataStorageInterface;
+		using namespace UE::Editor;
 
 		UTypedElementRegistry* Registry = UTypedElementRegistry::GetInstance();
 		if (!Registry)
@@ -2036,25 +2037,25 @@ bool FPerforceUpdateStatusWorker::UpdateStates() const
 			return;
 		}
 
-		auto GetRevisionControlRow = [DataStorage](const FString& InFilename) -> TypedElementRowHandle
+		auto GetRevisionControlRow = [DataStorage](const FString& InFilename) -> DataStorage::RowHandle
 		{
 			FString Filename = FPaths::SetExtension(InFilename, "");
 			FPaths::NormalizeFilename(Filename);
 			Filename = FPaths::ConvertRelativePathToFull(Filename);
 			
 			uint64 Index = GenerateIndexHash(Filename);
-			TypedElementRowHandle Row = DataStorage->FindIndexedRow(Index);
+			DataStorage::RowHandle Row = DataStorage->FindIndexedRow(Index);
 
 			if (!DataStorage->IsRowAvailable(Row))
 			{
-				static TypedElementTableHandle Table = DataStorage->FindTable(FName("Editor_RevisionControlTable"));
+				static DataStorage::RowHandle Table = DataStorage->FindTable(FName("Editor_RevisionControlTable"));
 				Row = DataStorage->AddRow(Table);
 				DataStorage->IndexRow(Index, Row);
 			}
 			return Row;
 		};
 
-		TypedElementRowHandle Row = GetRevisionControlRow(State.GetFilename());
+		DataStorage::RowHandle Row = GetRevisionControlRow(State.GetFilename());
 
 		if (!State.IsSourceControlled())
 		{
@@ -2125,10 +2126,10 @@ bool FPerforceUpdateStatusWorker::UpdateStates() const
 		
         if (FTypedElementPackageReference* BackReference = DataStorage->GetColumn<FTypedElementPackageReference>(Row))
         {
-        	if (RowHandle ObjectRow = BackReference->Row; ObjectRow != InvalidRowHandle)
+        	if (DataStorage::RowHandle ObjectRow = BackReference->Row; ObjectRow != InvalidRowHandle)
         	{
         		static TableHandle Table = DataStorage->FindTable(FName("Editor_PackageUpdateTable"));
-        		RowHandle UpdateRow = DataStorage->AddRow(Table);
+        		DataStorage::RowHandle UpdateRow = DataStorage->AddRow(Table);
         		DataStorage->AddColumn(UpdateRow, FTypedElementPackageUpdateColumn{ .ObjectRow = ObjectRow, .PackageRow = Row });
         	}
         }
diff --git a/Engine/Plugins/Experimental/Animation/HierarchyTable/Source/Editor/Private/HierarchyTable/Factories/OverrideFactory.cpp b/Engine/Plugins/Experimental/Animation/HierarchyTable/Source/Editor/Private/HierarchyTable/Factories/OverrideFactory.cpp
--- a/Engine/Plugins/Experimental/Animation/HierarchyTable/Source/Editor/Private/HierarchyTable/Factories/OverrideFactory.cpp
+++ b/Engine/Plugins/Experimental/Animation/HierarchyTable/Source/Editor/Private/HierarchyTable/Factories/OverrideFactory.cpp
@@ -7,7 +7,5 @@
 
 void UHierarchyTableOverrideFactory::RegisterWidgetConstructors(ITypedElementDataStorageInterface& DataStorage, ITypedElementDataStorageUiInterface& DataStorageUi) const
 {
-	using namespace TypedElementDataStorage;
-
-	DataStorageUi.RegisterWidgetFactory<FTypedElementWidgetConstructor_Override>(FName(TEXT("General.Cell")), FColumn<FTypedElementOverrideColumn>());
+	DataStorageUi.RegisterWidgetFactory<FTypedElementWidgetConstructor_Override>(FName(TEXT("General.Cell")), TypedElementDataStorage::FColumn<FTypedElementOverrideColumn>());
 }
\ No newline at end of file
diff --git a/Engine/Plugins/Experimental/Animation/HierarchyTable/Source/Editor/Private/HierarchyTableEditorToolkit.cpp b/Engine/Plugins/Experimental/Animation/HierarchyTable/Source/Editor/Private/HierarchyTableEditorToolkit.cpp
--- a/Engine/Plugins/Experimental/Animation/HierarchyTable/Source/Editor/Private/HierarchyTableEditorToolkit.cpp
+++ b/Engine/Plugins/Experimental/Animation/HierarchyTable/Source/Editor/Private/HierarchyTableEditorToolkit.cpp
@@ -87,7 +87,7 @@ void FHierarchyTableEditorToolkit::OnClose()
 	UTypedElementRegistry* Registry = UTypedElementRegistry::GetInstance();
 	ITypedElementDataStorageInterface* DSI = Registry->GetMutableDataStorage();
 
-	for (const TTuple<int32, TypedElementDataStorage::RowHandle>& Row : EntryIndexToHandleMap)
+	for (const TTuple<int32, UE::Editor::DataStorage::RowHandle>& Row : EntryIndexToHandleMap)
 	{
 		DSI->RemoveRow(Row.Value);
 	}
@@ -112,6 +112,7 @@ TSharedRef<SWidget> FHierarchyTableEditorToolkit::CreateTedsOutliner()
 	}
 
 	using namespace TypedElementQueryBuilder;
+	using namespace UE::Editor::DataStorage;
 
 	if (!ensure(HierarchyTable->TableType))
 	{
@@ -156,7 +157,7 @@ TSharedRef<SWidget> FHierarchyTableEditorToolkit::CreateTedsOutliner()
 	FTedsOutlinerModule& TedsOutlinerModule = FModuleManager::GetModuleChecked<FTedsOutlinerModule>("TedsOutliner");
 
 	ITypedElementDataStorageInterface* DSI = Registry->GetMutableDataStorage();
-	static TypedElementDataStorage::TableHandle Table = DSI->FindTable(FName("Editor_HierarchyTableTable"));
+	static TableHandle Table = DSI->FindTable(FName("Editor_HierarchyTableTable"));
 	
 	TArray<UScriptStruct*> BaseHierarchyTableTypeColumns = Handler->GetColumns();
 
@@ -164,7 +165,7 @@ TSharedRef<SWidget> FHierarchyTableEditorToolkit::CreateTedsOutliner()
 	{
 		FHierarchyTableEntryData* Entry = &HierarchyTable->TableData[EntryIndex];
 
-		TypedElementRowHandle Row = DSI->AddRow(Table);
+		RowHandle Row = DSI->AddRow(Table);
 
 		FTypedElementOverrideColumn OverrideEntry;
 		OverrideEntry.OwnerEntry = Entry;
@@ -173,7 +174,7 @@ TSharedRef<SWidget> FHierarchyTableEditorToolkit::CreateTedsOutliner()
 
 		DSI->AddColumn<FTypedElementLabelColumn>(Row, { .Label = Entry->Identifier.ToString() });
 
-		TypedElementRowHandle* ParentRow = EntryIndexToHandleMap.Find(Entry->Parent);
+		RowHandle* ParentRow = EntryIndexToHandleMap.Find(Entry->Parent);
 		if (ParentRow)
 		{
 			DSI->AddColumn<FTableRowParentColumn>(Row, { .Parent = *ParentRow });
@@ -229,6 +230,7 @@ void FHierarchyTableEditorToolkit::AddCurveEntry(const FName CurveName)
 	// TODO: Reimplement
 
 	/*
+	using namespace UE::Editor::DataStorage;
 	UTypedElementRegistry* Registry = UTypedElementRegistry::GetInstance();
 	ITypedElementDataStorageInterface* DSI = Registry->GetMutableDataStorage();
 	static TypedElementDataStorage::TableHandle Table = DSI->FindTable(FName("Editor_HierarchyTableTable"));
@@ -239,7 +241,7 @@ void FHierarchyTableEditorToolkit::AddCurveEntry(const FName CurveName)
 		return;
 	}
 
-	TypedElementRowHandle Row = DSI->AddRow(Table);
+	RowHandle Row = DSI->AddRow(Table);
 	DSI->AddColumn<FTypedElementOverrideColumn>(Row, { .bIsOverridden = false });
 	DSI->AddColumn<FTypedElementLabelColumn>(Row, { .Label = CurveName.ToString() });
 	DSI->AddColumn<FTypedElementMetadataColumn>(Row,
diff --git a/Engine/Plugins/Experimental/Animation/HierarchyTable/Source/Editor/Public/HierarchyTableEditorToolkit.h b/Engine/Plugins/Experimental/Animation/HierarchyTable/Source/Editor/Public/HierarchyTableEditorToolkit.h
--- a/Engine/Plugins/Experimental/Animation/HierarchyTable/Source/Editor/Public/HierarchyTableEditorToolkit.h
+++ b/Engine/Plugins/Experimental/Animation/HierarchyTable/Source/Editor/Public/HierarchyTableEditorToolkit.h
@@ -30,9 +30,9 @@ private:
 
 	UHierarchyTable* HierarchyTable;
 
-	TMap<int32, TypedElementDataStorage::RowHandle> EntryIndexToHandleMap;
+	TMap<int32, UE::Editor::DataStorage::RowHandle> EntryIndexToHandleMap;
 
-	TypedElementDataStorage::QueryHandle InitialColumnQuery;
+	UE::Editor::DataStorage::QueryHandle InitialColumnQuery;
 
-	//TypedElementDataStorage::QueryHandle CollectDataQuery;
+	//UE::Editor::DataStorage::QueryHandle CollectDataQuery;
 };
diff --git a/Engine/Plugins/Experimental/Animation/HierarchyTableBuiltin/Source/Editor/Private/MaskProfile/MaskProfileFactory.cpp b/Engine/Plugins/Experimental/Animation/HierarchyTableBuiltin/Source/Editor/Private/MaskProfile/MaskProfileFactory.cpp
--- a/Engine/Plugins/Experimental/Animation/HierarchyTableBuiltin/Source/Editor/Private/MaskProfile/MaskProfileFactory.cpp
+++ b/Engine/Plugins/Experimental/Animation/HierarchyTableBuiltin/Source/Editor/Private/MaskProfile/MaskProfileFactory.cpp
@@ -6,7 +6,5 @@
 
 void UHierarchyTableMaskFactory::RegisterWidgetConstructors(ITypedElementDataStorageInterface& DataStorage, ITypedElementDataStorageUiInterface& DataStorageUi) const
 {
-	using namespace TypedElementDataStorage;
-
-	DataStorageUi.RegisterWidgetFactory<FHierarchyTableMaskWidgetConstructor_Value>(FName(TEXT("General.Cell")), FColumn<FHierarchyTableMaskColumn_Value>());
+	DataStorageUi.RegisterWidgetFactory<FHierarchyTableMaskWidgetConstructor_Value>(FName(TEXT("General.Cell")), TypedElementDataStorage::FColumn<FHierarchyTableMaskColumn_Value>());
 }
\ No newline at end of file
diff --git a/Engine/Plugins/Experimental/Animation/HierarchyTableBuiltin/Source/Editor/Private/TimeProfile/TimeProfileFactory.cpp b/Engine/Plugins/Experimental/Animation/HierarchyTableBuiltin/Source/Editor/Private/TimeProfile/TimeProfileFactory.cpp
--- a/Engine/Plugins/Experimental/Animation/HierarchyTableBuiltin/Source/Editor/Private/TimeProfile/TimeProfileFactory.cpp
+++ b/Engine/Plugins/Experimental/Animation/HierarchyTableBuiltin/Source/Editor/Private/TimeProfile/TimeProfileFactory.cpp
@@ -6,10 +6,8 @@
 
 void UHierarchyTableTimeFactory::RegisterWidgetConstructors(ITypedElementDataStorageInterface& DataStorage, ITypedElementDataStorageUiInterface& DataStorageUi) const
 {
-	using namespace TypedElementDataStorage;
-
-	DataStorageUi.RegisterWidgetFactory<FHierarchyTableTimeWidgetConstructor_StartTime>(FName(TEXT("General.Cell")), FColumn<FHierarchyTableTimeColumn_StartTime>());
-	DataStorageUi.RegisterWidgetFactory<FHierarchyTableTimeWidgetConstructor_EndTime>(FName(TEXT("General.Cell")), FColumn<FHierarchyTableTimeColumn_EndTime>());
-	DataStorageUi.RegisterWidgetFactory<FHierarchyTableTimeWidgetConstructor_TimeFactor>(FName(TEXT("General.Cell")), FColumn<FHierarchyTableTimeColumn_TimeFactor>());
-	DataStorageUi.RegisterWidgetFactory<FHierarchyTableTimeWidgetConstructor_Preview>(FName(TEXT("General.Cell")), FColumn<FHierarchyTableTimeColumn_Preview>());
+	DataStorageUi.RegisterWidgetFactory<FHierarchyTableTimeWidgetConstructor_StartTime>(FName(TEXT("General.Cell")), TypedElementDataStorage::FColumn<FHierarchyTableTimeColumn_StartTime>());
+	DataStorageUi.RegisterWidgetFactory<FHierarchyTableTimeWidgetConstructor_EndTime>(FName(TEXT("General.Cell")), TypedElementDataStorage::FColumn<FHierarchyTableTimeColumn_EndTime>());
+	DataStorageUi.RegisterWidgetFactory<FHierarchyTableTimeWidgetConstructor_TimeFactor>(FName(TEXT("General.Cell")), TypedElementDataStorage::FColumn<FHierarchyTableTimeColumn_TimeFactor>());
+	DataStorageUi.RegisterWidgetFactory<FHierarchyTableTimeWidgetConstructor_Preview>(FName(TEXT("General.Cell")), TypedElementDataStorage::FColumn<FHierarchyTableTimeColumn_Preview>());
 }
\ No newline at end of file
diff --git a/Engine/Plugins/Experimental/EditorDataStorage/Source/TedsCore/Private/Compatibility/TedsTypedElementActorHandleFactory.cpp b/Engine/Plugins/Experimental/EditorDataStorage/Source/TedsCore/Private/Compatibility/TedsTypedElementActorHandleFactory.cpp
--- a/Engine/Plugins/Experimental/EditorDataStorage/Source/TedsCore/Private/Compatibility/TedsTypedElementActorHandleFactory.cpp
+++ b/Engine/Plugins/Experimental/EditorDataStorage/Source/TedsCore/Private/Compatibility/TedsTypedElementActorHandleFactory.cpp
@@ -48,7 +48,7 @@ void UTypedElementActorHandleDataStorageFactory::RegisterQuery_ActorHandlePopula
 	using namespace TypedElementQueryBuilder;
 	using namespace TypedElementDataStorage;
 	using DSI = ITypedElementDataStorageInterface;
-	if (!ensureMsgf(ActorHandlePopulateQuery == TypedElementInvalidQueryHandle, TEXT("Already registered query")))
+	if (!ensureMsgf(ActorHandlePopulateQuery == UE::Editor::DataStorage::InvalidQueryHandle, TEXT("Already registered query")))
 	{
 		return;
 	}
@@ -56,7 +56,7 @@ void UTypedElementActorHandleDataStorageFactory::RegisterQuery_ActorHandlePopula
 	ActorHandlePopulateQuery = DataStorage.RegisterQuery(
 	Select(TEXT("Populate actor typed element handles"),
 		FObserver::OnAdd<FTypedElementUObjectColumn>(),
-		[](DSI::IQueryContext& Context, TypedElementRowHandle Row, const FTypedElementUObjectColumn& ObjectColumn)
+		[](DSI::IQueryContext& Context, UE::Editor::DataStorage::RowHandle Row, const FTypedElementUObjectColumn& ObjectColumn)
 		{
 			if (UObject* Object = ObjectColumn.Object.Get())
 			{
@@ -76,8 +76,10 @@ void UTypedElementActorHandleDataStorageFactory::RegisterQuery_ActorHandlePopula
 
 void UTypedElementActorHandleDataStorageFactory::HandleBridgeEnabled(bool bEnabled)
 {
+	using namespace UE::Editor::DataStorage;
+
 	ITypedElementDataStorageInterface* DataStorage = UTypedElementRegistry::GetInstance()->GetMutableDataStorage();
-	
+
 	if (bEnabled)
 	{
 		using namespace TypedElementQueryBuilder;
@@ -109,7 +111,7 @@ void UTypedElementActorHandleDataStorageFactory::HandleBridgeEnabled(bool bEnabl
 			if (const AActor* Actor = Actors[Index].Get())
 			{
 				FTypedElementHandle Handle = UEngineElementsLibrary::AcquireEditorActorElementHandle(Actor);
-				DataStorage->AddColumn(CollatedRowHandles[Index], UE::Editor::DataStorage::Compatibility::FTypedElementColumn
+				DataStorage->AddColumn(CollatedRowHandles[Index], Compatibility::FTypedElementColumn
 				{
 					.Handle = Handle
 				});
@@ -121,7 +123,7 @@ void UTypedElementActorHandleDataStorageFactory::HandleBridgeEnabled(bool bEnabl
 	else
 	{
 		DataStorage->UnregisterQuery(ActorHandlePopulateQuery);
-		ActorHandlePopulateQuery = TypedElementInvalidQueryHandle;
+		ActorHandlePopulateQuery = InvalidQueryHandle;
 	}
 }
 
diff --git a/Engine/Plugins/Experimental/EditorDataStorage/Source/TedsCore/Private/Compatibility/TedsTypedElementActorHandleFactory.h b/Engine/Plugins/Experimental/EditorDataStorage/Source/TedsCore/Private/Compatibility/TedsTypedElementActorHandleFactory.h
--- a/Engine/Plugins/Experimental/EditorDataStorage/Source/TedsCore/Private/Compatibility/TedsTypedElementActorHandleFactory.h
+++ b/Engine/Plugins/Experimental/EditorDataStorage/Source/TedsCore/Private/Compatibility/TedsTypedElementActorHandleFactory.h
@@ -29,8 +29,8 @@ private:
 	void HandleBridgeEnabled(bool bEnabled);
 
 	FDelegateHandle BridgeEnableDelegateHandle;
-	TypedElementQueryHandle ActorHandlePopulateQuery = TypedElementInvalidQueryHandle;
-	TypedElementQueryHandle GetAllActorsQuery = TypedElementInvalidQueryHandle;
+	UE::Editor::DataStorage::QueryHandle ActorHandlePopulateQuery = UE::Editor::DataStorage::InvalidQueryHandle;
+	UE::Editor::DataStorage::QueryHandle GetAllActorsQuery = UE::Editor::DataStorage::InvalidQueryHandle;
 };
 
 
diff --git a/Engine/Plugins/Experimental/EditorDataStorage/Source/TedsCore/Private/Compatibility/TedsTypedElementBridgeQueries.cpp b/Engine/Plugins/Experimental/EditorDataStorage/Source/TedsCore/Private/Compatibility/TedsTypedElementBridgeQueries.cpp
--- a/Engine/Plugins/Experimental/EditorDataStorage/Source/TedsCore/Private/Compatibility/TedsTypedElementBridgeQueries.cpp
+++ b/Engine/Plugins/Experimental/EditorDataStorage/Source/TedsCore/Private/Compatibility/TedsTypedElementBridgeQueries.cpp
@@ -77,7 +77,7 @@ void UTypedElementBridgeDataStorageFactory::CleanupTypedElementColumns(ITypedEle
 
 	// Remove any TEv1 handles
 	{
-		TArray<TypedElementRowHandle> Handles;
+		TArray<UE::Editor::DataStorage::RowHandle> Handles;
 		using namespace TypedElementQueryBuilder;
 		DataStorage.RunQuery(
 			RemoveTypedElementRowHandleQuery,
@@ -87,7 +87,7 @@ void UTypedElementBridgeDataStorageFactory::CleanupTypedElementColumns(ITypedEle
 					Handles.Append(Context.GetRowHandles());
 				}));
 		
-		DataStorage.BatchAddRemoveColumns(TConstArrayView<TypedElementRowHandle>(Handles), {}, {UE::Editor::DataStorage::Compatibility::FTypedElementColumn::StaticStruct()});
+		DataStorage.BatchAddRemoveColumns(TConstArrayView<UE::Editor::DataStorage::RowHandle>(Handles), {}, {UE::Editor::DataStorage::Compatibility::FTypedElementColumn::StaticStruct()});
 	}
 }
 
diff --git a/Engine/Plugins/Experimental/EditorDataStorage/Source/TedsCore/Private/Compatibility/TedsTypedElementBridgeQueries.h b/Engine/Plugins/Experimental/EditorDataStorage/Source/TedsCore/Private/Compatibility/TedsTypedElementBridgeQueries.h
--- a/Engine/Plugins/Experimental/EditorDataStorage/Source/TedsCore/Private/Compatibility/TedsTypedElementBridgeQueries.h
+++ b/Engine/Plugins/Experimental/EditorDataStorage/Source/TedsCore/Private/Compatibility/TedsTypedElementBridgeQueries.h
@@ -34,6 +34,6 @@ private:
 	void CleanupTypedElementColumns(ITypedElementDataStorageInterface& DataStorage);
 	void HandleOnEnabled(IConsoleVariable* CVar);
 	
-	TypedElementQueryHandle RemoveTypedElementRowHandleQuery = TypedElementDataStorage::InvalidQueryHandle;
+	UE::Editor::DataStorage::QueryHandle RemoveTypedElementRowHandleQuery = UE::Editor::DataStorage::InvalidQueryHandle;
 	FDelegateHandle DebugEnabledDelegateHandle;
 };
diff --git a/Engine/Plugins/Experimental/EditorDataStorage/Source/TedsCore/Private/Compatibility/TypedElementActorParentQueries.cpp b/Engine/Plugins/Experimental/EditorDataStorage/Source/TedsCore/Private/Compatibility/TypedElementActorParentQueries.cpp
--- a/Engine/Plugins/Experimental/EditorDataStorage/Source/TedsCore/Private/Compatibility/TypedElementActorParentQueries.cpp
+++ b/Engine/Plugins/Experimental/EditorDataStorage/Source/TedsCore/Private/Compatibility/TypedElementActorParentQueries.cpp
@@ -9,7 +9,7 @@
 #include "Elements/Framework/TypedElementQueryBuilder.h"
 #include "GameFramework/Actor.h"
 
-namespace UE::TypedElementActorParentQueries::Local
+namespace UE::Editor::DataStorage::Queries::Private
 {
 	static bool bAddParentColumnToActors = false;
 	
@@ -22,7 +22,7 @@ namespace UE::TypedElementActorParentQueries::Local
 
 void UActorParentDataStorageFactory::RegisterQueries(ITypedElementDataStorageInterface& DataStorage)
 {
-	if(UE::TypedElementActorParentQueries::Local::bAddParentColumnToActors)
+	if(UE::Editor::DataStorage::Queries::Private::bAddParentColumnToActors)
 	{
 		RegisterAddParentColumn(DataStorage);
 		RegisterUpdateOrRemoveParentColumn(DataStorage);
@@ -33,13 +33,14 @@ void UActorParentDataStorageFactory::RegisterAddParentColumn(ITypedElementDataSt
 {
 	using namespace TypedElementDataStorage;
 	using namespace TypedElementQueryBuilder;
+	using namespace UE::Editor::DataStorage;
 	
 	DataStorage.RegisterQuery(
 		Select(
 			TEXT("Add parent column to actor"),
 			FProcessor(EQueryTickPhase::PrePhysics, DataStorage.GetQueryTickGroupName(EQueryTickGroups::SyncExternalToDataStorage))
 				.SetExecutionMode(EExecutionMode::GameThread),
-			[](IQueryContext& Context, TypedElementRowHandle Row, const FTypedElementUObjectColumn& Actor)
+			[](IQueryContext& Context, RowHandle Row, const FTypedElementUObjectColumn& Actor)
 			{
 				if (const AActor* ActorInstance = Cast<AActor>(Actor.Object))
 				{
@@ -70,13 +71,14 @@ void UActorParentDataStorageFactory::RegisterUpdateOrRemoveParentColumn(ITypedEl
 {
 	using namespace TypedElementDataStorage;
 	using namespace TypedElementQueryBuilder;
+	using namespace UE::Editor::DataStorage;
 	
 	DataStorage.RegisterQuery(
 		Select(
 			TEXT("Sync actor's parent to column"),
 			FProcessor(EQueryTickPhase::PrePhysics, DataStorage.GetQueryTickGroupName(EQueryTickGroups::SyncExternalToDataStorage))
 				.SetExecutionMode(EExecutionMode::GameThread),
-			[](IQueryContext& Context, TypedElementRowHandle Row, const FTypedElementUObjectColumn& Actor, FTableRowParentColumn& Parent)
+			[](IQueryContext& Context, RowHandle Row, const FTypedElementUObjectColumn& Actor, FTableRowParentColumn& Parent)
 			{
 				if (const AActor* ActorInstance = Cast<AActor>(Actor.Object))
 				{
diff --git a/Engine/Plugins/Experimental/EditorDataStorage/Source/TedsCore/Private/Compatibility/TypedElementActorTransformQueries.cpp b/Engine/Plugins/Experimental/EditorDataStorage/Source/TedsCore/Private/Compatibility/TypedElementActorTransformQueries.cpp
--- a/Engine/Plugins/Experimental/EditorDataStorage/Source/TedsCore/Private/Compatibility/TypedElementActorTransformQueries.cpp
+++ b/Engine/Plugins/Experimental/EditorDataStorage/Source/TedsCore/Private/Compatibility/TypedElementActorTransformQueries.cpp
@@ -19,6 +19,7 @@ void UActorTransformDataStorageFactory::RegisterActorAddTransformColumn(ITypedEl
 {
 	using namespace TypedElementQueryBuilder;
 	using namespace TypedElementDataStorage;
+	using namespace UE::Editor::DataStorage;
 	
 	DataStorage.RegisterQuery(
 		Select(
@@ -26,7 +27,7 @@ void UActorTransformDataStorageFactory::RegisterActorAddTransformColumn(ITypedEl
 			FProcessor(EQueryTickPhase::PrePhysics,
 				DataStorage.GetQueryTickGroupName(EQueryTickGroups::SyncExternalToDataStorage))
 				.SetExecutionMode(EExecutionMode::GameThread),
-			[](IQueryContext& Context, TypedElementRowHandle Row, const FTypedElementUObjectColumn& Actor)
+			[](IQueryContext& Context, RowHandle Row, const FTypedElementUObjectColumn& Actor)
 			{
 				if (const AActor* ActorInstance = Cast<AActor>(Actor.Object); ActorInstance != nullptr && ActorInstance->GetRootComponent())
 				{
@@ -45,13 +46,14 @@ void UActorTransformDataStorageFactory::RegisterActorLocalTransformToColumn(ITyp
 {
 	using namespace TypedElementQueryBuilder;
 	using namespace TypedElementDataStorage;
+	using namespace UE::Editor::DataStorage;
 	
 	DataStorage.RegisterQuery(
 		Select(
 			TEXT("Sync actor transform to column"),
 			FProcessor(EQueryTickPhase::PostPhysics, DataStorage.GetQueryTickGroupName(EQueryTickGroups::SyncExternalToDataStorage))
 				.SetExecutionMode(EExecutionMode::GameThread),
-			[](IQueryContext& Context, TypedElementRowHandle Row, const FTypedElementUObjectColumn& Actor, FTypedElementLocalTransformColumn& Transform)
+			[](IQueryContext& Context, RowHandle Row, const FTypedElementUObjectColumn& Actor, FTypedElementLocalTransformColumn& Transform)
 			{
 				if (const AActor* ActorInstance = Cast<AActor>(Actor.Object); ActorInstance != nullptr && ActorInstance->GetRootComponent() != nullptr)
 				{
diff --git a/Engine/Plugins/Experimental/EditorDataStorage/Source/TedsCore/Private/Compatibility/TypedElementActorViewportProcessors.cpp b/Engine/Plugins/Experimental/EditorDataStorage/Source/TedsCore/Private/Compatibility/TypedElementActorViewportProcessors.cpp
--- a/Engine/Plugins/Experimental/EditorDataStorage/Source/TedsCore/Private/Compatibility/TypedElementActorViewportProcessors.cpp
+++ b/Engine/Plugins/Experimental/EditorDataStorage/Source/TedsCore/Private/Compatibility/TypedElementActorViewportProcessors.cpp
@@ -21,13 +21,14 @@ namespace UE::Editor::DataStorage::Private
 			{
 				using namespace TypedElementQueryBuilder;
 				using DSI = ITypedElementDataStorageInterface;
+				using namespace UE::Editor::DataStorage;
 
 				TRACE_CPUPROFILER_EVENT_SCOPE(TEDS.Debug.AddOverlayColorToSelectionCommand);
 
 				if (ITypedElementDataStorageInterface* DataStorage = UTypedElementRegistry::GetInstance()->GetMutableDataStorage())
 				{
-					static TypedElementQueryHandle OverlayQuery = TypedElementInvalidQueryHandle;
-					if (OverlayQuery == TypedElementInvalidQueryHandle)
+					static QueryHandle OverlayQuery = InvalidQueryHandle;
+					if (OverlayQuery == InvalidQueryHandle)
 					{
 						OverlayQuery = DataStorage->RegisterQuery(
 							Select()
@@ -36,7 +37,7 @@ namespace UE::Editor::DataStorage::Private
 							.Compile());
 					}
 				
-					if (OverlayQuery == TypedElementInvalidQueryHandle)
+					if (OverlayQuery == InvalidQueryHandle)
 					{
 						return;
 					}
@@ -57,14 +58,14 @@ namespace UE::Editor::DataStorage::Private
 						return;
 					}
 
-					TArray<TypedElementRowHandle> RowHandles;
+					TArray<RowHandle> RowHandles;
 				
 					DataStorage->RunQuery(OverlayQuery, [ColorIndex, &RowHandles](const DSI::FQueryDescription&, DSI::IDirectQueryContext& Context)
 					{
 						RowHandles = Context.GetRowHandles();
 					});
 
-					for (TypedElementRowHandle Row : RowHandles)
+					for (RowHandle Row : RowHandles)
 					{
 						DataStorage->AddColumn(Row, 
 							FTypedElementViewportOutlineColorColumn{ .SelectionOutlineColorIndex = static_cast<uint8>(ColorIndex) });
@@ -80,13 +81,14 @@ namespace UE::Editor::DataStorage::Private
 			{
 				using namespace TypedElementQueryBuilder;
 				using DSI = ITypedElementDataStorageInterface;
+				using namespace UE::Editor::DataStorage;
 
 				TRACE_CPUPROFILER_EVENT_SCOPE(TEDS.Debug.AddOverlayColorToSelectionCommand);
 
 				if (ITypedElementDataStorageInterface* DataStorage = UTypedElementRegistry::GetInstance()->GetMutableDataStorage())
 				{
-					static TypedElementQueryHandle OverlayQuery = TypedElementInvalidQueryHandle;
-					if (OverlayQuery == TypedElementInvalidQueryHandle)
+					static QueryHandle OverlayQuery = InvalidQueryHandle;
+					if (OverlayQuery == InvalidQueryHandle)
 					{
 						OverlayQuery = DataStorage->RegisterQuery(
 							Select()
@@ -95,7 +97,7 @@ namespace UE::Editor::DataStorage::Private
 							.Compile());
 					}
 				
-					if (OverlayQuery == TypedElementInvalidQueryHandle)
+					if (OverlayQuery == InvalidQueryHandle)
 					{
 						return;
 					}
@@ -110,14 +112,14 @@ namespace UE::Editor::DataStorage::Private
 					FColor Color = FColor::FromHex(Args[0]);
 					Color.A = FMath::Clamp(Color.A, 0, 128);
 
-					TArray<TypedElementRowHandle> RowHandles;
+					TArray<RowHandle> RowHandles;
 				
 					DataStorage->RunQuery(OverlayQuery, [Color, &RowHandles](const DSI::FQueryDescription&, DSI::IDirectQueryContext& Context)
 					{
 						RowHandles = Context.GetRowHandles();
 					});
 
-					for (TypedElementRowHandle Row : RowHandles)
+					for (RowHandle Row : RowHandles)
 					{
 						DataStorage->RemoveColumn<FTypedElementViewportOverlayColorColumn>(Row);
 						DataStorage->AddColumn(Row, FTypedElementViewportOverlayColorColumn{ .OverlayColor = Color });
@@ -132,13 +134,14 @@ namespace UE::Editor::DataStorage::Private
 			{
 				using namespace TypedElementQueryBuilder;
 				using DSI = ITypedElementDataStorageInterface;
+				using namespace UE::Editor::DataStorage;
 
 				TRACE_CPUPROFILER_EVENT_SCOPE(TEDS.Debug.AddOverlayColorToSelectionCommand);
 
 				if (ITypedElementDataStorageInterface* DataStorage = UTypedElementRegistry::GetInstance()->GetMutableDataStorage())
 				{
-					static TypedElementQueryHandle OverlayQuery = TypedElementInvalidQueryHandle;
-					if (OverlayQuery == TypedElementInvalidQueryHandle)
+					static QueryHandle OverlayQuery = InvalidQueryHandle;
+					if (OverlayQuery == InvalidQueryHandle)
 					{
 						OverlayQuery = DataStorage->RegisterQuery(
 							Select()
@@ -147,19 +150,19 @@ namespace UE::Editor::DataStorage::Private
 							.Compile());
 					}
 				
-					if (OverlayQuery == TypedElementInvalidQueryHandle)
+					if (OverlayQuery == InvalidQueryHandle)
 					{
 						return;
 					}
 
-					TArray<TypedElementRowHandle> RowHandles;
+					TArray<RowHandle> RowHandles;
 				
 					DataStorage->RunQuery(OverlayQuery, [&RowHandles](const DSI::FQueryDescription&, DSI::IDirectQueryContext& Context)
 					{
 						RowHandles = Context.GetRowHandles();
 					});
 
-					for (TypedElementRowHandle Row : RowHandles)
+					for (RowHandle Row : RowHandles)
 					{
 						DataStorage->RemoveColumn<FTypedElementViewportOverlayColorColumn>(Row);
 					}
diff --git a/Engine/Plugins/Experimental/EditorDataStorage/Source/TedsCore/Private/Compatibility/TypedElementObjectReinstancingManager.cpp b/Engine/Plugins/Experimental/EditorDataStorage/Source/TedsCore/Private/Compatibility/TypedElementObjectReinstancingManager.cpp
--- a/Engine/Plugins/Experimental/EditorDataStorage/Source/TedsCore/Private/Compatibility/TypedElementObjectReinstancingManager.cpp
+++ b/Engine/Plugins/Experimental/EditorDataStorage/Source/TedsCore/Private/Compatibility/TypedElementObjectReinstancingManager.cpp
@@ -12,7 +12,7 @@
 DECLARE_LOG_CATEGORY_CLASS(LogTedsObjectReinstancing, Log, Log)
 
 UTedsObjectReinstancingManager::UTedsObjectReinstancingManager()
-	: MementoRowBaseTable(TypedElementInvalidTableHandle)
+	: MementoRowBaseTable(UE::Editor::DataStorage::InvalidTableHandle)
 {
 }
 
@@ -71,13 +71,15 @@ void UTedsObjectReinstancingManager::HandleOnObjectPreRemoved(
 void UTedsObjectReinstancingManager::HandleOnObjectsReinstanced(
 	const FCoreUObjectDelegates::FReplacementObjectMap& ObjectReplacementMap)
 {
+	using namespace UE::Editor::DataStorage;
+
 	ITypedElementDataStorageInterface* Interface = DataStorage;
 	for (FCoreUObjectDelegates::FReplacementObjectMap::TConstIterator Iter = ObjectReplacementMap.CreateConstIterator(); Iter; ++Iter)
 	{
 		const void* PreDeleteObject = Iter->Key;
-		if (const TypedElementRowHandle* MementoRowPtr = OldObjectToMementoMap.Find(PreDeleteObject))
+		if (const RowHandle* MementoRowPtr = OldObjectToMementoMap.Find(PreDeleteObject))
 		{
-			TypedElementRowHandle Memento = *MementoRowPtr;
+			RowHandle Memento = *MementoRowPtr;
 			
 			UObject* NewInstanceObject = Iter->Value;
 			if (NewInstanceObject == nullptr)
@@ -85,7 +87,7 @@ void UTedsObjectReinstancingManager::HandleOnObjectsReinstanced(
 				continue;
 			}
 			
-			TypedElementRowHandle NewObjectRow = DataStorageCompatibility->FindRowWithCompatibleObjectExplicit(NewInstanceObject);
+			RowHandle NewObjectRow = DataStorageCompatibility->FindRowWithCompatibleObjectExplicit(NewInstanceObject);
 			// Do the addition only if there's a recorded memento. Having a memento implies the object was previously registered and there's
 			// still an interest in it. Any other objects can therefore be ignored.
 			if (!Interface->IsRowAvailable(NewObjectRow))
diff --git a/Engine/Plugins/Experimental/EditorDataStorage/Source/TedsCore/Private/Compatibility/TypedElementUObjectWorldQueries.cpp b/Engine/Plugins/Experimental/EditorDataStorage/Source/TedsCore/Private/Compatibility/TypedElementUObjectWorldQueries.cpp
--- a/Engine/Plugins/Experimental/EditorDataStorage/Source/TedsCore/Private/Compatibility/TypedElementUObjectWorldQueries.cpp
+++ b/Engine/Plugins/Experimental/EditorDataStorage/Source/TedsCore/Private/Compatibility/TypedElementUObjectWorldQueries.cpp
@@ -17,13 +17,14 @@ void UObjectWorldDataStorageFactory::RegisterAddWorldColumn(ITypedElementDataSto
 {
 	using namespace TypedElementDataStorage;
 	using namespace TypedElementQueryBuilder;
+	using namespace UE::Editor::DataStorage;
 	
 	DataStorage.RegisterQuery(
 		Select(
 			TEXT("Add world column to UObject"),
 			FProcessor(EQueryTickPhase::PrePhysics, DataStorage.GetQueryTickGroupName(EQueryTickGroups::SyncExternalToDataStorage))
 				.SetExecutionMode(EExecutionMode::GameThread),
-			[](IQueryContext& Context, TypedElementRowHandle Row, const FTypedElementUObjectColumn& Object)
+			[](IQueryContext& Context, RowHandle Row, const FTypedElementUObjectColumn& Object)
 			{
 				// Not all objects, in particular actors, are always correctly cleaned up, resulting in dangling
 				// pointers in TEDS.
@@ -47,13 +48,14 @@ void UObjectWorldDataStorageFactory::RegisterUpdateOrRemoveWorldColumn(ITypedEle
 {
 	using namespace TypedElementDataStorage;
 	using namespace TypedElementQueryBuilder;
+	using namespace UE::Editor::DataStorage;
 	
 	DataStorage.RegisterQuery(
 		Select(
 			TEXT("Sync UObject's world to column"),
 			FProcessor(EQueryTickPhase::PostPhysics, DataStorage.GetQueryTickGroupName(EQueryTickGroups::SyncExternalToDataStorage))
 				.SetExecutionMode(EExecutionMode::GameThread),
-			[](IQueryContext& Context, TypedElementRowHandle Row, const FTypedElementUObjectColumn& Object, FTypedElementWorldColumn& World)
+			[](IQueryContext& Context, RowHandle Row, const FTypedElementUObjectColumn& Object, FTypedElementWorldColumn& World)
 			{
 				// Not all objects, in particular actors, are always correctly cleaned up, resulting in dangling
 				// pointers in TEDS.
diff --git a/Engine/Plugins/Experimental/EditorDataStorage/Source/TedsCore/Private/Debug/TypedElementDatabaseDebugCommands.cpp b/Engine/Plugins/Experimental/EditorDataStorage/Source/TedsCore/Private/Debug/TypedElementDatabaseDebugCommands.cpp
--- a/Engine/Plugins/Experimental/EditorDataStorage/Source/TedsCore/Private/Debug/TypedElementDatabaseDebugCommands.cpp
+++ b/Engine/Plugins/Experimental/EditorDataStorage/Source/TedsCore/Private/Debug/TypedElementDatabaseDebugCommands.cpp
@@ -68,11 +68,12 @@ namespace UE::Editor::DataStorage::Private
 	void PrintObjectLabels(FOutputDevice& Output)
 	{
 		using namespace TypedElementQueryBuilder;
+		using namespace UE::Editor::DataStorage;
 		using DSI = ITypedElementDataStorageInterface;
 
 		if (ITypedElementDataStorageInterface* DataStorage = UTypedElementRegistry::GetInstance()->GetMutableDataStorage())
 		{
-			static TypedElementQueryHandle LabelQuery = [DataStorage]
+			static QueryHandle LabelQuery = [DataStorage]
 			{
 				if constexpr (sizeof...(Conditions) > 0)
 				{
@@ -92,7 +93,7 @@ namespace UE::Editor::DataStorage::Private
 				}
 			}();
 
-			if (LabelQuery != TypedElementInvalidQueryHandle)
+			if (LabelQuery != InvalidQueryHandle)
 			{
 				FString Message;
 				DataStorage->RunQuery(LabelQuery, CreateDirectQueryCallbackBinding(
@@ -424,7 +425,7 @@ static FAutoConsoleCommand CVarPrintDynamicColumnWithQuery(
 		
 		const FName Identifier(*Args[0]);
 
-		const TypedElementQueryHandle Query = DataStorage->RegisterQuery(
+		const QueryHandle Query = DataStorage->RegisterQuery(
 			Select().
 				// Specify ReadOnly access to a dynamic column of type FTestDynamicColumn specified by the Identifier
 				ReadOnly<FTestDynamicColumn>(Identifier).
@@ -432,9 +433,9 @@ static FAutoConsoleCommand CVarPrintDynamicColumnWithQuery(
 
 		TStringBuilder<1024> StringBuilder;
 		TypedElementDataStorage::FQueryResult Result = DataStorage->RunQuery(Query,CreateDirectQueryCallbackBinding(
-			[Identifier, &StringBuilder](TypedElementDataStorage::IDirectQueryContext& Context, const TypedElementDataStorage::RowHandle* Rows)
+			[Identifier, &StringBuilder](TypedElementDataStorage::IDirectQueryContext& Context, const RowHandle* Rows)
 			{
-				const TArrayView<const TypedElementDataStorage::RowHandle> RowView = MakeConstArrayView(Rows, Context.GetRowCount());
+				const TArrayView<const RowHandle> RowView = MakeConstArrayView(Rows, Context.GetRowCount());
 				// Get pointer to the start of the range of columns to process
 				const FTestDynamicColumn* DynamicColumnsRangeStart = Context.GetColumn<FTestDynamicColumn>(Identifier);
 				const TArrayView<const FTestDynamicColumn> DynamicColumnView = MakeConstArrayView(DynamicColumnsRangeStart, Context.GetRowCount());
@@ -486,7 +487,7 @@ static FAutoConsoleCommand CVarCountDynamicTagWithQuery(
 				
 		const FName Identifier(*Args[0]);
 
-		const TypedElementQueryHandle Query = DataStorage->RegisterQuery(
+		const QueryHandle Query = DataStorage->RegisterQuery(
 			Select().
 				// Specify Any access to a dynamically created tag
 				Where().
@@ -494,7 +495,7 @@ static FAutoConsoleCommand CVarCountDynamicTagWithQuery(
 			Compile());
 
 		int32 Count = 0;
-		DataStorage->RunQuery(Query, CreateDirectQueryCallbackBinding([Identifier, &Count](TypedElementDataStorage::IDirectQueryContext& Context, const TypedElementDataStorage::RowHandle* Rows)
+		DataStorage->RunQuery(Query, CreateDirectQueryCallbackBinding([Identifier, &Count](TypedElementDataStorage::IDirectQueryContext& Context, const RowHandle* Rows)
 		{
 			Count += Context.GetRowCount();
 		}));
@@ -658,7 +659,7 @@ static FAutoConsoleCommand CVarMatchDynamicTag(
 			return;
 		}			
 
-		const TypedElementQueryHandle Query = [&Args, DataStorage]() -> TypedElementQueryHandle
+		const QueryHandle Query = [&Args, DataStorage]() -> QueryHandle
 		{
 			const FName Tag(*Args[0]);
 			if (Args.Num() == 1)
@@ -688,7 +689,7 @@ static FAutoConsoleCommand CVarMatchDynamicTag(
 		uint64 Count = 0;
 		
 		const TypedElementDataStorage::FQueryResult Result = DataStorage->RunQuery(Query, CreateDirectQueryCallbackBinding(
-			[&Count](const DSI::IDirectQueryContext& Context, const TypedElementRowHandle*)
+			[&Count](const DSI::IDirectQueryContext& Context, const RowHandle*)
 			{
 				Count += Context.GetRowCount();
 			}));
@@ -703,6 +704,8 @@ static FAutoConsoleCommand CVarAddDynamicTagFromEnum(
 	TEXT("Argument: Row, EnumValue\n"),
 	FConsoleCommandWithArgsDelegate::CreateLambda([](const TArray<FString>& Args)
 	{
+		using namespace UE::Editor::DataStorage;
+
 		ITypedElementDataStorageInterface* DataStorage = UTypedElementRegistry::GetInstance()->GetMutableDataStorage();
 
 		if (Args.Num() < 1 || Args.Num() > 2)
@@ -711,7 +714,7 @@ static FAutoConsoleCommand CVarAddDynamicTagFromEnum(
 		}
 		
 		const uint64 RowAsU64 = FCString::Strtoui64(*Args[0], nullptr, 10);
-		const TypedElementDataStorage::RowHandle Row = RowAsU64;
+		const RowHandle Row = RowAsU64;
 
 		if (Args.Num() == 1)
 		{
@@ -740,6 +743,7 @@ static FAutoConsoleCommand CVarRemoveDynamicTagFromEnum(
 	TEXT("Argument: Row\n"),
 	FConsoleCommandWithArgsDelegate::CreateLambda([](const TArray<FString>& Args)
 	{
+		using namespace UE::Editor::DataStorage;
 		ITypedElementDataStorageInterface* DataStorage = UTypedElementRegistry::GetInstance()->GetMutableDataStorage();
 
 		if (Args.Num() != 1)
@@ -748,7 +752,7 @@ static FAutoConsoleCommand CVarRemoveDynamicTagFromEnum(
 		}
 		
 		const uint64 RowAsU64 = FCString::Strtoui64(*Args[0], nullptr, 10);
-		const TypedElementDataStorage::RowHandle Row = RowAsU64;
+		const RowHandle Row = RowAsU64;
 		
 		DataStorage->RemoveColumn<ETedsDebugEnum>(Row);
 	}),
@@ -783,7 +787,7 @@ static FAutoConsoleCommand CVarMatchDynamicTagFromEnum(
 			}
 		}
 
-		const TypedElementQueryHandle Query = [&Args, DataStorage]() -> TypedElementQueryHandle
+		const QueryHandle Query = [&Args, DataStorage]() -> QueryHandle
 		{
 			if (Args.Num() == 0)
 			{
@@ -812,10 +816,10 @@ static FAutoConsoleCommand CVarMatchDynamicTagFromEnum(
 			}
 			else
 			{
-				return TypedElementDataStorage::InvalidQueryHandle;
+				return InvalidQueryHandle;
 			}
 		}();
-		if (Query == TypedElementDataStorage::InvalidQueryHandle)
+		if (Query == InvalidQueryHandle)
 		{
 			UE_LOG(LogEditorDataStorage, Error, TEXT("Invalid number of arguments"));
 			return;
@@ -824,7 +828,7 @@ static FAutoConsoleCommand CVarMatchDynamicTagFromEnum(
 		uint64 Count = 0;
 		
 		const TypedElementDataStorage::FQueryResult Result = DataStorage->RunQuery(Query, CreateDirectQueryCallbackBinding(
-			[&Count](const DSI::IDirectQueryContext& Context, const TypedElementRowHandle*)
+			[&Count](const DSI::IDirectQueryContext& Context, const RowHandle*)
 			{
 				Count += Context.GetRowCount();
 			}));
diff --git a/Engine/Plugins/Experimental/EditorDataStorage/Source/TedsCore/Private/Memento/TypedElementMementoSystem.cpp b/Engine/Plugins/Experimental/EditorDataStorage/Source/TedsCore/Private/Memento/TypedElementMementoSystem.cpp
--- a/Engine/Plugins/Experimental/EditorDataStorage/Source/TedsCore/Private/Memento/TypedElementMementoSystem.cpp
+++ b/Engine/Plugins/Experimental/EditorDataStorage/Source/TedsCore/Private/Memento/TypedElementMementoSystem.cpp
@@ -38,10 +38,8 @@ namespace UE::Editor::DataStorage
 		}
 	}
 
-	TypedElementRowHandle FMementoSystem::CreateMemento(TypedElementDataStorage::RowHandle SourceRow)
+	RowHandle FMementoSystem::CreateMemento(RowHandle SourceRow)
 	{
-		using namespace TypedElementDataStorage;
-
 		FScopedSharedLock Lock(EGlobalLockScope::Public);
 
 		RowHandle MementoRow = DataStorage.AddRow(MementoRowBaseTable);
@@ -49,7 +47,7 @@ namespace UE::Editor::DataStorage
 		return MementoRow;
 	}
 
-	void FMementoSystem::CreateMemento(TypedElementDataStorage::RowHandle ReservedMementoRow, TypedElementDataStorage::RowHandle SourceRow)
+	void FMementoSystem::CreateMemento(RowHandle ReservedMementoRow, RowHandle SourceRow)
 	{
 		FScopedSharedLock Lock(EGlobalLockScope::Public);
 	
@@ -57,10 +55,8 @@ namespace UE::Editor::DataStorage
 		CreateMementoInternal(ReservedMementoRow, SourceRow);
 	}
 
-	void FMementoSystem::CreateMementoInternal(TypedElementDataStorage::RowHandle MementoRow, TypedElementDataStorage::RowHandle SourceRow)
-	{
-		using namespace TypedElementDataStorage;
-	
+	void FMementoSystem::CreateMementoInternal(RowHandle MementoRow, RowHandle SourceRow)
+	{	
 		for (const UTedsMementoTranslatorBase* Translator : MementoTranslators)
 		{
 			if (void* SourceColumn = DataStorage.GetColumnData(SourceRow, Translator->GetColumnType()))
@@ -84,7 +80,7 @@ namespace UE::Editor::DataStorage
 		}
 	}
 
-	void FMementoSystem::RestoreMemento(TypedElementDataStorage::RowHandle MementoRow, TypedElementDataStorage::RowHandle TargetRow)
+	void FMementoSystem::RestoreMemento(RowHandle MementoRow, RowHandle TargetRow)
 	{
 		FScopedSharedLock Lock(EGlobalLockScope::Public);
 
@@ -111,7 +107,7 @@ namespace UE::Editor::DataStorage
 		}
 	}
 
-	void FMementoSystem::DestroyMemento(TypedElementDataStorage::RowHandle MementoRow)
+	void FMementoSystem::DestroyMemento(RowHandle MementoRow)
 	{
 		// No need to lock this as no internal data is used.
 
diff --git a/Engine/Plugins/Experimental/EditorDataStorage/Source/TedsCore/Private/Memento/TypedElementMementoSystem.h b/Engine/Plugins/Experimental/EditorDataStorage/Source/TedsCore/Private/Memento/TypedElementMementoSystem.h
--- a/Engine/Plugins/Experimental/EditorDataStorage/Source/TedsCore/Private/Memento/TypedElementMementoSystem.h
+++ b/Engine/Plugins/Experimental/EditorDataStorage/Source/TedsCore/Private/Memento/TypedElementMementoSystem.h
@@ -15,16 +15,16 @@ namespace UE::Editor::DataStorage
 	public:
 		explicit FMementoSystem(ITypedElementDataStorageInterface& InDataStorage);
 
-		TypedElementDataStorage::RowHandle CreateMemento(TypedElementDataStorage::RowHandle SourceRow);
-		void CreateMemento(TypedElementDataStorage::RowHandle ReservedMementoRow, TypedElementDataStorage::RowHandle SourceRow);
-		void RestoreMemento(TypedElementDataStorage::RowHandle MementoRow, TypedElementDataStorage::RowHandle TargetRow);
-		void DestroyMemento(TypedElementDataStorage::RowHandle MementoRow);
+		RowHandle CreateMemento(RowHandle SourceRow);
+		void CreateMemento(RowHandle ReservedMementoRow, RowHandle SourceRow);
+		void RestoreMemento(RowHandle MementoRow, RowHandle TargetRow);
+		void DestroyMemento(RowHandle MementoRow);
 
 	private:
-		void CreateMementoInternal(TypedElementDataStorage::RowHandle MementoRow, TypedElementDataStorage::RowHandle SourceRow);
+		void CreateMementoInternal(RowHandle MementoRow, RowHandle SourceRow);
 
 		TArray<const UTedsMementoTranslatorBase*> MementoTranslators;
-		TypedElementDataStorage::TableHandle MementoRowBaseTable;
+		TableHandle MementoRowBaseTable;
 		ITypedElementDataStorageInterface& DataStorage;
 	};
 } // namespace UE::Editor::DataStorage
diff --git a/Engine/Plugins/Experimental/EditorDataStorage/Source/TedsCore/Private/Processors/TypedElementProcessorAdaptors.cpp b/Engine/Plugins/Experimental/EditorDataStorage/Source/TedsCore/Private/Processors/TypedElementProcessorAdaptors.cpp
--- a/Engine/Plugins/Experimental/EditorDataStorage/Source/TedsCore/Private/Processors/TypedElementProcessorAdaptors.cpp
+++ b/Engine/Plugins/Experimental/EditorDataStorage/Source/TedsCore/Private/Processors/TypedElementProcessorAdaptors.cpp
@@ -30,13 +30,13 @@ namespace UE::Editor::DataStorage
 				return Context.GetNumEntities();
 			}
 
-			TConstArrayView<TypedElementRowHandle> GetRowHandles() const
+			TConstArrayView<RowHandle> GetRowHandles() const
 			{
 				static_assert(
-					sizeof(TypedElementRowHandle) == sizeof(FMassEntityHandle) && alignof(TypedElementRowHandle) == alignof(FMassEntityHandle),
-					"TypedElementRowHandle and FMassEntityHandle need to by layout compatible to support TEDS.");
+					sizeof(RowHandle) == sizeof(FMassEntityHandle) && alignof(RowHandle) == alignof(FMassEntityHandle),
+					"UE::Editor::DataStorage::RowHandle and FMassEntityHandle must be layout compatible.");
 				TConstArrayView<FMassEntityHandle> Entities = Context.GetEntities();
-				return TConstArrayView<TypedElementRowHandle>(reinterpret_cast<const TypedElementRowHandle*>(Entities.GetData()), Entities.Num());
+				return TConstArrayView<RowHandle>(reinterpret_cast<const RowHandle*>(Entities.GetData()), Entities.Num());
 			}
 	
 			const void* GetColumn(const UScriptStruct* ColumnType) const
@@ -95,7 +95,7 @@ namespace UE::Editor::DataStorage
 				return false;
 			}
 
-			bool HasColumn(TypedElementDataStorage::RowHandle Row, const UScriptStruct* ColumnType) const
+			bool HasColumn(RowHandle Row, const UScriptStruct* ColumnType) const
 			{
 				FMassEntityHandle Entity = FMassEntityHandle::FromNumber(Row);
 				FMassEntityManager& Manager = Context.GetEntityManagerChecked();
@@ -183,12 +183,12 @@ namespace UE::Editor::DataStorage
 				return Environment.GetUpdateCycleId();
 			}
 
-			bool IsRowAvailable(TypedElementDataStorage::RowHandle Row) const
+			bool IsRowAvailable(RowHandle Row) const
 			{
 				return Environment.GetMassEntityManager().IsEntityValid(FMassEntityHandle::FromNumber(Row));
 			}
 
-			bool IsRowAssigned(TypedElementDataStorage::RowHandle Row) const
+			bool IsRowAssigned(RowHandle Row) const
 			{
 				return Environment.GetMassEntityManager().IsEntityActive(FMassEntityHandle::FromNumber(Row));
 			}
@@ -268,7 +268,7 @@ namespace UE::Editor::DataStorage
 					});
 			}
 
-			void* AddColumnUninitialized(TypedElementDataStorage::RowHandle Row, const UScriptStruct* ObjectType) 
+			void* AddColumnUninitialized(RowHandle Row, const UScriptStruct* ObjectType) 
 			{
 				return AddColumnUninitialized(Row, ObjectType,
 					[](const UScriptStruct& TypeInfo, void* Destination, void* Source)
@@ -277,7 +277,7 @@ namespace UE::Editor::DataStorage
 					});
 			}
 
-			void* AddColumnUninitialized(TypedElementDataStorage::RowHandle Row, const UScriptStruct* ObjectType, ObjectCopyOrMove Relocator)
+			void* AddColumnUninitialized(RowHandle Row, const UScriptStruct* ObjectType, ObjectCopyOrMove Relocator)
 			{
 				checkf(ObjectType->IsChildOf(FMassFragment::StaticStruct()), TEXT("Column [%s] can not be a tag"), *ObjectType->GetName());
 		
@@ -411,7 +411,7 @@ namespace UE::Editor::DataStorage
 				return ColumnData;
 			}
 	
-			void AddColumns(TypedElementDataStorage::RowHandle Row, TConstArrayView<const UScriptStruct*> ColumnTypes) 
+			void AddColumns(RowHandle Row, TConstArrayView<const UScriptStruct*> ColumnTypes) 
 			{
 				struct FAddedColumns
 				{
@@ -433,7 +433,7 @@ namespace UE::Editor::DataStorage
 					});
 			}
 
-			void AddColumns(TConstArrayView<TypedElementDataStorage::RowHandle> Rows, TConstArrayView<const UScriptStruct*> ColumnTypes) 
+			void AddColumns(TConstArrayView<RowHandle> Rows, TConstArrayView<const UScriptStruct*> ColumnTypes) 
 			{
 				struct FAddedColumns
 				{
@@ -448,7 +448,7 @@ namespace UE::Editor::DataStorage
 		
 				FMassEntityHandle* Entities = ScratchBuffer.EmplaceArray<FMassEntityHandle>(Rows.Num());
 				AddedColumns->Entities = Entities;
-				for (TypedElementRowHandle Row : Rows)
+				for (RowHandle Row : Rows)
 				{
 					*Entities = FMassEntityHandle::FromNumber(Row);
 					Entities++;
@@ -470,7 +470,7 @@ namespace UE::Editor::DataStorage
 					});
 			}
 
-			void RemoveColumns(TypedElementDataStorage::RowHandle Row, TConstArrayView<const UScriptStruct*> ColumnTypes) 
+			void RemoveColumns(RowHandle Row, TConstArrayView<const UScriptStruct*> ColumnTypes) 
 			{
 				struct FRemovedColumns
 				{
@@ -495,7 +495,7 @@ namespace UE::Editor::DataStorage
 				}
 			}
 
-			void RemoveColumns(TConstArrayView<TypedElementDataStorage::RowHandle> Rows, TConstArrayView<const UScriptStruct*> ColumnTypes) 
+			void RemoveColumns(TConstArrayView<RowHandle> Rows, TConstArrayView<const UScriptStruct*> ColumnTypes) 
 			{
 				struct FRemovedColumns
 				{
@@ -510,7 +510,7 @@ namespace UE::Editor::DataStorage
 
 				FMassEntityHandle* Entities = ScratchBuffer.EmplaceArray<FMassEntityHandle>(Rows.Num());
 				RemovedColumns->Entities = Entities;
-				for (TypedElementRowHandle Row : Rows)
+				for (RowHandle Row : Rows)
 				{
 					*Entities = FMassEntityHandle::FromNumber(Row);
 					Entities++;
@@ -562,14 +562,14 @@ namespace UE::Editor::DataStorage
 					});
 			}
 	
-			TypedElementDataStorage::RowHandle AddRow(TypedElementDataStorage::TableHandle Table) 
+			RowHandle AddRow(TableHandle Table) 
 			{
 				FMassEntityHandle EntityHandle = Environment.GetMassEntityManager().ReserveEntity();
 				FMassArchetypeHandle ArchetypeHandle = Environment.LookupMassArchetype(Table);
 
 				if (!ArchetypeHandle.IsValid())
 				{
-					return TypedElementDataStorage::InvalidRowHandle;
+					return InvalidRowHandle;
 				}
 		
 				struct CommandInfo
@@ -594,19 +594,19 @@ namespace UE::Editor::DataStorage
 				return TedsRowHandle;
 			}
 
-			void RemoveRow(TypedElementRowHandle Row) 
+			void RemoveRow(RowHandle Row) 
 			{
 				this->Context.Defer().DestroyEntity(FMassEntityHandle::FromNumber(Row));
 			}
 
-			void RemoveRows(TConstArrayView<TypedElementRowHandle> Rows) 
+			void RemoveRows(TConstArrayView<RowHandle> Rows) 
 			{
 				// Row handles and entities map 1:1 for data, so a reintpret_cast can be safely done to avoid
 				// having to allocate memory and iterating over the rows.
 
-				static_assert(sizeof(FMassEntityHandle) == sizeof(TypedElementRowHandle), 
+				static_assert(sizeof(FMassEntityHandle) == sizeof(RowHandle), 
 					"Sizes of mass entity and data storage row have gone out of sync.");
-				static_assert(alignof(FMassEntityHandle) == alignof(TypedElementRowHandle),
+				static_assert(alignof(FMassEntityHandle) == alignof(RowHandle),
 					"Alignment of mass entity and data storage row have gone out of sync.");
 
 				this->Context.Defer().DestroyEntities(
@@ -621,13 +621,13 @@ namespace UE::Editor::DataStorage
 			{}
 	
 			uint32 GetRowCount() const override { return Implementation.GetRowCount(); }
-			TConstArrayView<TypedElementDataStorage::RowHandle> GetRowHandles() const override { return Implementation.GetRowHandles(); }
+			TConstArrayView<RowHandle> GetRowHandles() const override { return Implementation.GetRowHandles(); }
 			const void* GetColumn(const UScriptStruct* ColumnType) const override { return Implementation.GetColumn(ColumnType); }
 			void* GetMutableColumn(const UScriptStruct* ColumnType) override { return Implementation.GetMutableColumn(ColumnType); }
 			void GetColumns(TArrayView<char*> RetrievedAddresses, TConstArrayView<TWeakObjectPtr<const UScriptStruct>> ColumnTypes, TConstArrayView<TypedElementDataStorage::EQueryAccessType> AccessTypes) override { return Implementation.GetColumns(RetrievedAddresses, ColumnTypes, AccessTypes); }
 			void GetColumnsUnguarded(int32 TypeCount, char** RetrievedAddresses, const TWeakObjectPtr<const UScriptStruct>* ColumnTypes, const TypedElementDataStorage::EQueryAccessType* AccessTypes) override { return Implementation.GetColumnsUnguarded(TypeCount, RetrievedAddresses, ColumnTypes, AccessTypes); }
 			bool HasColumn(const UScriptStruct* ColumnType) const override { return Implementation.HasColumn(ColumnType); }
-			bool HasColumn(TypedElementDataStorage::RowHandle Row, const UScriptStruct* ColumnType) const override { return Implementation.HasColumn(Row, ColumnType); }
+			bool HasColumn(RowHandle Row, const UScriptStruct* ColumnType) const override { return Implementation.HasColumn(Row, ColumnType); }
 			virtual const UScriptStruct* FindDynamicColumnType(const UE::Editor::DataStorage::FDynamicColumnDescription& Description) const override { return Implementation.FindDynamicColumnType(Description); }
 
 			FMassContextCommon Implementation;
@@ -641,29 +641,29 @@ namespace UE::Editor::DataStorage
 
 			~FMassSubqueryContextForwarder() override = default;
 			uint32 GetRowCount() const override { return Implementation.GetRowCount(); }
-			TConstArrayView<TypedElementDataStorage::RowHandle> GetRowHandles() const override { return Implementation.GetRowHandles(); }
+			TConstArrayView<RowHandle> GetRowHandles() const override { return Implementation.GetRowHandles(); }
 			const void* GetColumn(const UScriptStruct* ColumnType) const override { return Implementation.GetColumn(ColumnType); }
 			void* GetMutableColumn(const UScriptStruct* ColumnType) override { return Implementation.GetMutableColumn(ColumnType); }
 			void GetColumns(TArrayView<char*> RetrievedAddresses, TConstArrayView<TWeakObjectPtr<const UScriptStruct>> ColumnTypes, TConstArrayView<TypedElementDataStorage::EQueryAccessType> AccessTypes) override { return Implementation.GetColumns(RetrievedAddresses, ColumnTypes, AccessTypes); }
 			void GetColumnsUnguarded(int32 TypeCount, char** RetrievedAddresses, const TWeakObjectPtr<const UScriptStruct>* ColumnTypes, const TypedElementDataStorage::EQueryAccessType* AccessTypes) override { return Implementation.GetColumnsUnguarded(TypeCount, RetrievedAddresses, ColumnTypes, AccessTypes); }
 			bool HasColumn(const UScriptStruct* ColumnType) const override { return Implementation.HasColumn(ColumnType); }
-			bool HasColumn(TypedElementDataStorage::RowHandle Row, const UScriptStruct* ColumnType) const override { return Implementation.HasColumn(Row, ColumnType); }
+			bool HasColumn(RowHandle Row, const UScriptStruct* ColumnType) const override { return Implementation.HasColumn(Row, ColumnType); }
 			uint64 GetUpdateCycleId() const override { return Implementation.GetUpdateCycleId(); }
-			bool IsRowAvailable(TypedElementDataStorage::RowHandle Row) const override { return Implementation.IsRowAvailable(Row); }
-			bool IsRowAssigned(TypedElementDataStorage::RowHandle Row) const override { return Implementation.IsRowAssigned(Row); }
+			bool IsRowAvailable(RowHandle Row) const override { return Implementation.IsRowAvailable(Row); }
+			bool IsRowAssigned(RowHandle Row) const override { return Implementation.IsRowAssigned(Row); }
 			void ActivateQueries(FName ActivationName) override  { return Implementation.ActivateQueries(ActivationName); }
-			TypedElementDataStorage::RowHandle AddRow(TypedElementDataStorage::TableHandle Table) override  { return Implementation.AddRow(Table); }
-			void RemoveRow(TypedElementDataStorage::RowHandle Row) override { return Implementation.RemoveRow(Row); }
-			void RemoveRows(TConstArrayView<TypedElementDataStorage::RowHandle> Rows) override  { return Implementation.RemoveRows(Rows); }
-			void AddColumns(TypedElementDataStorage::RowHandle Row, TConstArrayView<const UScriptStruct*> ColumnTypes) override { return Implementation.AddColumns(Row, ColumnTypes); }
-			void AddColumns(TConstArrayView<TypedElementDataStorage::RowHandle> Rows, TConstArrayView<const UScriptStruct*> ColumnTypes) override  { return Implementation.AddColumns(Rows, ColumnTypes); }
+			RowHandle AddRow(TableHandle Table) override  { return Implementation.AddRow(Table); }
+			void RemoveRow(RowHandle Row) override { return Implementation.RemoveRow(Row); }
+			void RemoveRows(TConstArrayView<RowHandle> Rows) override  { return Implementation.RemoveRows(Rows); }
+			void AddColumns(RowHandle Row, TConstArrayView<const UScriptStruct*> ColumnTypes) override { return Implementation.AddColumns(Row, ColumnTypes); }
+			void AddColumns(TConstArrayView<RowHandle> Rows, TConstArrayView<const UScriptStruct*> ColumnTypes) override  { return Implementation.AddColumns(Rows, ColumnTypes); }
 			void AddColumns(TConstArrayView<TypedElementDataStorage::RowHandle> Rows, TConstArrayView<UE::Editor::DataStorage::FDynamicColumnDescription> DynamicColumnDescriptions) override { Implementation.AddColumns(Rows, DynamicColumnDescriptions); }
-			void* AddColumnUninitialized(TypedElementDataStorage::RowHandle Row, const UScriptStruct* ColumnType) override  { return Implementation.AddColumnUninitialized(Row, ColumnType); }
-			void* AddColumnUninitialized(TypedElementDataStorage::RowHandle Row, const UScriptStruct* ObjectType, ObjectCopyOrMove Relocator) override { return Implementation.AddColumnUninitialized(Row, ObjectType, Relocator); }
+			void* AddColumnUninitialized(RowHandle Row, const UScriptStruct* ColumnType) override  { return Implementation.AddColumnUninitialized(Row, ColumnType); }
+			void* AddColumnUninitialized(RowHandle Row, const UScriptStruct* ObjectType, ObjectCopyOrMove Relocator) override { return Implementation.AddColumnUninitialized(Row, ObjectType, Relocator); }
 			void* AddColumnUninitialized(TypedElementDataStorage::RowHandle Row, const UE::Editor::DataStorage::FDynamicColumnDescription& DynamicColumnDescription) override { return Implementation.AddColumnUninitialized(Row, DynamicColumnDescription); }
 			void* AddColumnUninitialized(TypedElementDataStorage::RowHandle Row, const UE::Editor::DataStorage::FDynamicColumnDescription& DynamicColumnDescription, ObjectCopyOrMove Relocator) override { return Implementation.AddColumnUninitialized(Row, DynamicColumnDescription, Relocator); }
-			void RemoveColumns(TypedElementDataStorage::RowHandle Row, TConstArrayView<const UScriptStruct*> ColumnTypes) override { Implementation.RemoveColumns(Row, ColumnTypes); }
-			void RemoveColumns(TConstArrayView<TypedElementDataStorage::RowHandle> Rows, TConstArrayView<const UScriptStruct*> ColumnTypes) override { Implementation.RemoveColumns(Rows, ColumnTypes); }
+			void RemoveColumns(RowHandle Row, TConstArrayView<const UScriptStruct*> ColumnTypes) override { Implementation.RemoveColumns(Row, ColumnTypes); }
+			void RemoveColumns(TConstArrayView<RowHandle> Rows, TConstArrayView<const UScriptStruct*> ColumnTypes) override { Implementation.RemoveColumns(Rows, ColumnTypes); }
 			const UScriptStruct* FindDynamicColumnType(const UE::Editor::DataStorage::FDynamicColumnDescription& Description) const override { return Implementation.FindDynamicColumnType(Description); }
 			
 
@@ -721,17 +721,15 @@ namespace UE::Editor::DataStorage
 				}
 			}
 	
-			TypedElementDataStorage::RowHandle FindIndexedRow(TypedElementDataStorage::IndexHash Index) const
+			RowHandle FindIndexedRow(TypedElementDataStorage::IndexHash Index) const
 			{
-				using namespace UE::Editor::DataStorage;
-
 				EGlobalLockScope Scope = FGlobalLock::GetLockStatus(EGlobalLockScope::Internal) == EGlobalLockStatus::Unlocked
 					? EGlobalLockScope::Public // There's no internal lock so use a public lock instead.
 					: EGlobalLockScope::Internal; // There's an internal lock set so use that.
 				return Environment.GetIndexTable().FindIndexedRow(Scope, Index);
 			}
 
-			TypedElementDataStorage::FQueryResult RunQuery(TypedElementQueryHandle Query)
+			TypedElementDataStorage::FQueryResult RunQuery(QueryHandle Query)
 			{
 				const FExtendedQueryStore::Handle Handle(Query);
 				// This can be safely called because there's not callback, which means no columns are accessed, even for select queries.
@@ -749,7 +747,7 @@ namespace UE::Editor::DataStorage
 			{
 				if (SubqueryIndex < QueryDescription.Subqueries.Num())
 				{
-					const TypedElementQueryHandle SubqueryHandle = QueryDescription.Subqueries[SubqueryIndex];
+					const QueryHandle SubqueryHandle = QueryDescription.Subqueries[SubqueryIndex];
 					const FExtendedQueryStore::Handle StorageHandle(SubqueryHandle);
 					return QueryStore.RunQuery(Context.GetEntityManagerChecked(), Environment, Context, StorageHandle, Callback);
 				}
@@ -764,7 +762,7 @@ namespace UE::Editor::DataStorage
 			{
 				if (SubqueryIndex < QueryDescription.Subqueries.Num())
 				{
-					const TypedElementQueryHandle SubqueryHandle = QueryDescription.Subqueries[SubqueryIndex];
+					const QueryHandle SubqueryHandle = QueryDescription.Subqueries[SubqueryIndex];
 					const FExtendedQueryStore::Handle StorageHandle(SubqueryHandle);
 					return QueryStore.RunQuery(Context.GetEntityManagerChecked(), Environment, Context, StorageHandle, Row, Callback);
 				}
@@ -789,41 +787,40 @@ namespace UE::Editor::DataStorage
 			{}
 	
 			uint32 GetRowCount() const override { return Implementation.GetRowCount(); }
-			TConstArrayView<TypedElementDataStorage::RowHandle> GetRowHandles() const override { return Implementation.GetRowHandles(); }
+			TConstArrayView<RowHandle> GetRowHandles() const override { return Implementation.GetRowHandles(); }
 			const void* GetColumn(const UScriptStruct* ColumnType) const override { return Implementation.GetColumn(ColumnType); }
 			void* GetMutableColumn(const UScriptStruct* ColumnType) override { return Implementation.GetMutableColumn(ColumnType); }
 			void GetColumns(TArrayView<char*> RetrievedAddresses, TConstArrayView<TWeakObjectPtr<const UScriptStruct>> ColumnTypes, TConstArrayView<TypedElementDataStorage::EQueryAccessType> AccessTypes) override { return Implementation.GetColumns(RetrievedAddresses, ColumnTypes, AccessTypes); }
 			void GetColumnsUnguarded(int32 TypeCount, char** RetrievedAddresses, const TWeakObjectPtr<const UScriptStruct>* ColumnTypes, const TypedElementDataStorage::EQueryAccessType* AccessTypes) override { return Implementation.GetColumnsUnguarded(TypeCount, RetrievedAddresses, ColumnTypes, AccessTypes); }
 			bool HasColumn(const UScriptStruct* ColumnType) const override { return Implementation.HasColumn(ColumnType); }
-			bool HasColumn(TypedElementDataStorage::RowHandle Row, const UScriptStruct* ColumnType) const override { return Implementation.HasColumn(Row, ColumnType); }
+			bool HasColumn(RowHandle Row, const UScriptStruct* ColumnType) const override { return Implementation.HasColumn(Row, ColumnType); }
 			uint64 GetUpdateCycleId() const override { return Implementation.GetUpdateCycleId(); }
-			bool IsRowAvailable(TypedElementDataStorage::RowHandle Row) const override { return Implementation.IsRowAvailable(Row); }
-			bool IsRowAssigned(TypedElementDataStorage::RowHandle Row) const override { return Implementation.IsRowAssigned(Row); }
+			bool IsRowAvailable(RowHandle Row) const override { return Implementation.IsRowAvailable(Row); }
+			bool IsRowAssigned(RowHandle Row) const override { return Implementation.IsRowAssigned(Row); }
 			void ActivateQueries(FName ActivationName) override  { return Implementation.ActivateQueries(ActivationName); }
-			TypedElementDataStorage::RowHandle AddRow(TypedElementDataStorage::TableHandle Table) override  { return Implementation.AddRow(Table); }
-			void RemoveRow(TypedElementDataStorage::RowHandle Row) override { return Implementation.RemoveRow(Row); }
-			void RemoveRows(TConstArrayView<TypedElementDataStorage::RowHandle> Rows) override  { return Implementation.RemoveRows(Rows); }
-			void AddColumns(TypedElementDataStorage::RowHandle Row, TConstArrayView<const UScriptStruct*> ColumnTypes) override { return Implementation.AddColumns(Row, ColumnTypes); }
-			void AddColumns(TConstArrayView<TypedElementDataStorage::RowHandle> Rows, TConstArrayView<const UScriptStruct*> ColumnTypes) override  { return Implementation.AddColumns(Rows, ColumnTypes); }
+			RowHandle AddRow(TableHandle Table) override  { return Implementation.AddRow(Table); }
+			void RemoveRow(RowHandle Row) override { return Implementation.RemoveRow(Row); }
+			void RemoveRows(TConstArrayView<RowHandle> Rows) override  { return Implementation.RemoveRows(Rows); }
+			void AddColumns(RowHandle Row, TConstArrayView<const UScriptStruct*> ColumnTypes) override { return Implementation.AddColumns(Row, ColumnTypes); }
+			void AddColumns(TConstArrayView<RowHandle> Rows, TConstArrayView<const UScriptStruct*> ColumnTypes) override  { return Implementation.AddColumns(Rows, ColumnTypes); }
 			void AddColumns(TConstArrayView<TypedElementDataStorage::RowHandle> Rows, TConstArrayView<UE::Editor::DataStorage::FDynamicColumnDescription> DynamicColumnDescriptions) override { return Implementation.AddColumns(Rows, DynamicColumnDescriptions); }
-			void* AddColumnUninitialized(TypedElementDataStorage::RowHandle Row, const UScriptStruct* ColumnType) override  { return Implementation.AddColumnUninitialized(Row, ColumnType); }
-			void* AddColumnUninitialized(TypedElementDataStorage::RowHandle Row, const UScriptStruct* ObjectType, ObjectCopyOrMove Relocator) override { return Implementation.AddColumnUninitialized(Row, ObjectType, Relocator); }
+			void* AddColumnUninitialized(RowHandle Row, const UScriptStruct* ColumnType) override  { return Implementation.AddColumnUninitialized(Row, ColumnType); }
+			void* AddColumnUninitialized(RowHandle Row, const UScriptStruct* ObjectType, ObjectCopyOrMove Relocator) override { return Implementation.AddColumnUninitialized(Row, ObjectType, Relocator); }
 			void* AddColumnUninitialized(TypedElementDataStorage::RowHandle Row, const FDynamicColumnDescription& DynamicColumnDescription, ObjectCopyOrMove Relocator) override { return Implementation.AddColumnUninitialized(Row, DynamicColumnDescription, Relocator); }
 			void* AddColumnUninitialized(TypedElementDataStorage::RowHandle Row, const FDynamicColumnDescription& DynamicColumnDescription) override { return Implementation.AddColumnUninitialized(Row, DynamicColumnDescription); };
-			void RemoveColumns(TypedElementDataStorage::RowHandle Row, TConstArrayView<const UScriptStruct*> ColumnTypes) override { Implementation.RemoveColumns(Row, ColumnTypes); }
-			void RemoveColumns(TConstArrayView<TypedElementDataStorage::RowHandle> Rows, TConstArrayView<const UScriptStruct*> ColumnTypes) override { Implementation.RemoveColumns(Rows, ColumnTypes); }
+			void RemoveColumns(RowHandle Row, TConstArrayView<const UScriptStruct*> ColumnTypes) override { Implementation.RemoveColumns(Row, ColumnTypes); }
+			void RemoveColumns(TConstArrayView<RowHandle> Rows, TConstArrayView<const UScriptStruct*> ColumnTypes) override { Implementation.RemoveColumns(Rows, ColumnTypes); }
 			const UScriptStruct* FindDynamicColumnType(const UE::Editor::DataStorage::FDynamicColumnDescription& Description) const override { return Implementation.FindDynamicColumnType(Description); }
 
 			const UObject* GetDependency(const UClass* DependencyClass) override { return Implementation.GetDependency(DependencyClass); }
 			UObject* GetMutableDependency(const UClass* DependencyClass) override { return Implementation.GetMutableDependency(DependencyClass); }
 			void GetDependencies(TArrayView<UObject*> RetrievedAddresses, TConstArrayView<TWeakObjectPtr<const UClass>> DependencyTypes, TConstArrayView<TypedElementDataStorage::EQueryAccessType> AccessTypes) override { return Implementation.GetDependencies(RetrievedAddresses, DependencyTypes, AccessTypes); }
-			TypedElementDataStorage::RowHandle FindIndexedRow(TypedElementDataStorage::IndexHash Index) const override { return Implementation.FindIndexedRow(Index); }
-			TypedElementDataStorage::FQueryResult RunQuery(TypedElementDataStorage::QueryHandle Query) override { return Implementation.RunQuery(Query); }
+			RowHandle FindIndexedRow(TypedElementDataStorage::IndexHash Index) const override { return Implementation.FindIndexedRow(Index); }
+			TypedElementDataStorage::FQueryResult RunQuery(QueryHandle Query) override { return Implementation.RunQuery(Query); }
 			TypedElementDataStorage::FQueryResult RunSubquery(int32 SubqueryIndex) override { return Implementation.RunSubquery(SubqueryIndex); }
 			TypedElementDataStorage::FQueryResult RunSubquery(int32 SubqueryIndex, TypedElementDataStorage::SubqueryCallbackRef Callback) override { return Implementation.RunSubquery(SubqueryIndex, Callback); }
-			TypedElementDataStorage::FQueryResult RunSubquery(int32 SubqueryIndex, TypedElementDataStorage::RowHandle Row, TypedElementDataStorage::SubqueryCallbackRef Callback) override { return Implementation.RunSubquery(SubqueryIndex, Row, Callback); }
+			TypedElementDataStorage::FQueryResult RunSubquery(int32 SubqueryIndex, RowHandle Row, TypedElementDataStorage::SubqueryCallbackRef Callback) override { return Implementation.RunSubquery(SubqueryIndex, Row, Callback); }
 	
-
 			FMassQueryContextImplementation Implementation;
 		};
 	} // namespace Processors::Private
@@ -882,6 +879,8 @@ bool FTypedElementQueryProcessorData::CommonQueryConfiguration(
 	FEnvironment& InEnvironment,
 	TArrayView<FMassEntityQuery> Subqueries)
 {
+	using namespace UE::Editor::DataStorage;
+
 	ParentQuery = InQueryHandle;
 	QueryStore = &InQueryStore;
 	Environment = &InEnvironment;
@@ -891,7 +890,7 @@ bool FTypedElementQueryProcessorData::CommonQueryConfiguration(
 	{
 		bool Result = true;
 		int32 CurrentSubqueryIndex = 0;
-		for (TypedElementQueryHandle SubqueryHandle : InQuery.Description.Subqueries)
+		for (QueryHandle SubqueryHandle : InQuery.Description.Subqueries)
 		{
 			const FExtendedQueryStore::Handle SubqueryStoreHandle(SubqueryHandle);
 			if (const FExtendedQuery* Subquery = InQueryStore.Get(SubqueryStoreHandle))
diff --git a/Engine/Plugins/Experimental/EditorDataStorage/Source/TedsCore/Private/Queries/TypedElementAlertQueries.cpp b/Engine/Plugins/Experimental/EditorDataStorage/Source/TedsCore/Private/Queries/TypedElementAlertQueries.cpp
--- a/Engine/Plugins/Experimental/EditorDataStorage/Source/TedsCore/Private/Queries/TypedElementAlertQueries.cpp
+++ b/Engine/Plugins/Experimental/EditorDataStorage/Source/TedsCore/Private/Queries/TypedElementAlertQueries.cpp
@@ -19,12 +19,13 @@ FAutoConsoleCommand AddRandomAlertToRowConsoleCommand(
 		{
 			using namespace TypedElementDataStorage;
 			using namespace TypedElementQueryBuilder;
+			using namespace UE::Editor::DataStorage;
 
 			TRACE_CPUPROFILER_EVENT_SCOPE(TEDS.Debug.AddRandomAlertToSelectedRows);
 
 			if (ITypedElementDataStorageInterface* DataStorage = UTypedElementRegistry::GetInstance()->GetMutableDataStorage())
 			{
-				static TypedElementQueryHandle Query = [DataStorage]
+				static QueryHandle Query = [DataStorage]
 				{
 					return DataStorage->RegisterQuery(
 						Select()
@@ -78,6 +79,7 @@ FAutoConsoleCommand ClearAllAlertsConsoleCommand(
 		{
 			using namespace TypedElementDataStorage;
 			using namespace TypedElementQueryBuilder;
+			using namespace UE::Editor::DataStorage;
 
 			TRACE_CPUPROFILER_EVENT_SCOPE(TEDS.Debug.ClearAllAlertInfo);
 
@@ -86,7 +88,7 @@ FAutoConsoleCommand ClearAllAlertsConsoleCommand(
 
 			if (ITypedElementDataStorageInterface* DataStorage = UTypedElementRegistry::GetInstance()->GetMutableDataStorage())
 			{
-				static TypedElementQueryHandle AlertInfoQuery = [DataStorage]
+				static QueryHandle AlertInfoQuery = [DataStorage]
 				{
 					return DataStorage->RegisterQuery(
 						Select()
@@ -118,12 +120,13 @@ FAutoConsoleCommand ClearSelectedAlertsConsoleCommand(
 		{
 			using namespace TypedElementDataStorage;
 			using namespace TypedElementQueryBuilder;
+			using namespace UE::Editor::DataStorage;
 
 			TRACE_CPUPROFILER_EVENT_SCOPE(TEDS.Debug.ClearSelectedAlerts);
 
 			if (ITypedElementDataStorageInterface* DataStorage = UTypedElementRegistry::GetInstance()->GetMutableDataStorage())
 			{
-				static TypedElementQueryHandle AlertQuery = [DataStorage]
+				static QueryHandle AlertQuery = [DataStorage]
 				{
 					return DataStorage->RegisterQuery(
 						Select()
@@ -275,6 +278,7 @@ void UTypedElementAlertQueriesFactory::RegisterOnAddQueries(ITypedElementDataSto
 {
 	using namespace TypedElementDataStorage;
 	using namespace TypedElementQueryBuilder;
+	using namespace UE::Editor::DataStorage;
 	
 	DataStorage.RegisterQuery(
 		Select(
@@ -305,12 +309,13 @@ void UTypedElementAlertQueriesFactory::RegisterOnRemoveQueries(ITypedElementData
 {
 	using namespace TypedElementDataStorage;
 	using namespace TypedElementQueryBuilder;
+	using namespace UE::Editor::DataStorage;
 	
 	DataStorage.RegisterQuery(
 		Select(
 			TEXT("Remove alert"),
 			FObserver::OnRemove<FTypedElementAlertColumn>(),
-			[](IQueryContext& Context, TypedElementDataStorage::RowHandle Row, FTypedElementAlertColumn& Alert)
+			[](IQueryContext& Context, RowHandle Row, FTypedElementAlertColumn& Alert)
 			{
 				Context.ActivateQueries(AlertConditionName);
 			})
@@ -320,7 +325,7 @@ void UTypedElementAlertQueriesFactory::RegisterOnRemoveQueries(ITypedElementData
 		Select(
 			TEXT("Update alert upon parent removal"),
 			FObserver::OnRemove<FTableRowParentColumn>(),
-			[](IQueryContext& Context, TypedElementDataStorage::RowHandle Row)
+			[](IQueryContext& Context, RowHandle Row)
 			{
 				Context.ActivateQueries(AlertConditionName);
 			})
@@ -330,7 +335,7 @@ void UTypedElementAlertQueriesFactory::RegisterOnRemoveQueries(ITypedElementData
 }
 
 void UTypedElementAlertQueriesFactory::AddChildAlertsToHierarchy(
-	TypedElementDataStorage::IQueryContext& Context, TypedElementDataStorage::RowHandle Parent, int32 ParentQueryIndex)
+	TypedElementDataStorage::IQueryContext& Context, UE::Editor::DataStorage::RowHandle Parent, int32 ParentQueryIndex)
 {
 	using namespace TypedElementDataStorage;
 
@@ -353,11 +358,12 @@ void UTypedElementAlertQueriesFactory::AddChildAlertsToHierarchy(
 }
 
 void UTypedElementAlertQueriesFactory::IncrementParents(
-	TypedElementDataStorage::IQueryContext& Context, TypedElementDataStorage::RowHandle Row, FTypedElementAlertColumnType AlertType,
+	TypedElementDataStorage::IQueryContext& Context, UE::Editor::DataStorage::RowHandle Row, FTypedElementAlertColumnType AlertType,
 	int32 ChildAlertQueryIndex)
 {
 	using namespace TypedElementDataStorage;
 	using namespace TypedElementQueryBuilder;
+	using namespace UE::Editor::DataStorage;
 	
 	while (Context.IsRowAvailable(Row))
 	{
diff --git a/Engine/Plugins/Experimental/EditorDataStorage/Source/TedsCore/Private/Queries/TypedElementExtendedQueryStore.cpp b/Engine/Plugins/Experimental/EditorDataStorage/Source/TedsCore/Private/Queries/TypedElementExtendedQueryStore.cpp
--- a/Engine/Plugins/Experimental/EditorDataStorage/Source/TedsCore/Private/Queries/TypedElementExtendedQueryStore.cpp
+++ b/Engine/Plugins/Experimental/EditorDataStorage/Source/TedsCore/Private/Queries/TypedElementExtendedQueryStore.cpp
@@ -283,7 +283,7 @@ namespace UE::Editor::DataStorage
 	}
 
 	TypedElementDataStorage::FQueryResult FExtendedQueryStore::RunQuery(FMassEntityManager& EntityManager,
-		FEnvironment& Environment, FMassExecutionContext& ParentContext, Handle Query, TypedElementRowHandle Row,
+		FEnvironment& Environment, FMassExecutionContext& ParentContext, Handle Query, RowHandle Row,
 		TypedElementDataStorage::SubqueryCallbackRef Callback)
 	{
 		using ActionType = ITypedElementDataStorageInterface::FQueryDescription::EActionType;
diff --git a/Engine/Plugins/Experimental/EditorDataStorage/Source/TedsCore/Private/Queries/TypedElementExtendedQueryStore.h b/Engine/Plugins/Experimental/EditorDataStorage/Source/TedsCore/Private/Queries/TypedElementExtendedQueryStore.h
--- a/Engine/Plugins/Experimental/EditorDataStorage/Source/TedsCore/Private/Queries/TypedElementExtendedQueryStore.h
+++ b/Engine/Plugins/Experimental/EditorDataStorage/Source/TedsCore/Private/Queries/TypedElementExtendedQueryStore.h
@@ -119,7 +119,7 @@ namespace UE::Editor::DataStorage
 			FEnvironment& Environment,
 			FMassExecutionContext& ParentContext,
 			Handle Query,
-			TypedElementRowHandle Row,
+			DataStorage::RowHandle Row,
 			TypedElementDataStorage::SubqueryCallbackRef Callback);
 		void RunPhasePreambleQueries(
 			FMassEntityManager& EntityManager,
diff --git a/Engine/Plugins/Experimental/EditorDataStorage/Source/TedsCore/Private/Queries/TypedElementMiscQueries.cpp b/Engine/Plugins/Experimental/EditorDataStorage/Source/TedsCore/Private/Queries/TypedElementMiscQueries.cpp
--- a/Engine/Plugins/Experimental/EditorDataStorage/Source/TedsCore/Private/Queries/TypedElementMiscQueries.cpp
+++ b/Engine/Plugins/Experimental/EditorDataStorage/Source/TedsCore/Private/Queries/TypedElementMiscQueries.cpp
@@ -8,15 +8,16 @@
 void UTypedElementRemoveSyncToWorldTagFactory::RegisterQueries(ITypedElementDataStorageInterface& DataStorage)
 {
 	using namespace TypedElementQueryBuilder;
+	using namespace UE::Editor::DataStorage;
 	using DSI = ITypedElementDataStorageInterface;
 
 	DataStorage.RegisterQuery(
 		Select(
 			TEXT("Remove 'sync to world' tag"),
 			FPhaseAmble(FPhaseAmble::ELocation::Postamble, DSI::EQueryTickPhase::FrameEnd),
-			[](DSI::IQueryContext& Context, const TypedElementRowHandle* Rows)
+			[](DSI::IQueryContext& Context, const RowHandle* Rows)
 			{
-				Context.RemoveColumns<FTypedElementSyncBackToWorldTag>(TConstArrayView<TypedElementRowHandle>(Rows, Context.GetRowCount()));
+				Context.RemoveColumns<FTypedElementSyncBackToWorldTag>(TConstArrayView<RowHandle>(Rows, Context.GetRowCount()));
 			}
 		)
 		.Where()
@@ -28,9 +29,9 @@ void UTypedElementRemoveSyncToWorldTagFactory::RegisterQueries(ITypedElementData
 	Select(
 		TEXT("Remove 'sync from world' tag"),
 		FPhaseAmble(FPhaseAmble::ELocation::Postamble, DSI::EQueryTickPhase::FrameEnd),
-		[](DSI::IQueryContext& Context, const TypedElementRowHandle* Rows)
+		[](DSI::IQueryContext& Context, const RowHandle* Rows)
 		{
-			Context.RemoveColumns<FTypedElementSyncFromWorldTag>(TConstArrayView<TypedElementRowHandle>(Rows, Context.GetRowCount()));
+			Context.RemoveColumns<FTypedElementSyncFromWorldTag>(TConstArrayView<RowHandle>(Rows, Context.GetRowCount()));
 		}
 	)
 	.Where()
diff --git a/Engine/Plugins/Experimental/EditorDataStorage/Source/TedsCore/Private/TypedElementDatabase.cpp b/Engine/Plugins/Experimental/EditorDataStorage/Source/TedsCore/Private/TypedElementDatabase.cpp
--- a/Engine/Plugins/Experimental/EditorDataStorage/Source/TedsCore/Private/TypedElementDatabase.cpp
+++ b/Engine/Plugins/Experimental/EditorDataStorage/Source/TedsCore/Private/TypedElementDatabase.cpp
@@ -286,14 +286,14 @@ TypedElementDataStorage::TableHandle UEditorDataStorage::RegisterTable(TConstArr
 	return InvalidTableHandle;
 }
 
-TypedElementDataStorage::TableHandle UEditorDataStorage::RegisterTable(TypedElementDataStorage::TableHandle SourceTable,
+UE::Editor::DataStorage::TableHandle UEditorDataStorage::RegisterTable(TypedElementDataStorage::TableHandle SourceTable,
 	TConstArrayView<const UScriptStruct*> ColumnList, const FName Name)
 {
-	using namespace TypedElementDataStorage;
+	using namespace UE::Editor::DataStorage;
 
 	if (ActiveEditorEntityManager && SourceTable < Tables.Num() && !TableNameLookup.Contains(Name))
 	{
-		TypedElementDataStorage::TableHandle Result = Tables.Num();
+		TableHandle Result = Tables.Num();
 		FMassArchetypeCreationParams ArchetypeCreationParams;
 		ArchetypeCreationParams.DebugName = Name;
 		ArchetypeCreationParams.ChunkMemorySize = GetTableChunkSize(Name);
@@ -307,22 +307,22 @@ TypedElementDataStorage::TableHandle UEditorDataStorage::RegisterTable(TypedElem
 	return InvalidTableHandle;
 }
 
-TypedElementDataStorage::TableHandle UEditorDataStorage::FindTable(const FName Name)
+UE::Editor::DataStorage::TableHandle UEditorDataStorage::FindTable(const FName Name)
 {
-	using namespace TypedElementDataStorage;
+	using namespace UE::Editor::DataStorage;
 
 	TableHandle* TableHandle = TableNameLookup.Find(Name);
 	return TableHandle ? *TableHandle : InvalidTableHandle;
 }
 
-TypedElementDataStorage::RowHandle UEditorDataStorage::ReserveRow()
+UE::Editor::DataStorage::RowHandle UEditorDataStorage::ReserveRow()
 {
 	return ActiveEditorEntityManager 
 		? ActiveEditorEntityManager->ReserveEntity().AsNumber()
-		: TypedElementDataStorage::InvalidRowHandle;
+		: UE::Editor::DataStorage::InvalidRowHandle;
 }
 
-void UEditorDataStorage::BatchReserveRows(int32 Count, TFunctionRef<void(TypedElementDataStorage::RowHandle)> ReservationCallback)
+void UEditorDataStorage::BatchReserveRows(int32 Count, TFunctionRef<void(RowHandle)> ReservationCallback)
 {
 	using namespace TypedElementDataStorage;
 
@@ -338,7 +338,7 @@ void UEditorDataStorage::BatchReserveRows(int32 Count, TFunctionRef<void(TypedEl
 	}
 }
 
-void UEditorDataStorage::BatchReserveRows(TArrayView<TypedElementDataStorage::RowHandle> ReservedRows)
+void UEditorDataStorage::BatchReserveRows(TArrayView<RowHandle> ReservedRows)
 {
 	using namespace TypedElementDataStorage;
 
@@ -354,18 +354,18 @@ void UEditorDataStorage::BatchReserveRows(TArrayView<TypedElementDataStorage::Ro
 	}
 }
 
-TypedElementDataStorage::RowHandle UEditorDataStorage::AddRow(TypedElementDataStorage::TableHandle Table)
+UE::Editor::DataStorage::RowHandle UEditorDataStorage::AddRow(TableHandle Table)
 {
 	checkf(Table < Tables.Num(), TEXT("Attempting to add a row to a non-existing table."));
 	return ActiveEditorEntityManager 
 		? ActiveEditorEntityManager->CreateEntity(Tables[Table]).AsNumber() 
-		: TypedElementDataStorage::InvalidRowHandle;
+		: UE::Editor::DataStorage::InvalidRowHandle;
 }
 
-TypedElementDataStorage::RowHandle UEditorDataStorage::AddRow(TypedElementDataStorage::TableHandle Table,
+UE::Editor::DataStorage::RowHandle UEditorDataStorage::AddRow(TableHandle Table,
 	RowCreationCallbackRef OnCreated)
 {
-	using namespace TypedElementDataStorage;
+	using namespace UE::Editor::DataStorage;
 
 	OnCreated.CheckCallable();
 	if (ActiveEditorEntityManager)
@@ -385,7 +385,7 @@ TypedElementDataStorage::RowHandle UEditorDataStorage::AddRow(TypedElementDataSt
 	return InvalidRowHandle;
 }
 
-bool UEditorDataStorage::AddRow(TypedElementDataStorage::RowHandle ReservedRow, TypedElementDataStorage::TableHandle Table)
+bool UEditorDataStorage::AddRow(RowHandle ReservedRow, TableHandle Table)
 {
 	checkf(!IsRowAssigned(ReservedRow), TEXT("Attempting to assign a table to row that already has a table assigned."));
 	checkf(Table < Tables.Num(), TEXT("Attempting to add a row to a non-existing table."));
@@ -400,7 +400,7 @@ bool UEditorDataStorage::AddRow(TypedElementDataStorage::RowHandle ReservedRow,
 	}
 }
 
-bool UEditorDataStorage::AddRow(TypedElementDataStorage::RowHandle ReservedRow, TypedElementDataStorage::TableHandle Table,
+bool UEditorDataStorage::AddRow(RowHandle ReservedRow, TableHandle Table,
 	RowCreationCallbackRef OnCreated)
 {
 	OnCreated.CheckCallable();
@@ -418,7 +418,7 @@ bool UEditorDataStorage::AddRow(TypedElementDataStorage::RowHandle ReservedRow,
 }
 
 bool UEditorDataStorage::BatchAddRow(
-	TypedElementDataStorage::TableHandle Table, int32 Count, RowCreationCallbackRef OnCreated)
+	TableHandle Table, int32 Count, RowCreationCallbackRef OnCreated)
 {
 	OnCreated.CheckCallable();
 	if (ActiveEditorEntityManager)
@@ -440,9 +440,10 @@ bool UEditorDataStorage::BatchAddRow(
 	return false;
 }
 
-bool UEditorDataStorage::BatchAddRow(TypedElementDataStorage::TableHandle Table, 
-	TConstArrayView<TypedElementDataStorage::RowHandle> ReservedHandles, RowCreationCallbackRef OnCreated)
+bool UEditorDataStorage::BatchAddRow(TableHandle Table, 
+	TConstArrayView<RowHandle> ReservedHandles, RowCreationCallbackRef OnCreated)
 {
+	using namespace UE::Editor;
 	OnCreated.CheckCallable();
 	if (ActiveEditorEntityManager)
 	{
@@ -450,7 +451,7 @@ bool UEditorDataStorage::BatchAddRow(TypedElementDataStorage::TableHandle Table,
 	
 		// Depend on the fact that a row handle is an alias for an entity within the Mass powered backend. This
 		// avoids the need for copying to a temporary array;
-		static_assert(sizeof(TypedElementRowHandle) == sizeof(FMassEntityHandle), 
+		static_assert(sizeof(RowHandle) == sizeof(FMassEntityHandle), 
 			"BatchAddRow in TEDS requires the row handle and the Mass entity handle to be the same size.");
 		
 		TConstArrayView<FMassEntityHandle> Entities(
@@ -458,7 +459,7 @@ bool UEditorDataStorage::BatchAddRow(TypedElementDataStorage::TableHandle Table,
 		TSharedRef<FMassEntityManager::FEntityCreationContext> Context =
 			ActiveEditorEntityManager->BatchCreateReservedEntities(Tables[Table], Entities);
 
-		for (TypedElementRowHandle Entity : ReservedHandles)
+		for (RowHandle Entity : ReservedHandles)
 		{
 			OnCreated(Entity);
 		}
@@ -469,7 +470,7 @@ bool UEditorDataStorage::BatchAddRow(TypedElementDataStorage::TableHandle Table,
 }
 
 
-void UEditorDataStorage::RemoveRow(TypedElementDataStorage::RowHandle Row)
+void UEditorDataStorage::RemoveRow(RowHandle Row)
 {
 	FMassEntityHandle Entity = FMassEntityHandle::FromNumber(Row);
 	if (ActiveEditorEntityManager && ActiveEditorEntityManager->IsEntityValid(Entity))
@@ -487,17 +488,17 @@ void UEditorDataStorage::RemoveRow(TypedElementDataStorage::RowHandle Row)
 	}
 }
 
-bool UEditorDataStorage::IsRowAvailable(TypedElementDataStorage::RowHandle Row) const
+bool UEditorDataStorage::IsRowAvailable(RowHandle Row) const
 {
 	return ActiveEditorEntityManager ? UE::Editor::DataStorage::Legacy::FCommandBuffer::Execute_IsRowAvailable(*ActiveEditorEntityManager, Row) : false;
 }
 
-bool UEditorDataStorage::IsRowAssigned(TypedElementDataStorage::RowHandle Row) const
+bool UEditorDataStorage::IsRowAssigned(RowHandle Row) const
 {
 	return ActiveEditorEntityManager ? UE::Editor::DataStorage::Legacy::FCommandBuffer::Execute_IsRowAssigned(*ActiveEditorEntityManager, Row) : false;
 }
 
-void UEditorDataStorage::AddColumn(TypedElementRowHandle Row, const UScriptStruct* ColumnType)
+void UEditorDataStorage::AddColumn(RowHandle Row, const UScriptStruct* ColumnType)
 {
 	if (ColumnType && ActiveEditorEntityManager)
 	{
@@ -512,7 +513,7 @@ void UEditorDataStorage::AddColumn(TypedElementRowHandle Row, const UScriptStruc
 	}
 }
 
-void UEditorDataStorage::AddColumnData(TypedElementRowHandle Row, const UScriptStruct* ColumnType,
+void UEditorDataStorage::AddColumnData(RowHandle Row, const UScriptStruct* ColumnType,
 	const ColumnCreationCallbackRef& Initializer,
 	ColumnCopyOrMoveCallback Relocator)
 {
@@ -539,7 +540,7 @@ void UEditorDataStorage::AddColumnData(TypedElementRowHandle Row, const UScriptS
 	}
 }
 
-void UEditorDataStorage::RemoveColumn(TypedElementRowHandle Row, const UScriptStruct* ColumnType)
+void UEditorDataStorage::RemoveColumn(RowHandle Row, const UScriptStruct* ColumnType)
 {
 	if (ColumnType && ActiveEditorEntityManager)
 	{
@@ -554,7 +555,7 @@ void UEditorDataStorage::RemoveColumn(TypedElementRowHandle Row, const UScriptSt
 	}
 }
 
-const void* UEditorDataStorage::GetColumnData(TypedElementRowHandle Row, const UScriptStruct* ColumnType) const
+const void* UEditorDataStorage::GetColumnData(RowHandle Row, const UScriptStruct* ColumnType) const
 {
 	FMassEntityHandle Entity = FMassEntityHandle::FromNumber(Row);
 	if (ActiveEditorEntityManager &&
@@ -576,12 +577,12 @@ const void* UEditorDataStorage::GetColumnData(TypedElementRowHandle Row, const U
 	return nullptr;
 }
 
-void* UEditorDataStorage::GetColumnData(TypedElementRowHandle Row, const UScriptStruct* ColumnType)
+void* UEditorDataStorage::GetColumnData(RowHandle Row, const UScriptStruct* ColumnType)
 {
 	return const_cast<void*>(static_cast<const UEditorDataStorage*>(this)->GetColumnData(Row, ColumnType));
 }
 
-void UEditorDataStorage::AddColumns(TypedElementRowHandle Row, TConstArrayView<const UScriptStruct*> Columns)
+void UEditorDataStorage::AddColumns(RowHandle Row, TConstArrayView<const UScriptStruct*> Columns)
 {
 	if (ActiveEditorEntityManager)
 	{
@@ -604,7 +605,7 @@ void UEditorDataStorage::AddColumns(TypedElementRowHandle Row, TConstArrayView<c
 	}
 }
 
-void UEditorDataStorage::AddColumn(TypedElementRowHandle Row, const UE::Editor::DataStorage::FDynamicTag& Tag, const FName& InValue)
+void UEditorDataStorage::AddColumn(RowHandle Row, const UE::Editor::DataStorage::FDynamicTag& Tag, const FName& InValue)
 {
 	if (ActiveEditorEntityManager)
 	{
@@ -618,7 +619,7 @@ void UEditorDataStorage::AddColumn(TypedElementRowHandle Row, const UE::Editor::
 	}
 }
 
-void UEditorDataStorage::RemoveColumn(TypedElementRowHandle Row, const UE::Editor::DataStorage::FDynamicTag& Tag)
+void UEditorDataStorage::RemoveColumn(RowHandle Row, const UE::Editor::DataStorage::FDynamicTag& Tag)
 {
 	if (ActiveEditorEntityManager)
 	{
@@ -631,7 +632,7 @@ void UEditorDataStorage::RemoveColumn(TypedElementRowHandle Row, const UE::Edito
 	}
 }
 
-void UEditorDataStorage::RemoveColumns(TypedElementRowHandle Row, TConstArrayView<const UScriptStruct*> Columns)
+void UEditorDataStorage::RemoveColumns(RowHandle Row, TConstArrayView<const UScriptStruct*> Columns)
 {
 	FMassEntityHandle Entity = FMassEntityHandle::FromNumber(Row);
 	if (ActiveEditorEntityManager)
@@ -654,7 +655,7 @@ void UEditorDataStorage::RemoveColumns(TypedElementRowHandle Row, TConstArrayVie
 	}
 }
 
-void UEditorDataStorage::AddRemoveColumns(TypedElementRowHandle Row,
+void UEditorDataStorage::AddRemoveColumns(RowHandle Row,
 	TConstArrayView<const UScriptStruct*> ColumnsToAdd, TConstArrayView<const UScriptStruct*> ColumnsToRemove)
 {
 	if (ActiveEditorEntityManager)
@@ -695,7 +696,7 @@ void UEditorDataStorage::AddRemoveColumns(TypedElementRowHandle Row,
 	}
 }
 
-void UEditorDataStorage::BatchAddRemoveColumns(TConstArrayView<TypedElementRowHandle> Rows, 
+void UEditorDataStorage::BatchAddRemoveColumns(TConstArrayView<RowHandle> Rows, 
 	TConstArrayView<const UScriptStruct*> ColumnsToAdd, TConstArrayView<const UScriptStruct*> ColumnsToRemove)
 {	
 	if (ActiveEditorEntityManager)
@@ -721,7 +722,7 @@ void UEditorDataStorage::BatchAddRemoveColumns(TConstArrayView<TypedElementRowHa
 			
 			// Sort rows (entities) into to matching table (archetype) bucket.
 			EntityArchetypeLookup LookupTable;
-			for (TypedElementRowHandle EntityId : Rows)
+			for (RowHandle EntityId : Rows)
 			{
 				FMassEntityHandle Entity = FMassEntityHandle::FromNumber(EntityId);
 				if (ActiveEditorEntityManager->IsEntityActive(Entity))
@@ -764,7 +765,7 @@ void UEditorDataStorage::BatchAddRemoveColumns(TConstArrayView<TypedElementRowHa
 	}
 }
 
-bool UEditorDataStorage::HasColumns(TypedElementRowHandle Row, TConstArrayView<const UScriptStruct*> ColumnTypes) const
+bool UEditorDataStorage::HasColumns(RowHandle Row, TConstArrayView<const UScriptStruct*> ColumnTypes) const
 {
 	if (ActiveEditorEntityManager)
 	{
@@ -808,7 +809,7 @@ bool UEditorDataStorage::HasColumns(TypedElementRowHandle Row, TConstArrayView<c
 	return false;
 }
 
-bool UEditorDataStorage::HasColumns(TypedElementRowHandle Row, TConstArrayView<TWeakObjectPtr<const UScriptStruct>> ColumnTypes) const
+bool UEditorDataStorage::HasColumns(RowHandle Row, TConstArrayView<TWeakObjectPtr<const UScriptStruct>> ColumnTypes) const
 {
 	if (ActiveEditorEntityManager)
 	{
@@ -854,7 +855,7 @@ bool UEditorDataStorage::HasColumns(TypedElementRowHandle Row, TConstArrayView<T
 	return false;
 }
 
-void UEditorDataStorage::ListColumns(TypedElementDataStorage::RowHandle Row, ColumnListCallbackRef Callback) const
+void UEditorDataStorage::ListColumns(RowHandle Row, ColumnListCallbackRef Callback) const
 {
 	if (ActiveEditorEntityManager)
 	{
@@ -879,7 +880,7 @@ void UEditorDataStorage::ListColumns(TypedElementDataStorage::RowHandle Row, Col
 	}
 }
 
-void UEditorDataStorage::ListColumns(TypedElementDataStorage::RowHandle Row, ColumnListWithDataCallbackRef Callback)
+void UEditorDataStorage::ListColumns(RowHandle Row, ColumnListWithDataCallbackRef Callback)
 {
 	if (ActiveEditorEntityManager)
 	{
@@ -912,7 +913,7 @@ void UEditorDataStorage::ListColumns(TypedElementDataStorage::RowHandle Row, Col
 	}
 }
 
-bool UEditorDataStorage::MatchesColumns(TypedElementDataStorage::RowHandle Row, const TypedElementDataStorage::FQueryConditions& Conditions) const
+bool UEditorDataStorage::MatchesColumns(RowHandle Row, const TypedElementDataStorage::FQueryConditions& Conditions) const
 {
 	if (ActiveEditorEntityManager)
 	{
@@ -977,14 +978,14 @@ void UEditorDataStorage::UnregisterTickGroup(FName GroupName, EQueryTickPhase Ph
 	Environment->GetQueryStore().UnregisterTickGroup(GroupName, Phase);
 }
 
-TypedElementQueryHandle UEditorDataStorage::RegisterQuery(FQueryDescription&& Query)
+UE::Editor::DataStorage::QueryHandle UEditorDataStorage::RegisterQuery(FQueryDescription&& Query)
 {
 	return (ActiveEditorEntityManager && ActiveEditorPhaseManager)
 		? Environment->GetQueryStore().RegisterQuery(MoveTemp(Query), *Environment, *ActiveEditorEntityManager, *ActiveEditorPhaseManager).Packed()
-		: TypedElementInvalidQueryHandle;
+		: UE::Editor::DataStorage::InvalidQueryHandle;
 }
 
-void UEditorDataStorage::UnregisterQuery(TypedElementQueryHandle Query)
+void UEditorDataStorage::UnregisterQuery(QueryHandle Query)
 {
 	if (ActiveEditorEntityManager && ActiveEditorPhaseManager)
 	{
@@ -993,7 +994,7 @@ void UEditorDataStorage::UnregisterQuery(TypedElementQueryHandle Query)
 	}
 }
 
-const ITypedElementDataStorageInterface::FQueryDescription& UEditorDataStorage::GetQueryDescription(TypedElementQueryHandle Query) const
+const ITypedElementDataStorageInterface::FQueryDescription& UEditorDataStorage::GetQueryDescription(QueryHandle Query) const
 {
 	const UE::Editor::DataStorage::FExtendedQueryStore::Handle StorageHandle(Query);
 	return Environment->GetQueryStore().GetQueryDescription(StorageHandle);
@@ -1023,7 +1024,7 @@ FName UEditorDataStorage::GetQueryTickGroupName(EQueryTickGroups Group) const
 	}
 }
 
-ITypedElementDataStorageInterface::FQueryResult UEditorDataStorage::RunQuery(TypedElementQueryHandle Query)
+ITypedElementDataStorageInterface::FQueryResult UEditorDataStorage::RunQuery(QueryHandle Query)
 {
 	TRACE_CPUPROFILER_EVENT_SCOPE(TEDS.RunQuery);
 
@@ -1039,7 +1040,7 @@ ITypedElementDataStorageInterface::FQueryResult UEditorDataStorage::RunQuery(Typ
 }
 
 ITypedElementDataStorageInterface::FQueryResult UEditorDataStorage::RunQuery(
-	TypedElementQueryHandle Query, ITypedElementDataStorageInterface::DirectQueryCallbackRef Callback)
+	QueryHandle Query, ITypedElementDataStorageInterface::DirectQueryCallbackRef Callback)
 {
 	TRACE_CPUPROFILER_EVENT_SCOPE(TEDS.RunQuery);
 
@@ -1056,7 +1057,7 @@ ITypedElementDataStorageInterface::FQueryResult UEditorDataStorage::RunQuery(
 }
 
 ITypedElementDataStorageInterface::FQueryResult UEditorDataStorage::RunQuery(
-	TypedElementQueryHandle Query, TypedElementDataStorage::EDirectQueryExecutionFlags Flags,
+	QueryHandle Query, TypedElementDataStorage::EDirectQueryExecutionFlags Flags,
 	ITypedElementDataStorageInterface::DirectQueryCallbackRef Callback)
 {
 	TRACE_CPUPROFILER_EVENT_SCOPE(TEDS.RunQuery);
@@ -1080,24 +1081,24 @@ void UEditorDataStorage::ActivateQueries(FName ActivationName)
 	}
 }
 
-TypedElementDataStorage::RowHandle UEditorDataStorage::FindIndexedRow(TypedElementDataStorage::IndexHash Index) const
+UE::Editor::DataStorage::RowHandle UEditorDataStorage::FindIndexedRow(TypedElementDataStorage::IndexHash Index) const
 {
 	return Environment->GetIndexTable().FindIndexedRow(UE::Editor::DataStorage::EGlobalLockScope::Public, Index);
 }
 
-void UEditorDataStorage::IndexRow(TypedElementDataStorage::IndexHash Index, TypedElementDataStorage::RowHandle Row)
+void UEditorDataStorage::IndexRow(TypedElementDataStorage::IndexHash Index, RowHandle Row)
 {
 	Environment->GetIndexTable().IndexRow(UE::Editor::DataStorage::EGlobalLockScope::Public, Index, Row);
 }
 
 void UEditorDataStorage::BatchIndexRows(
-	TConstArrayView<TPair<TypedElementDataStorage::IndexHash, TypedElementDataStorage::RowHandle>> IndexRowPairs)
+	TConstArrayView<TPair<TypedElementDataStorage::IndexHash, RowHandle>> IndexRowPairs)
 {
 	Environment->GetIndexTable().BatchIndexRows(UE::Editor::DataStorage::EGlobalLockScope::Public, IndexRowPairs);
 }
 
 void UEditorDataStorage::ReindexRow(TypedElementDataStorage::IndexHash OriginalIndex, TypedElementDataStorage::IndexHash NewIndex, 
-	TypedElementDataStorage::RowHandle RowHandle)
+	RowHandle RowHandle)
 {
 	Environment->GetIndexTable().ReindexRow(UE::Editor::DataStorage::EGlobalLockScope::Public, OriginalIndex, NewIndex, RowHandle);
 }
