[Teds] Update naming for RowHandle, QueryHandle, and TableHandle

diff --git a/Engine/Plugins/Experimental/EditorDataStorage/Source/TedsCore/Private/TypedElementDatabase.h b/Engine/Plugins/Experimental/EditorDataStorage/Source/TedsCore/Private/TypedElementDatabase.h
--- a/Engine/Plugins/Experimental/EditorDataStorage/Source/TedsCore/Private/TypedElementDatabase.h
+++ b/Engine/Plugins/Experimental/EditorDataStorage/Source/TedsCore/Private/TypedElementDatabase.h
@@ -36,6 +36,9 @@ class TEDSCORE_API UEditorDataStorage
 	using ColumnListCallbackRef = UE::Editor::DataStorage::ColumnListCallbackRef;
 	using ColumnListWithDataCallbackRef = UE::Editor::DataStorage::ColumnListWithDataCallbackRef;
 	using ColumnCopyOrMoveCallback = UE::Editor::DataStorage::ColumnCopyOrMoveCallback;
+	using RowHandle = UE::Editor::DataStorage::RowHandle;
+	using TableHandle = UE::Editor::DataStorage::TableHandle;
+	using QueryHandle = UE::Editor::DataStorage::QueryHandle;
 
 public:
 	template<typename FactoryType, typename DatabaseType>
@@ -90,48 +93,48 @@ public:
 	TSharedPtr<FMassEntityManager> GetActiveMutableEditorEntityManager();
 	TSharedPtr<const FMassEntityManager> GetActiveEditorEntityManager() const;
 
-	virtual TypedElementDataStorage::TableHandle RegisterTable(TConstArrayView<const UScriptStruct*> ColumnList, const FName Name) override;
-	virtual TypedElementDataStorage::TableHandle RegisterTable(
-		TypedElementDataStorage::TableHandle SourceTable, TConstArrayView<const UScriptStruct*> ColumnList, const FName Name) override;
-	virtual TypedElementDataStorage::TableHandle FindTable(const FName Name) override;
+	virtual TableHandle RegisterTable(TConstArrayView<const UScriptStruct*> ColumnList, const FName Name) override;
+	virtual TableHandle RegisterTable(
+		TableHandle SourceTable, TConstArrayView<const UScriptStruct*> ColumnList, const FName Name) override;
+	virtual TableHandle FindTable(const FName Name) override;
 
-	virtual TypedElementDataStorage::RowHandle ReserveRow() override;
-	virtual void BatchReserveRows(int32 Count, TFunctionRef<void(TypedElementDataStorage::RowHandle)> ReservationCallback) override;
-	virtual void BatchReserveRows(TArrayView<TypedElementDataStorage::RowHandle> ReservedRows) override;
-	virtual TypedElementDataStorage::RowHandle AddRow(TypedElementDataStorage::TableHandle Table, 
+	virtual RowHandle ReserveRow() override;
+	virtual void BatchReserveRows(int32 Count, TFunctionRef<void(RowHandle)> ReservationCallback) override;
+	virtual void BatchReserveRows(TArrayView<RowHandle> ReservedRows) override;
+	virtual RowHandle AddRow(TableHandle Table, 
 		RowCreationCallbackRef OnCreated) override;
-	TypedElementDataStorage::RowHandle AddRow(TypedElementDataStorage::TableHandle Table) override;
-	virtual bool AddRow(TypedElementDataStorage::RowHandle ReservedRow, TypedElementDataStorage::TableHandle Table) override;
-	virtual bool AddRow(TypedElementDataStorage::RowHandle ReservedRow, TypedElementDataStorage::TableHandle Table,
+	RowHandle AddRow(TableHandle Table) override;
+	virtual bool AddRow(RowHandle ReservedRow, TableHandle Table) override;
+	virtual bool AddRow(RowHandle ReservedRow, TableHandle Table,
 		RowCreationCallbackRef OnCreated) override;
-	virtual bool BatchAddRow(TypedElementDataStorage::TableHandle Table, int32 Count,
+	virtual bool BatchAddRow(TableHandle Table, int32 Count,
 		RowCreationCallbackRef OnCreated) override;
-	virtual bool BatchAddRow(TypedElementDataStorage::TableHandle Table, TConstArrayView<TypedElementDataStorage::RowHandle> ReservedHandles,
+	virtual bool BatchAddRow(TableHandle Table, TConstArrayView<RowHandle> ReservedHandles,
 		RowCreationCallbackRef OnCreated) override;
-	virtual void RemoveRow(TypedElementDataStorage::RowHandle Row) override;
-	virtual bool IsRowAvailable(TypedElementDataStorage::RowHandle Row) const override;
-	virtual bool IsRowAssigned(TypedElementDataStorage::RowHandle Row) const override;
+	virtual void RemoveRow(RowHandle Row) override;
+	virtual bool IsRowAvailable(RowHandle Row) const override;
+	virtual bool IsRowAssigned(RowHandle Row) const override;
 
-	virtual void AddColumn(TypedElementRowHandle Row, const UScriptStruct* ColumnType) override;
-	virtual void AddColumn(TypedElementRowHandle Row, const UE::Editor::DataStorage::FDynamicTag& Tag, const FName& InValue) override;
-	virtual void AddColumnData(TypedElementRowHandle Row, const UScriptStruct* ColumnType,
+	virtual void AddColumn(RowHandle Row, const UScriptStruct* ColumnType) override;
+	virtual void AddColumn(RowHandle Row, const UE::Editor::DataStorage::FDynamicTag& Tag, const FName& InValue) override;
+	virtual void AddColumnData(RowHandle Row, const UScriptStruct* ColumnType,
 		const ColumnCreationCallbackRef& Initializer,
 		ColumnCopyOrMoveCallback Relocator) override;
-	virtual void RemoveColumn(TypedElementRowHandle Row, const UScriptStruct* ColumnType) override;
-	virtual void RemoveColumn(TypedElementRowHandle Row, const UE::Editor::DataStorage::FDynamicTag& Tag) override;
-	virtual void* GetColumnData(TypedElementRowHandle Row, const UScriptStruct* ColumnType) override;
-	virtual const void* GetColumnData(TypedElementRowHandle Row, const UScriptStruct* ColumnType) const override;
-	virtual void AddColumns(TypedElementRowHandle Row, TConstArrayView<const UScriptStruct*> Columns) override;
-	virtual void RemoveColumns(TypedElementRowHandle Row, TConstArrayView<const UScriptStruct*> Columns) override;
-	virtual void AddRemoveColumns(TypedElementRowHandle Row, TConstArrayView<const UScriptStruct*> ColumnsToAdd,
+	virtual void RemoveColumn(RowHandle Row, const UScriptStruct* ColumnType) override;
+	virtual void RemoveColumn(RowHandle Row, const UE::Editor::DataStorage::FDynamicTag& Tag) override;
+	virtual void* GetColumnData(RowHandle Row, const UScriptStruct* ColumnType) override;
+	virtual const void* GetColumnData(RowHandle Row, const UScriptStruct* ColumnType) const override;
+	virtual void AddColumns(RowHandle Row, TConstArrayView<const UScriptStruct*> Columns) override;
+	virtual void RemoveColumns(RowHandle Row, TConstArrayView<const UScriptStruct*> Columns) override;
+	virtual void AddRemoveColumns(RowHandle Row, TConstArrayView<const UScriptStruct*> ColumnsToAdd,
 		TConstArrayView<const UScriptStruct*> ColumnsToRemove) override;
-	virtual void BatchAddRemoveColumns(TConstArrayView<TypedElementRowHandle> Rows,TConstArrayView<const UScriptStruct*> ColumnsToAdd,
+	virtual void BatchAddRemoveColumns(TConstArrayView<RowHandle> Rows,TConstArrayView<const UScriptStruct*> ColumnsToAdd,
 		TConstArrayView<const UScriptStruct*> ColumnsToRemove) override;
-	virtual bool HasColumns(TypedElementRowHandle Row, TConstArrayView<const UScriptStruct*> ColumnTypes) const override;
-	virtual bool HasColumns(TypedElementRowHandle Row, TConstArrayView<TWeakObjectPtr<const UScriptStruct>> ColumnTypes) const override;
-	virtual void ListColumns(TypedElementDataStorage::RowHandle Row, ColumnListCallbackRef Callback) const;
-	virtual void ListColumns(TypedElementDataStorage::RowHandle Row, ColumnListWithDataCallbackRef Callback);
-	virtual bool MatchesColumns(TypedElementDataStorage::RowHandle Row, const TypedElementDataStorage::FQueryConditions& Conditions) const override;
+	virtual bool HasColumns(RowHandle Row, TConstArrayView<const UScriptStruct*> ColumnTypes) const override;
+	virtual bool HasColumns(RowHandle Row, TConstArrayView<TWeakObjectPtr<const UScriptStruct>> ColumnTypes) const override;
+	virtual void ListColumns(RowHandle Row, ColumnListCallbackRef Callback) const;
+	virtual void ListColumns(RowHandle Row, ColumnListWithDataCallbackRef Callback);
+	virtual bool MatchesColumns(RowHandle Row, const TypedElementDataStorage::FQueryConditions& Conditions) const override;
 
 	const UScriptStruct* FindDynamicColumn(const UE::Editor::DataStorage::FDynamicColumnDescription& Description) const override;
 	const UScriptStruct* GenerateDynamicColumn(const UE::Editor::DataStorage::FDynamicColumnDescription& Description) override;
@@ -139,24 +142,24 @@ public:
 	void RegisterTickGroup(FName GroupName, EQueryTickPhase Phase, FName BeforeGroup, FName AfterGroup, TypedElementDataStorage::EExecutionMode ExecutionMode);
 	void UnregisterTickGroup(FName GroupName, EQueryTickPhase Phase);
 
-	TypedElementQueryHandle RegisterQuery(FQueryDescription&& Query) override;
-	virtual void UnregisterQuery(TypedElementQueryHandle Query) override;
-	virtual const FQueryDescription& GetQueryDescription(TypedElementQueryHandle Query) const override;
+	QueryHandle RegisterQuery(FQueryDescription&& Query) override;
+	virtual void UnregisterQuery(QueryHandle Query) override;
+	virtual const FQueryDescription& GetQueryDescription(QueryHandle Query) const override;
 	virtual FName GetQueryTickGroupName(EQueryTickGroups Group) const override;
-	virtual FQueryResult RunQuery(TypedElementQueryHandle Query) override;
-	virtual FQueryResult RunQuery(TypedElementQueryHandle Query, DirectQueryCallbackRef Callback) override;
-	virtual FQueryResult RunQuery(TypedElementQueryHandle Query, TypedElementDataStorage::EDirectQueryExecutionFlags Flags, 
+	virtual FQueryResult RunQuery(QueryHandle Query) override;
+	virtual FQueryResult RunQuery(QueryHandle Query, DirectQueryCallbackRef Callback) override;
+	virtual FQueryResult RunQuery(QueryHandle Query, TypedElementDataStorage::EDirectQueryExecutionFlags Flags, 
 		DirectQueryCallbackRef Callback) override;
 	virtual void ActivateQueries(FName ActivationName) override;
 
-	virtual TypedElementDataStorage::RowHandle FindIndexedRow(TypedElementDataStorage::IndexHash Index) const override;
-	virtual void IndexRow(TypedElementDataStorage::IndexHash Index, TypedElementDataStorage::RowHandle Row) override;
+	virtual RowHandle FindIndexedRow(TypedElementDataStorage::IndexHash Index) const override;
+	virtual void IndexRow(TypedElementDataStorage::IndexHash Index, RowHandle Row) override;
 	virtual void BatchIndexRows(
-		TConstArrayView<TPair<TypedElementDataStorage::IndexHash, TypedElementDataStorage::RowHandle>> IndexRowPairs) override;
+		TConstArrayView<TPair<TypedElementDataStorage::IndexHash, RowHandle>> IndexRowPairs) override;
 	virtual void ReindexRow(
 		TypedElementDataStorage::IndexHash OriginalIndex, 
 		TypedElementDataStorage::IndexHash NewIndex, 
-		TypedElementDataStorage::RowHandle Row) override;
+		RowHandle Row) override;
 	virtual void RemoveIndex(TypedElementDataStorage::IndexHash Index) override;
 
 	virtual FTypedElementOnDataStorageUpdate& OnUpdate() override;
@@ -170,7 +173,7 @@ public:
 	TSharedPtr<UE::Editor::DataStorage::FEnvironment> GetEnvironment();
 	TSharedPtr<const UE::Editor::DataStorage::FEnvironment> GetEnvironment() const;
 
-	FMassArchetypeHandle LookupArchetype(TypedElementDataStorage::TableHandle InTableHandle) const;
+	FMassArchetypeHandle LookupArchetype(TableHandle InTableHandle) const;
 
 	void DebugPrintQueryCallbacks(FOutputDevice& Output);
 
@@ -200,7 +203,7 @@ private:
 	static const FName TickGroupName_SyncDataStorageToExternal;
 	
 	TArray<FMassArchetypeHandle> Tables;
-	TMap<FName, TypedElementDataStorage::TableHandle> TableNameLookup;
+	TMap<FName, TableHandle> TableNameLookup;
 
 	// Ordered array of factories by the return value of GetOrder()
 	TArray<FFactoryTypePair> Factories;
diff --git a/Engine/Plugins/Experimental/EditorDataStorage/Source/TedsCore/Private/TypedElementDatabaseCommandBuffer.cpp b/Engine/Plugins/Experimental/EditorDataStorage/Source/TedsCore/Private/TypedElementDatabaseCommandBuffer.cpp
--- a/Engine/Plugins/Experimental/EditorDataStorage/Source/TedsCore/Private/TypedElementDatabaseCommandBuffer.cpp
+++ b/Engine/Plugins/Experimental/EditorDataStorage/Source/TedsCore/Private/TypedElementDatabaseCommandBuffer.cpp
@@ -171,18 +171,18 @@ namespace UE::Editor::DataStorage::Legacy
 	// Execute section
 	//
 
-	bool FCommandBuffer::Execute_IsRowAvailable(const FMassEntityManager& MassEntityManager, TypedElementDataStorage::RowHandle Row)
+	bool FCommandBuffer::Execute_IsRowAvailable(const FMassEntityManager& MassEntityManager, RowHandle Row)
 	{
-		if (Row == TypedElementInvalidRowHandle)
+		if (Row == InvalidRowHandle)
 		{
 			return false;
 		}
 		return MassEntityManager.IsEntityValid(FMassEntityHandle::FromNumber(Row));
 	}
 
-	bool FCommandBuffer::Execute_IsRowAssigned(const FMassEntityManager& MassEntityManager, TypedElementDataStorage::RowHandle Row)
+	bool FCommandBuffer::Execute_IsRowAssigned(const FMassEntityManager& MassEntityManager, RowHandle Row)
 	{
-		if (Row == TypedElementInvalidRowHandle)
+		if (Row == InvalidRowHandle)
 		{
 			return false;
 		}
@@ -214,7 +214,7 @@ namespace UE::Editor::DataStorage::Legacy
 	}
 
 	void FCommandBuffer::Execute_AddDataColumnCommand(
-		FMassEntityManager& MassEntityManager, TypedElementDataStorage::RowHandle Row, const UScriptStruct* ColumnType,
+		FMassEntityManager& MassEntityManager, RowHandle Row, const UScriptStruct* ColumnType,
 		void* Data, ColumnCopyOrMoveCallback Relocator)
 	{
 		if (ColumnType)
@@ -240,7 +240,7 @@ namespace UE::Editor::DataStorage::Legacy
 		}
 	}
 
-	void FCommandBuffer::Execute_AddSharedColumnCommand(FMassEntityManager& MassEntityManager, TypedElementDataStorage::RowHandle Row, const FConstSharedStruct& SharedColumn)
+	void FCommandBuffer::Execute_AddSharedColumnCommand(FMassEntityManager& MassEntityManager, RowHandle Row, const FConstSharedStruct& SharedColumn)
 	{
 		if (SharedColumn.IsValid())
 		{
@@ -249,7 +249,7 @@ namespace UE::Editor::DataStorage::Legacy
 		}
 	}
 
-	void FCommandBuffer::Execute_RemoveSharedColumnCommand(FMassEntityManager& MassEntityManager, TypedElementDataStorage::RowHandle Row, const UScriptStruct& ColumnType)
+	void FCommandBuffer::Execute_RemoveSharedColumnCommand(FMassEntityManager& MassEntityManager, RowHandle Row, const UScriptStruct& ColumnType)
 	{
 		if (ColumnType.IsChildOf(FTedsSharedColumn::StaticStruct()))
 		{
@@ -290,7 +290,7 @@ namespace UE::Editor::DataStorage::Legacy
 
 	void FCommandBuffer::Execute_RemoveColumnsCommand(
 		FMassEntityManager& MassEntityManager,
-		TypedElementDataStorage::RowHandle Row,
+		RowHandle Row,
 		FMassFragmentBitSet FragmentsToRemove,
 		FMassTagBitSet TagsToRemove)
 	{
diff --git a/Engine/Plugins/Experimental/EditorDataStorage/Source/TedsCore/Private/TypedElementDatabaseCommandBuffer.h b/Engine/Plugins/Experimental/EditorDataStorage/Source/TedsCore/Private/TypedElementDatabaseCommandBuffer.h
--- a/Engine/Plugins/Experimental/EditorDataStorage/Source/TedsCore/Private/TypedElementDatabaseCommandBuffer.h
+++ b/Engine/Plugins/Experimental/EditorDataStorage/Source/TedsCore/Private/TypedElementDatabaseCommandBuffer.h
@@ -25,47 +25,47 @@ namespace UE::Editor::DataStorage
 			explicit FCommandBuffer(FEnvironment& InEnvironment);
 	
 			/** Returns a pointer to the queued data column, if it exists and hasn't been processed yet. Otherwise null is returned. */
-			void* GetQueuedDataColumn(TypedElementDataStorage::RowHandle Row, const UScriptStruct* ColumnType);
+			void* GetQueuedDataColumn(RowHandle Row, const UScriptStruct* ColumnType);
 			/** Returns if the column on the provided row is pending processing. */
-			bool HasColumn(TypedElementDataStorage::RowHandle Row, const UScriptStruct* ColumnType) const;
+			bool HasColumn(RowHandle Row, const UScriptStruct* ColumnType) const;
 			/** Clears all entries in the internal cache for the provided row. */
-			void Clear(TypedElementDataStorage::RowHandle Row);
+			void Clear(RowHandle Row);
 	
 			/**
 			 * @section Queue_* functions create a command that can be stored for later execution in the provided command buffer.
 			 * Typically the Execute_* counter part will be called for execution.
 			 */
-			void Queue_AddColumnCommand(TypedElementDataStorage::RowHandle Row, const UScriptStruct* ColumnType);
-			void* Queue_AddDataColumnCommandUnitialized(TypedElementDataStorage::RowHandle Row, const UScriptStruct* ColumnType,
+			void Queue_AddColumnCommand(RowHandle Row, const UScriptStruct* ColumnType);
+			void* Queue_AddDataColumnCommandUnitialized(RowHandle Row, const UScriptStruct* ColumnType,
 				ColumnCopyOrMoveCallback Relocator);
-			void Queue_AddColumnsCommand(TypedElementDataStorage::RowHandle Row, FMassFragmentBitSet FragmentsToAdd, FMassTagBitSet TagsToAdd);
+			void Queue_AddColumnsCommand(RowHandle Row, FMassFragmentBitSet FragmentsToAdd, FMassTagBitSet TagsToAdd);
 	
-			void Queue_RemoveColumnCommand(TypedElementDataStorage::RowHandle Row, const UScriptStruct* ColumnType);
-			void Queue_RemoveColumnsCommand(TypedElementDataStorage::RowHandle Row, FMassFragmentBitSet FragmentsToRemove, FMassTagBitSet TagsToRemove);
+			void Queue_RemoveColumnCommand(RowHandle Row, const UScriptStruct* ColumnType);
+			void Queue_RemoveColumnsCommand(RowHandle Row, FMassFragmentBitSet FragmentsToRemove, FMassTagBitSet TagsToRemove);
 	
 			/**
 			 * @section Execute_* functions directly execute a command with limited validation checks.
 			 */
 	
-			static bool Execute_IsRowAvailable(const FMassEntityManager& MassEntityManager, TypedElementDataStorage::RowHandle Row);
-			static bool Execute_IsRowAssigned(const FMassEntityManager& MassEntityManager, TypedElementDataStorage::RowHandle Row);
+			static bool Execute_IsRowAvailable(const FMassEntityManager& MassEntityManager, RowHandle Row);
+			static bool Execute_IsRowAssigned(const FMassEntityManager& MassEntityManager, RowHandle Row);
 	
 			static void Execute_AddColumnCommand(
-				FMassEntityManager& MassEntityManager, TypedElementDataStorage::RowHandle Row, const UScriptStruct* ColumnType);
+				FMassEntityManager& MassEntityManager, RowHandle Row, const UScriptStruct* ColumnType);
 			static void Execute_AddDataColumnCommand(
-				FMassEntityManager& MassEntityManager, TypedElementDataStorage::RowHandle Row, const UScriptStruct* ColumnType,
+				FMassEntityManager& MassEntityManager, RowHandle Row, const UScriptStruct* ColumnType,
 				void* Data, ColumnCopyOrMoveCallback Relocator);
 			static void Execute_AddSharedColumnCommand(
-				FMassEntityManager& MassEntityManager, TypedElementDataStorage::RowHandle Row, const FConstSharedStruct& SharedColumn);
+				FMassEntityManager& MassEntityManager, RowHandle Row, const FConstSharedStruct& SharedColumn);
 			static void Execute_RemoveSharedColumnCommand(
-				FMassEntityManager& MassEntityManager, TypedElementDataStorage::RowHandle Row, const UScriptStruct& ColumnType);
+				FMassEntityManager& MassEntityManager, RowHandle Row, const UScriptStruct& ColumnType);
 			static void Execute_AddColumnsCommand(
-				FMassEntityManager& MassEntityManager, TypedElementDataStorage::RowHandle Row,
+				FMassEntityManager& MassEntityManager, RowHandle Row,
 				FMassFragmentBitSet FragmentsToAdd, FMassTagBitSet TagsToAdd);
 			static void Execute_RemoveColumnCommand(
-				FMassEntityManager& MassEntityManager, TypedElementDataStorage::RowHandle Row, const UScriptStruct* ColumnType);
+				FMassEntityManager& MassEntityManager, RowHandle Row, const UScriptStruct* ColumnType);
 			static void Execute_RemoveColumnsCommand(
-				FMassEntityManager& MassEntityManager, TypedElementDataStorage::RowHandle Row,
+				FMassEntityManager& MassEntityManager, RowHandle Row,
 				FMassFragmentBitSet FragmentsToRemove, FMassTagBitSet TagsToRemove);
 	
 			void ProcessCommands();
@@ -106,14 +106,14 @@ namespace UE::Editor::DataStorage
 	
 			struct FCommand
 			{
-				TypedElementDataStorage::RowHandle Row;
+				RowHandle Row;
 				CommandData Data;
 			};
 	
 			template<typename T>
-			void AddCommand(TypedElementDataStorage::RowHandle Row, T&& Args);
+			void AddCommand(RowHandle Row, T&& Args);
 	
-			using PendingColumnMappingKey = TPair<TypedElementDataStorage::RowHandle, TWeakObjectPtr<const UScriptStruct>>;
+			using PendingColumnMappingKey = TPair<RowHandle, TWeakObjectPtr<const UScriptStruct>>;
 			TMap<PendingColumnMappingKey, void*> PendingColumns;
 			TArray<FCommand> Commands;
 			FEnvironment& Environment;
diff --git a/Engine/Plugins/Experimental/EditorDataStorage/Source/TedsCore/Private/TypedElementDatabaseCompatibility.cpp b/Engine/Plugins/Experimental/EditorDataStorage/Source/TedsCore/Private/TypedElementDatabaseCompatibility.cpp
--- a/Engine/Plugins/Experimental/EditorDataStorage/Source/TedsCore/Private/TypedElementDatabaseCompatibility.cpp
+++ b/Engine/Plugins/Experimental/EditorDataStorage/Source/TedsCore/Private/TypedElementDatabaseCompatibility.cpp
@@ -202,23 +202,23 @@ void UEditorDataStorageCompatibility::UnregisterObjectRemovedCallback(FDelegateH
 	// Since removing object has be immediately executed in some situation, adding the callback can not be delayed through the command buffer.
 	FScopedExclusiveLock Lock(EGlobalLockScope::Public);
 
-	PreObjectRemovedCallbackList.RemoveAll([Handle](const TPair<UE::Editor::DataStorage::ObjectRemovedCallback, FDelegateHandle>& Element)->bool
+	PreObjectRemovedCallbackList.RemoveAll([Handle](const TPair<ObjectRemovedCallback, FDelegateHandle>& Element)->bool
 	{
 		return Element.Value == Handle;
 	});
 }
 
-TypedElementRowHandle UEditorDataStorageCompatibility::AddCompatibleObjectExplicit(UObject* Object)
+UE::Editor::DataStorage::RowHandle UEditorDataStorageCompatibility::AddCompatibleObjectExplicit(UObject* Object)
 {
 	// Because AddCompatibleObjectExplicitTransactionable needs a finer grained control over the lock, there's no higher up lock here.
 
 	bool bCanAddObject =
 		ensureMsgf(Storage, TEXT("Trying to add a UObject to Typed Element's Data Storage before the storage is available.")) &&
 		ShouldAddObject(Object);
-	return bCanAddObject ? AddCompatibleObjectExplicitTransactionable<true>(Object) : TypedElementDataStorage::InvalidRowHandle;
+	return bCanAddObject ? AddCompatibleObjectExplicitTransactionable<true>(Object) : UE::Editor::DataStorage::InvalidRowHandle;
 }
 
-TypedElementRowHandle UEditorDataStorageCompatibility::AddCompatibleObjectExplicit(void* Object, TWeakObjectPtr<UScriptStruct> TypeInfo)
+UE::Editor::DataStorage::RowHandle UEditorDataStorageCompatibility::AddCompatibleObjectExplicit(void* Object, TWeakObjectPtr<UScriptStruct> TypeInfo)
 {
 	using namespace TypedElementDataStorage;
 	using namespace UE::Editor::DataStorage;
@@ -277,7 +277,7 @@ void UEditorDataStorageCompatibility::RemoveCompatibleObjectExplicit(void* Objec
 	}
 }
 
-TypedElementRowHandle UEditorDataStorageCompatibility::FindRowWithCompatibleObjectExplicit(const UObject* Object) const
+UE::Editor::DataStorage::RowHandle UEditorDataStorageCompatibility::FindRowWithCompatibleObjectExplicit(const UObject* Object) const
 {
 	using namespace TypedElementDataStorage;
 	using namespace UE::Editor::DataStorage;
@@ -292,7 +292,7 @@ TypedElementRowHandle UEditorDataStorageCompatibility::FindRowWithCompatibleObje
 	return InvalidRowHandle;
 }
 
-TypedElementRowHandle UEditorDataStorageCompatibility::FindRowWithCompatibleObjectExplicit(const void* Object) const
+UE::Editor::DataStorage::RowHandle UEditorDataStorageCompatibility::FindRowWithCompatibleObjectExplicit(const void* Object) const
 {
 	// Thread safety is only needed by FindIndexedRow which internally takes care of it.
 
@@ -438,14 +438,14 @@ TypedElementDataStorage::TableHandle UEditorDataStorageCompatibility::FindBestMa
 }
 
 template<bool bEnableTransactions>
-TypedElementRowHandle UEditorDataStorageCompatibility::AddCompatibleObjectExplicitTransactionable(UObject* Object)
+UE::Editor::DataStorage::RowHandle UEditorDataStorageCompatibility::AddCompatibleObjectExplicitTransactionable(UObject* Object)
 {
 	using namespace TypedElementDataStorage;
 	using namespace UE::Editor::DataStorage;
 
 	FScopedExclusiveLock Lock(EGlobalLockScope::Public);
 
-	TypedElementRowHandle Result = FindRowWithCompatibleObjectExplicit(Object);
+	UE::Editor::DataStorage::RowHandle Result = FindRowWithCompatibleObjectExplicit(Object);
 	if (!Storage->IsRowAvailable(Result))
 	{
 		Result = Storage->ReserveRow();
@@ -546,13 +546,13 @@ void UEditorDataStorageCompatibility::RemoveCompatibleObjectExplicitTransactiona
 	}
 }
 
-TypedElementRowHandle UEditorDataStorageCompatibility::DealiasObject(const UObject* Object) const
+UE::Editor::DataStorage::RowHandle UEditorDataStorageCompatibility::DealiasObject(const UObject* Object) const
 {
 	// Thread safe because it's only called from functions that already lock.
 
 	for (const ObjectToRowDealiaser& Dealiaser : ObjectToRowDialiasers)
 	{
-		if (TypedElementRowHandle Row = Dealiaser(*this, Object); Storage->IsRowAvailable(Row))
+		if (UE::Editor::DataStorage::RowHandle Row = Dealiaser(*this, Object); Storage->IsRowAvailable(Row))
 		{
 			return Row;
 		}
@@ -708,7 +708,7 @@ TOptional<TWeakObjectPtr<UObject>> UEditorDataStorageCompatibility::FPendingType
 //
 
 template<typename AddressType>
-void UEditorDataStorageCompatibility::PendingRegistration<AddressType>::Add(TypedElementRowHandle ReservedRowHandle, AddressType Address)
+void UEditorDataStorageCompatibility::PendingRegistration<AddressType>::Add(UE::Editor::DataStorage::RowHandle ReservedRowHandle, AddressType Address)
 {
 	// Thread-safe as it's only called from functions that already lock using.
 	Entries.Emplace(FEntry{ .Address = Address, .Row = ReservedRowHandle });
@@ -740,7 +740,7 @@ void UEditorDataStorageCompatibility::PendingRegistration<AddressType>::ForEachA
 
 template<typename AddressType>
 void UEditorDataStorageCompatibility::PendingRegistration<AddressType>::ProcessEntries(ITypedElementDataStorageInterface& StorageInterface,
-	UEditorDataStorageCompatibility& Compatibility, const TFunctionRef<void(TypedElementRowHandle, const AddressType&)>& SetupRowCallback)
+	UEditorDataStorageCompatibility& Compatibility, const TFunctionRef<void(UE::Editor::DataStorage::RowHandle, const AddressType&)>& SetupRowCallback)
 {
 	// Thread-safe as it's only called from functions that already lock using.
 	using namespace TypedElementDataStorage;
@@ -825,7 +825,7 @@ void UEditorDataStorageCompatibility::PendingRegistration<AddressType>::ProcessE
 			if (Current->Table != CurrentTable)
 			{
 				StorageInterface.BatchAddRow(CurrentTable, Compatibility.RowScratchBuffer,
-					[&SetupRowCallback, &TableFront](TypedElementRowHandle Row)
+					[&SetupRowCallback, &TableFront](UE::Editor::DataStorage::RowHandle Row)
 					{
 						SetupRowCallback(Row, TableFront->Address);
 						++TableFront;
@@ -837,7 +837,7 @@ void UEditorDataStorageCompatibility::PendingRegistration<AddressType>::ProcessE
 			Compatibility.RowScratchBuffer.Add(Current->Row);
 		}
 		StorageInterface.BatchAddRow(CurrentTable, Compatibility.RowScratchBuffer,
-			[&SetupRowCallback, &TableFront](TypedElementRowHandle Row)
+			[&SetupRowCallback, &TableFront](UE::Editor::DataStorage::RowHandle Row)
 			{
 				SetupRowCallback(Row, TableFront->Address);
 				++TableFront;
@@ -907,7 +907,7 @@ void UEditorDataStorageCompatibility::TickPendingUObjectRegistration()
 	if (!UObjectsPendingRegistration.IsEmpty())
 	{
 		UObjectsPendingRegistration.ProcessEntries(*Storage, *this,
-			[this](TypedElementRowHandle Row, const TWeakObjectPtr<UObject>& Object)
+			[this](UE::Editor::DataStorage::RowHandle Row, const TWeakObjectPtr<UObject>& Object)
 			{
 				ITypedElementDataStorageInterface* Interface = Storage;
 				Interface->AddColumn(Row, FTypedElementUObjectColumn{ .Object = Object });
@@ -937,7 +937,7 @@ void UEditorDataStorageCompatibility::TickPendingExternalObjectRegistration()
 	if (!ExternalObjectsPendingRegistration.IsEmpty())
 	{
 		ExternalObjectsPendingRegistration.ProcessEntries(*Storage, *this,
-			[this](TypedElementRowHandle Row, const ExternalObjectRegistration& Object)
+			[this](UE::Editor::DataStorage::RowHandle Row, const ExternalObjectRegistration& Object)
 			{
 				ITypedElementDataStorageInterface* Interface = Storage;
 				Interface->AddColumn(Row, FTypedElementExternalObjectColumn{ .Object = Object.Object });
diff --git a/Engine/Plugins/Experimental/EditorDataStorage/Source/TedsCore/Private/TypedElementDatabaseCompatibility.h b/Engine/Plugins/Experimental/EditorDataStorage/Source/TedsCore/Private/TypedElementDatabaseCompatibility.h
--- a/Engine/Plugins/Experimental/EditorDataStorage/Source/TedsCore/Private/TypedElementDatabaseCompatibility.h
+++ b/Engine/Plugins/Experimental/EditorDataStorage/Source/TedsCore/Private/TypedElementDatabaseCompatibility.h
@@ -48,20 +48,20 @@ public:
 
 	void RegisterRegistrationFilter(ObjectRegistrationFilter Filter) override;
 	void RegisterDealiaserCallback(ObjectToRowDealiaser Dealiaser) override;
-	void RegisterTypeTableAssociation(TObjectPtr<UStruct> TypeInfo, TypedElementDataStorage::TableHandle Table) override;
+	void RegisterTypeTableAssociation(TObjectPtr<UStruct> TypeInfo, UE::Editor::DataStorage::TableHandle Table) override;
 	FDelegateHandle RegisterObjectAddedCallback(UE::Editor::DataStorage::ObjectAddedCallback&& OnObjectAdded);
 	void UnregisterObjectAddedCallback(FDelegateHandle Handle);
 	FDelegateHandle RegisterObjectRemovedCallback(UE::Editor::DataStorage::ObjectRemovedCallback&& OnObjectRemoved);
 	void UnregisterObjectRemovedCallback(FDelegateHandle Handle);
 	
-	TypedElementRowHandle AddCompatibleObjectExplicit(UObject* Object) override;
-	TypedElementRowHandle AddCompatibleObjectExplicit(void* Object, TWeakObjectPtr<UScriptStruct> TypeInfo) override;
+	UE::Editor::DataStorage::RowHandle AddCompatibleObjectExplicit(UObject* Object) override;
+	UE::Editor::DataStorage::RowHandle AddCompatibleObjectExplicit(void* Object, TWeakObjectPtr<UScriptStruct> TypeInfo) override;
 	
 	void RemoveCompatibleObjectExplicit(UObject* Object) override;
 	void RemoveCompatibleObjectExplicit(void* Object) override;
 
-	TypedElementRowHandle FindRowWithCompatibleObjectExplicit(const UObject* Object) const override;
-	TypedElementRowHandle FindRowWithCompatibleObjectExplicit(const void* Object) const override;
+	UE::Editor::DataStorage::RowHandle FindRowWithCompatibleObjectExplicit(const UObject* Object) const override;
+	UE::Editor::DataStorage::RowHandle FindRowWithCompatibleObjectExplicit(const void* Object) const override;
 
 	bool SupportsExtension(FName Extension) const override;
 	void ListExtensions(TFunctionRef<void(FName)> Callback) const override;
@@ -83,7 +83,7 @@ private:
 	private:
 		TWeakObjectPtr<UEditorDataStorageCompatibility> Owner;
 		TWeakObjectPtr<UObject> TargetObject;
-		TypedElementDataStorage::RowHandle MementoRow = TypedElementDataStorage::InvalidRowHandle;
+		UE::Editor::DataStorage::RowHandle MementoRow = UE::Editor::DataStorage::InvalidRowHandle;
 	};
 	class FDeregistrationCommandChange final : public FCommandChange
 	{
@@ -98,7 +98,7 @@ private:
 	private:
 		TWeakObjectPtr<UEditorDataStorageCompatibility> Owner;
 		TWeakObjectPtr<UObject> TargetObject;
-		TypedElementDataStorage::RowHandle MementoRow = TypedElementDataStorage::InvalidRowHandle;
+		UE::Editor::DataStorage::RowHandle MementoRow = UE::Editor::DataStorage::InvalidRowHandle;
 	};
 
 	void Prepare();
@@ -107,14 +107,14 @@ private:
 	void RegisterTypeInformationQueries();
 	
 	bool ShouldAddObject(const UObject* Object) const;
-	TypedElementDataStorage::TableHandle FindBestMatchingTable(const UStruct* TypeInfo) const;
+	UE::Editor::DataStorage::TableHandle FindBestMatchingTable(const UStruct* TypeInfo) const;
 	template<bool bEnableTransactions>
-	TypedElementRowHandle AddCompatibleObjectExplicitTransactionable(UObject* Object);
+	UE::Editor::DataStorage::RowHandle AddCompatibleObjectExplicitTransactionable(UObject* Object);
 	template<bool bEnableTransactions>
 	void RemoveCompatibleObjectExplicitTransactionable(const UObject* Object);
 	template<bool bEnableTransactions>
-	void RemoveCompatibleObjectExplicitTransactionable(const UObject* Object, TypedElementDataStorage::RowHandle ObjectRow);
-	TypedElementRowHandle DealiasObject(const UObject* Object) const;
+	void RemoveCompatibleObjectExplicitTransactionable(const UObject* Object, UE::Editor::DataStorage::RowHandle ObjectRow);
+	UE::Editor::DataStorage::RowHandle DealiasObject(const UObject* Object) const;
 
 	void Tick();
 	void TickPendingCommands();
@@ -125,8 +125,8 @@ private:
 	void OnPrePropertyChanged(UObject* Object, const FEditPropertyChain& PropertyChain);
 	void OnPostEditChangeProperty(UObject* Object, FPropertyChangedEvent& PropertyChangedEvent);
 	void OnObjectModified(UObject* Object);
-	void TriggerOnObjectAdded(const void* Object, UE::Editor::DataStorage::FObjectTypeInfo TypeInfo, TypedElementDataStorage::RowHandle Row) const;
-	void TriggerOnPreObjectRemoved(const void* Object, UE::Editor::DataStorage::FObjectTypeInfo TypeInfo, TypedElementDataStorage::RowHandle Row) const;
+	void TriggerOnObjectAdded(const void* Object, UE::Editor::DataStorage::FObjectTypeInfo TypeInfo, UE::Editor::DataStorage::RowHandle Row) const;
+	void TriggerOnPreObjectRemoved(const void* Object, UE::Editor::DataStorage::FObjectTypeInfo TypeInfo, UE::Editor::DataStorage::RowHandle Row) const;
 	void OnObjectReinstanced(const FCoreUObjectDelegates::FReplacementObjectMap& ReplacedObjects);
 
 	void OnPostGcUnreachableAnalysis();
@@ -156,7 +156,7 @@ private:
 		PendingTypeInformationMap PendingTypeInformationUpdates[2];
 		PendingTypeInformationMap* PendingTypeInformationUpdatesActive;
 		PendingTypeInformationMap* PendingTypeInformationUpdatesSwapped;
-		TArray<TTuple<TWeakObjectPtr<UStruct>, TypedElementDataStorage::TableHandle>> UpdatedTypeInfoScratchBuffer;
+		TArray<TTuple<TWeakObjectPtr<UStruct>, UE::Editor::DataStorage::TableHandle>> UpdatedTypeInfoScratchBuffer;
 		UE::FMutex Safeguard;
 		std::atomic<bool> bHasPendingUpdate = false;
 	};
@@ -175,19 +175,19 @@ private:
 		struct FEntry
 		{
 			AddressType Address;
-			TypedElementDataStorage::RowHandle Row;
-			TypedElementDataStorage::TableHandle Table;
+			UE::Editor::DataStorage::RowHandle Row;
+			UE::Editor::DataStorage::TableHandle Table;
 		};
 		TArray<FEntry> Entries;
 
 	public:
-		void Add(TypedElementRowHandle ReservedRowHandle, AddressType Address);
+		void Add(UE::Editor::DataStorage::RowHandle ReservedRowHandle, AddressType Address);
 		bool IsEmpty() const;
 		int32 Num() const;
 		
 		void ForEachAddress(const TFunctionRef<void(AddressType&)>& Callback);
 		void ProcessEntries(ITypedElementDataStorageInterface& Storage, UEditorDataStorageCompatibility& Compatibility,
-			const TFunctionRef<void(TypedElementRowHandle, const AddressType&)>& SetupRowCallback);
+			const TFunctionRef<void(UE::Editor::DataStorage::RowHandle, const AddressType&)>& SetupRowCallback);
 		void Reset();
 	};
 
@@ -195,11 +195,11 @@ private:
 	UE::Editor::DataStorage::CompatibilityCommandBuffer::FCollection PendingCommands;
 	PendingRegistration<TWeakObjectPtr<UObject>> UObjectsPendingRegistration;
 	PendingRegistration<ExternalObjectRegistration> ExternalObjectsPendingRegistration;
-	TArray<TypedElementDataStorage::RowHandle> RowScratchBuffer;
+	TArray<UE::Editor::DataStorage::RowHandle> RowScratchBuffer;
 	
 	TArray<ObjectRegistrationFilter> ObjectRegistrationFilters;
 	TArray<ObjectToRowDealiaser> ObjectToRowDialiasers;
-	using TypeToTableMapType = TMap<TWeakObjectPtr<UStruct>, TypedElementDataStorage::TableHandle>;
+	using TypeToTableMapType = TMap<TWeakObjectPtr<UStruct>, UE::Editor::DataStorage::TableHandle>;
 	TypeToTableMapType TypeToTableMap;
 	TArray<TPair<UE::Editor::DataStorage::ObjectAddedCallback, FDelegateHandle>> ObjectAddedCallbackList;
 	TArray<TPair<UE::Editor::DataStorage::ObjectRemovedCallback, FDelegateHandle>> PreObjectRemovedCallbackList;
@@ -239,9 +239,9 @@ private:
 	FDelegateHandle PostGcUnreachableAnalysisHandle;
 	
 	TSharedPtr<UE::Editor::DataStorage::FEnvironment> Environment;
-	TypedElementDataStorage::QueryHandle ClassTypeInfoQuery;
-	TypedElementDataStorage::QueryHandle ScriptStructTypeInfoQuery;
-	TypedElementDataStorage::QueryHandle UObjectQuery;
+	UE::Editor::DataStorage::QueryHandle ClassTypeInfoQuery;
+	UE::Editor::DataStorage::QueryHandle ScriptStructTypeInfoQuery;
+	UE::Editor::DataStorage::QueryHandle UObjectQuery;
 };
 
 SIZE_T GetTypeHash(const UEditorDataStorageCompatibility::FSyncTagInfo& Column);
diff --git a/Engine/Plugins/Experimental/EditorDataStorage/Source/TedsCore/Private/TypedElementDatabaseUI.cpp b/Engine/Plugins/Experimental/EditorDataStorage/Source/TedsCore/Private/TypedElementDatabaseUI.cpp
--- a/Engine/Plugins/Experimental/EditorDataStorage/Source/TedsCore/Private/TypedElementDatabaseUI.cpp
+++ b/Engine/Plugins/Experimental/EditorDataStorage/Source/TedsCore/Private/TypedElementDatabaseUI.cpp
@@ -421,7 +421,7 @@ void UEditorDataStorageUi::CreateWidgetInstance(
 	const TypedElementDataStorage::FMetaDataView& Arguments,
 	const WidgetCreatedCallback& ConstructionCallback)
 {
-	TypedElementRowHandle Row = Storage->AddRow(WidgetTable);
+	UE::Editor::DataStorage::RowHandle Row = Storage->AddRow(WidgetTable);
 	Storage->AddColumns(Row, Constructor.GetAdditionalColumnsList());
 	TSharedPtr<SWidget> Widget = Constructor.ConstructFinalWidget(Row, Storage, this, Arguments);
 	if (Widget)
@@ -434,7 +434,7 @@ void UEditorDataStorageUi::CreateWidgetInstance(
 	}
 }
 
-TSharedPtr<SWidget> UEditorDataStorageUi::ConstructWidget(TypedElementRowHandle Row, FTypedElementWidgetConstructor& Constructor,
+TSharedPtr<SWidget> UEditorDataStorageUi::ConstructWidget(UE::Editor::DataStorage::RowHandle Row, FTypedElementWidgetConstructor& Constructor,
 	const TypedElementDataStorage::FMetaDataView& Arguments)
 {
 	return Constructor.ConstructFinalWidget(Row, Storage, this, Arguments);
diff --git a/Engine/Plugins/Experimental/EditorDataStorage/Source/TedsCore/Private/TypedElementDatabaseUI.h b/Engine/Plugins/Experimental/EditorDataStorage/Source/TedsCore/Private/TypedElementDatabaseUI.h
--- a/Engine/Plugins/Experimental/EditorDataStorage/Source/TedsCore/Private/TypedElementDatabaseUI.h
+++ b/Engine/Plugins/Experimental/EditorDataStorage/Source/TedsCore/Private/TypedElementDatabaseUI.h
@@ -47,7 +47,7 @@ public:
 
 	void ConstructWidgets(FName Purpose, const TypedElementDataStorage::FMetaDataView& Arguments,
 		const WidgetCreatedCallback& ConstructionCallback) override;
-	TSharedPtr<SWidget> ConstructWidget(TypedElementRowHandle Row, FTypedElementWidgetConstructor& Constructor,
+	TSharedPtr<SWidget> ConstructWidget(UE::Editor::DataStorage::RowHandle Row, FTypedElementWidgetConstructor& Constructor,
 		const TypedElementDataStorage::FMetaDataView& Arguments) override;
 
 	void ListWidgetPurposes(const WidgetPurposeCallback& Callback) const override;
@@ -108,7 +108,7 @@ private:
 		const TypedElementDataStorage::FMetaDataView& Arguments,
 		const WidgetConstructorCallback& Callback);
 
-	TypedElementTableHandle WidgetTable{ TypedElementInvalidTableHandle };
+	UE::Editor::DataStorage::TableHandle WidgetTable{ UE::Editor::DataStorage::InvalidTableHandle };
 	
 	TMap<FName, FPurposeInfo> WidgetPurposes;
 	
diff --git a/Engine/Plugins/Experimental/EditorDataStorage/Source/TedsCore/Tests/TypedElementDataStorageInterfaceTests.cpp b/Engine/Plugins/Experimental/EditorDataStorage/Source/TedsCore/Tests/TypedElementDataStorageInterfaceTests.cpp
--- a/Engine/Plugins/Experimental/EditorDataStorage/Source/TedsCore/Tests/TypedElementDataStorageInterfaceTests.cpp
+++ b/Engine/Plugins/Experimental/EditorDataStorage/Source/TedsCore/Tests/TypedElementDataStorageInterfaceTests.cpp
@@ -6,575 +6,578 @@
 #include "Elements/Framework/TypedElementRegistry.h"
 #include "Misc/AutomationTest.h"
 
-namespace TypedElementDataStorageTests
-{	
-	TypedElementTableHandle RegisterTestTableA(ITypedElementDataStorageInterface* TedsInterface)
+namespace UE::Editor::DataStorage
+{
+	namespace Tests
 	{
-		// RegisterTable is not idempotent, so this utility function makes it so
+		TableHandle RegisterTestTableA(ITypedElementDataStorageInterface* TedsInterface)
+		{
+			// RegisterTable is not idempotent, so this utility function makes it so
+		
+			const FName TestTableAName = TEXT("TestTableA");
+			TableHandle Table = TedsInterface->FindTable(TestTableAName);
+			if (Table != InvalidTableHandle)
+			{
+				return Table;
+			}
 		
-		const FName TestTableAName = TEXT("TestTableA");
-		TypedElementTableHandle Table = TedsInterface->FindTable(TestTableAName);
-		if (Table != TypedElementInvalidTableHandle)
+			return TedsInterface->RegisterTable(
+			{
+				FTestColumnA::StaticStruct(),
+				FTestColumnB::StaticStruct(),
+				FTestTagColumnA::StaticStruct(),
+				FTestTagColumnB::StaticStruct()
+			},
+			TEXT("TestTableA"));
+		}
+
+		bool IsEmptyQueryDescription(const TypedElementDataStorage::FQueryDescription& QueryDescription)
 		{
-			return Table;
+			return
+				QueryDescription.Callback.MonitoredType == nullptr &&
+				QueryDescription.SelectionTypes.IsEmpty() &&
+				QueryDescription.SelectionAccessTypes.IsEmpty() &&
+				QueryDescription.SelectionMetaData.IsEmpty() &&
+				QueryDescription.ConditionTypes.IsEmpty() &&
+				QueryDescription.ConditionOperators.IsEmpty() &&
+				QueryDescription.DependencyTypes.IsEmpty() &&
+				QueryDescription.DependencyFlags.IsEmpty() &&
+				QueryDescription.CachedDependencies.IsEmpty() &&
+				QueryDescription.Subqueries.IsEmpty();
 		}
-		
-		return TedsInterface->RegisterTable(
+
+		void IncrementThing(int64* QueryCallCountPtr)
 		{
-			FTestColumnA::StaticStruct(),
-			FTestColumnB::StaticStruct(),
-			FTestTagColumnA::StaticStruct(),
-			FTestTagColumnB::StaticStruct()
-		},
-		TEXT("TestTableA"));
-	}
-
-	bool IsEmptyQueryDescription(const TypedElementDataStorage::FQueryDescription& QueryDescription)
-	{
-		return
-			QueryDescription.Callback.MonitoredType == nullptr &&
-			QueryDescription.SelectionTypes.IsEmpty() &&
-			QueryDescription.SelectionAccessTypes.IsEmpty() &&
-			QueryDescription.SelectionMetaData.IsEmpty() &&
-			QueryDescription.ConditionTypes.IsEmpty() &&
-			QueryDescription.ConditionOperators.IsEmpty() &&
-			QueryDescription.DependencyTypes.IsEmpty() &&
-			QueryDescription.DependencyFlags.IsEmpty() &&
-			QueryDescription.CachedDependencies.IsEmpty() &&
-			QueryDescription.Subqueries.IsEmpty();
-	}
-
-	void IncrementThing(int64* QueryCallCountPtr)
-	{
-		++(*QueryCallCountPtr);
-	}
-
-	// Temporary wrapper to only register a single observer
-	// There is no way to unregister an observer at this time
-	template<typename ColumnType>
-	void RegisterTestObserver(
-		FName ObserverName,
-		ITypedElementDataStorageInterface* TedsInterface, TypedElementQueryBuilder::FObserver::EEvent EventType, int64* QueryCallCountPtr, TypedElementQueryHandle* QueryHandleInOut)
-	{
-		using namespace TypedElementQueryBuilder;
-		using DSI = ITypedElementDataStorageInterface;
-		TypedElementDataStorage::FQueryDescription QueryDescription = TedsInterface->GetQueryDescription(*QueryHandleInOut);
+			++(*QueryCallCountPtr);
+		}
 
-		// Check if it is empty... If not then we need to register
-		// This is for an idempotent API to the test harness				
-		if (IsEmptyQueryDescription(QueryDescription))
+		// Temporary wrapper to only register a single observer
+		// There is no way to unregister an observer at this time
+		template<typename ColumnType>
+		void RegisterTestObserver(
+			FName ObserverName,
+			ITypedElementDataStorageInterface* TedsInterface, TypedElementQueryBuilder::FObserver::EEvent EventType, int64* QueryCallCountPtr, QueryHandle* QueryHandleInOut)
 		{
-			*QueryHandleInOut = TedsInterface->RegisterQuery(
-				Select(
-					ObserverName,
-					FObserver(EventType, ColumnType::StaticStruct()),
-					[QueryCallCountPtr](DSI::IQueryContext& Context, TypedElementRowHandle Row)
-					{
-						++(*QueryCallCountPtr);
-					}).Compile());
+			using namespace TypedElementQueryBuilder;
+			using DSI = ITypedElementDataStorageInterface;
+			TypedElementDataStorage::FQueryDescription QueryDescription = TedsInterface->GetQueryDescription(*QueryHandleInOut);
+
+			// Check if it is empty... If not then we need to register
+			// This is for an idempotent API to the test harness				
+			if (IsEmptyQueryDescription(QueryDescription))
+			{
+				*QueryHandleInOut = TedsInterface->RegisterQuery(
+					Select(
+						ObserverName,
+						FObserver(EventType, ColumnType::StaticStruct()),
+						[QueryCallCountPtr](DSI::IQueryContext& Context, RowHandle Row)
+						{
+							++(*QueryCallCountPtr);
+						}).Compile());
+			}
 		}
-	}
 	
-	BEGIN_DEFINE_SPEC(TypedElementDataStorageTestsFixture, "TypedElementsDataStorage", EAutomationTestFlags::EditorContext | EAutomationTestFlags::EngineFilter)
-		ITypedElementDataStorageInterface* TedsInterface = nullptr;
-		TypedElementTableHandle TestTableHandleA = TypedElementInvalidTableHandle;
-
-		TArray<TypedElementRowHandle> CreatedRows;
-
-		TypedElementQueryHandle DataColumnA_AddObserverHandle = TypedElementInvalidQueryHandle;
-		TypedElementQueryHandle DataColumnC_AddObserverHandle = TypedElementInvalidQueryHandle;
-		TypedElementQueryHandle DataColumnB_RemoveObserverHandle = TypedElementInvalidQueryHandle;
-		TypedElementQueryHandle TagColumnA_AddObserverHandle = TypedElementInvalidQueryHandle;
-		TypedElementQueryHandle TagColumnC_AddObserverHandle = TypedElementInvalidQueryHandle;
-		TypedElementQueryHandle TagColumnB_RemoveObserverHandle = TypedElementInvalidQueryHandle;
+		BEGIN_DEFINE_SPEC(TypedElementDataStorageTestsFixture, "TypedElementsDataStorage", EAutomationTestFlags::EditorContext | EAutomationTestFlags::EngineFilter)
+			ITypedElementDataStorageInterface* TedsInterface = nullptr;
+			TableHandle TestTableHandleA = InvalidTableHandle;
+
+			TArray<RowHandle> CreatedRows;
+
+			QueryHandle DataColumnA_AddObserverHandle = InvalidQueryHandle;
+			QueryHandle DataColumnC_AddObserverHandle = InvalidQueryHandle;
+			QueryHandle DataColumnB_RemoveObserverHandle = InvalidQueryHandle;
+			QueryHandle TagColumnA_AddObserverHandle = InvalidQueryHandle;
+			QueryHandle TagColumnC_AddObserverHandle = InvalidQueryHandle;
+			QueryHandle TagColumnB_RemoveObserverHandle = InvalidQueryHandle;
 	
-		// Non-zero to catch problems in test where the difference is not being used
-		// Avoids need to reset this to 0 for each test
-		int64 DataColumnA_AddObserverCallCount = 10;
-		int64 DataColumnC_AddObserverCallCount = 10;
-		int64 DataColumnB_RemoveObserverCallCount = 10;
-		int64 TagColumnA_AddObserverCallCount = 10;
-		int64 TagColumnC_AddObserverCallCount = 10;
-		int64 TagColumnB_RemoveObserverCallCount = 10;
+			// Non-zero to catch problems in test where the difference is not being used
+			// Avoids need to reset this to 0 for each test
+			int64 DataColumnA_AddObserverCallCount = 10;
+			int64 DataColumnC_AddObserverCallCount = 10;
+			int64 DataColumnB_RemoveObserverCallCount = 10;
+			int64 TagColumnA_AddObserverCallCount = 10;
+			int64 TagColumnC_AddObserverCallCount = 10;
+			int64 TagColumnB_RemoveObserverCallCount = 10;
 	
-		void CreateTestRows(int32 Count)
-		{
-			CreatedRows.Reserve(Count);
-			for (int32 Index = 0; Index < Count; ++Index)
+			void CreateTestRows(int32 Count)
 			{
-				TypedElementRowHandle RowHandle = TedsInterface->AddRow(TestTableHandleA);
-				TestFalse("Expect valid row handle", RowHandle == TypedElementInvalidRowHandle);
-				CreatedRows.Add(RowHandle);
+				CreatedRows.Reserve(Count);
+				for (int32 Index = 0; Index < Count; ++Index)
+				{
+					RowHandle RowHandle = TedsInterface->AddRow(TestTableHandleA);
+					TestFalse("Expect valid row handle", RowHandle == InvalidRowHandle);
+					CreatedRows.Add(RowHandle);
+				}
 			}
-		}
 
-		void CleanupTestRows()
-		{
-			for (TypedElementRowHandle RowHandle : CreatedRows)
+			void CleanupTestRows()
 			{
-				TestTrue("Expected row to have been available", TedsInterface->IsRowAvailable(RowHandle));
-				TestTrue("Expected row to have been assigned to a table", TedsInterface->IsRowAssigned(RowHandle));
-				TedsInterface->RemoveRow(RowHandle);
-				TestFalse("Expected row to have been unassigned a table", TedsInterface->IsRowAssigned(RowHandle));
-				TestFalse("Expected row to be not available", TedsInterface->IsRowAvailable(RowHandle));
+				for (RowHandle RowHandle : CreatedRows)
+				{
+					TestTrue("Expected row to have been available", TedsInterface->IsRowAvailable(RowHandle));
+					TestTrue("Expected row to have been assigned to a table", TedsInterface->IsRowAssigned(RowHandle));
+					TedsInterface->RemoveRow(RowHandle);
+					TestFalse("Expected row to have been unassigned a table", TedsInterface->IsRowAssigned(RowHandle));
+					TestFalse("Expected row to be not available", TedsInterface->IsRowAvailable(RowHandle));
+				}
+				CreatedRows.Empty();
 			}
-			CreatedRows.Empty();
-		}
-	END_DEFINE_SPEC(TypedElementDataStorageTestsFixture)
-
-	void TypedElementDataStorageTestsFixture::Define()
-	{
-		BeforeEach([this]()
-		{
-			UTypedElementRegistry* TypedElementRegistry = UTypedElementRegistry::GetInstance();
-			TedsInterface = TypedElementRegistry->GetMutableDataStorage();
-			TestTrue("", TedsInterface != nullptr);
-		});
+		END_DEFINE_SPEC(TypedElementDataStorageTestsFixture)
 
-		Describe("", [this]()
+		void TypedElementDataStorageTestsFixture::Define()
 		{
 			BeforeEach([this]()
 			{
-				TestTrue("Test requires valid TEDS Interface", TedsInterface != nullptr);
-			});
-			
-			Describe("", [this]()
-			{
-				Describe("RegisterTable(TConstArrayView<const UScriptStruct*>, const FName)", [this]()
-				{
-					It("should register a table and provide a valid handle", [this]()
-					{
-						// Note, this test doesn't really do anything a second time.
-						// The test fixture depends on global state which currently cannot be cleaned
-						// up.
-						TypedElementTableHandle Handle = RegisterTestTableA(TedsInterface);
-						TestNotEqual("Expecting valid table handle", Handle, TypedElementInvalidTableHandle);
-					});
-				});
+				UTypedElementRegistry* TypedElementRegistry = UTypedElementRegistry::GetInstance();
+				TedsInterface = TypedElementRegistry->GetMutableDataStorage();
+				TestTrue("", TedsInterface != nullptr);
 			});
 
 			Describe("", [this]()
 			{
 				BeforeEach([this]()
 				{
-					TestTableHandleA = RegisterTestTableA(TedsInterface);
-				});
-				
-				AfterEach([this]()
-				{
-					CleanupTestRows();
-				});
-				
-				Describe("AddRow(TypedElementTableHandle)", [this]()
-				{
-					for (int32 RowCount = 1; RowCount <= 2; ++RowCount)
-					{
-						It(FString::Printf(TEXT("should should be able to create %d rows when called %d times"), RowCount, RowCount), [this, RowCount]()
-						{
-							for (int32 CreatedRowIndex = 0; CreatedRowIndex < RowCount; ++CreatedRowIndex)
-							{
-								TypedElementRowHandle RowHandle = TedsInterface->AddRow(TestTableHandleA);
-								CreatedRows.Add(RowHandle);
-								TestNotEqual("Expecting valid row", RowHandle, TypedElementInvalidRowHandle);
-							}
-						});
-					}
+					TestTrue("Test requires valid TEDS Interface", TedsInterface != nullptr);
 				});
-
-				Describe("BatchAddRow(TypedElementTableHandle, int32, TypedElementDataStorageCreationCallbackRef)", [this]()
+			
+				Describe("", [this]()
 				{
-					for (int32 RowCount = 1; RowCount <= 2; ++RowCount)
+					Describe("RegisterTable(TConstArrayView<const UScriptStruct*>, const FName)", [this]()
 					{
-						It(FString::Printf(TEXT("should should be able to create %d rows when called once"), RowCount), [this, RowCount]()
+						It("should register a table and provide a valid handle", [this]()
 						{
-							auto RowInitializeCallback = [this](TypedElementRowHandle RowHandle)
-							{
-								CreatedRows.Add(RowHandle);
-								TestNotEqual("Expecting valid row", RowHandle, TypedElementInvalidRowHandle);
-							};
-											
-							TypedElementRowHandle RowHandle = TedsInterface->BatchAddRow(TestTableHandleA, RowCount, RowInitializeCallback);
-											
-							TestEqual("Unexpected number of created rows", CreatedRows.Num(), RowCount);
+							// Note, this test doesn't really do anything a second time.
+							// The test fixture depends on global state which currently cannot be cleaned
+							// up.
+							TableHandle Handle = RegisterTestTableA(TedsInterface);
+							TestNotEqual("Expecting valid table handle", Handle, InvalidTableHandle);
 						});
-					}
+					});
 				});
 
-			
 				Describe("", [this]()
 				{
-					/**
-					 * Prerequisites for this batch of tests are:
-					 * - A test table is registered
-					 * - An add and remove observer is registered for FTestInt64Column
-					 */
 					BeforeEach([this]()
 					{
 						TestTableHandleA = RegisterTestTableA(TedsInterface);
-						TestNotEqual("Expecting valid table handle", TestTableHandleA, TypedElementInvalidTableHandle);
-
-						RegisterTestObserver<FTestColumnA>(
-							TEXT("Increment CallCount when FTestColumnA added"),
-							TedsInterface,
-							TypedElementQueryBuilder::FObserver::EEvent::Add,
-							&DataColumnA_AddObserverCallCount,
-							&DataColumnA_AddObserverHandle);
-						TestTrue("Expect valid query observer handle", DataColumnA_AddObserverHandle != TypedElementInvalidQueryHandle);
-						
-						RegisterTestObserver<FTestColumnC>(
-							TEXT("Increment CallCount when FTestColumnC added"),
-							TedsInterface,
-							TypedElementQueryBuilder::FObserver::EEvent::Add,
-							&DataColumnC_AddObserverCallCount,
-							&DataColumnC_AddObserverHandle);
-						TestTrue("Expect valid query observer handle", DataColumnC_AddObserverHandle != TypedElementInvalidQueryHandle);
-
-						RegisterTestObserver<FTestColumnB>(
-							TEXT("Increment CallCount when FTestColumnB removed"),
-							TedsInterface,
-							TypedElementQueryBuilder::FObserver::EEvent::Remove,
-							&DataColumnB_RemoveObserverCallCount,
-							&DataColumnB_RemoveObserverHandle);
-						TestTrue("Expect valid query observer handle", DataColumnB_RemoveObserverHandle != TypedElementInvalidQueryHandle);
-
-						RegisterTestObserver<FTestTagColumnA>(
-							TEXT("Increment CallCount when FTestColumnA added"),
-							TedsInterface,
-							TypedElementQueryBuilder::FObserver::EEvent::Add,
-							&TagColumnA_AddObserverCallCount,
-							&TagColumnA_AddObserverHandle);
-						TestTrue("Expect valid query observer handle", TagColumnA_AddObserverHandle != TypedElementInvalidQueryHandle);
-												
-						RegisterTestObserver<FTestTagColumnC>(
-							TEXT("Increment CallCount when FTestTagColumnC added"),
-							TedsInterface,
-							TypedElementQueryBuilder::FObserver::EEvent::Add,
-							&TagColumnC_AddObserverCallCount, 
-							&TagColumnC_AddObserverHandle);
-						TestTrue("Expect valid query observer handle", TagColumnC_AddObserverHandle != TypedElementInvalidQueryHandle);
-
-						RegisterTestObserver<FTestTagColumnB>(
-							TEXT("Increment CallCount when FTestTagColumnB removed"),
-							TedsInterface, TypedElementQueryBuilder::FObserver::EEvent::Remove,
-							&TagColumnB_RemoveObserverCallCount,
-							&TagColumnB_RemoveObserverHandle);
-						TestTrue("Expect valid query observer handle", TagColumnB_RemoveObserverHandle != TypedElementInvalidQueryHandle);					
 					});
-					
+				
 					AfterEach([this]()
 					{
 						CleanupTestRows();
 					});
-
-					Describe("AddRow(TypedElementTableHandle)", [this]()
+				
+					Describe("AddRow(TableHandle)", [this]()
 					{
 						for (int32 RowCount = 1; RowCount <= 2; ++RowCount)
 						{
-							It(FString::Printf(TEXT("should invoke Add observer %d times when called with table containing data column %d times"), RowCount, RowCount), [this, RowCount]()
+							It(FString::Printf(TEXT("should should be able to create %d rows when called %d times"), RowCount, RowCount), [this, RowCount]()
 							{
-								int64 PreviousQueryCallCount = DataColumnA_AddObserverCallCount;
-								
 								for (int32 CreatedRowIndex = 0; CreatedRowIndex < RowCount; ++CreatedRowIndex)
 								{
-									TypedElementRowHandle RowHandle = TedsInterface->AddRow(TestTableHandleA);
+									RowHandle RowHandle = TedsInterface->AddRow(TestTableHandleA);
 									CreatedRows.Add(RowHandle);
+									TestNotEqual("Expecting valid row", RowHandle, InvalidRowHandle);
 								}
-
-								TestEqual(TEXT("Expect observer to be called correct number of times"), DataColumnA_AddObserverCallCount, PreviousQueryCallCount + RowCount);
 							});
 						}
 					});
 
-					Describe("BatchAddRow(TypedElementTableHandle, int32, TypedElementDataStorageCreationCallbackRef)", [this]()
+					Describe("BatchAddRow(TableHandle, int32, TypedElementDataStorageCreationCallbackRef)", [this]()
 					{
 						for (int32 RowCount = 1; RowCount <= 2; ++RowCount)
 						{
-							It(FString::Printf(TEXT("should invoke Add observer %d times when called with table containing data column"), RowCount), [this, RowCount]()
+							It(FString::Printf(TEXT("should should be able to create %d rows when called once"), RowCount), [this, RowCount]()
 							{
-								int64 PreviousQueryCallCount = DataColumnA_AddObserverCallCount;
-								
-								auto RowInitializeCallback = [this](TypedElementRowHandle RowHandle)
+								auto RowInitializeCallback = [this](RowHandle RowHandle)
 								{
 									CreatedRows.Add(RowHandle);
+									TestNotEqual("Expecting valid row", RowHandle, InvalidRowHandle);
 								};
-																
-								TypedElementRowHandle RowHandle = TedsInterface->BatchAddRow(TestTableHandleA, RowCount, RowInitializeCallback);
-								
-								TestEqual(TEXT("Expect observer to be called correct number of times"), DataColumnA_AddObserverCallCount, PreviousQueryCallCount + RowCount);
+											
+								RowHandle RowHandle = TedsInterface->BatchAddRow(TestTableHandleA, RowCount, RowInitializeCallback);
+											
+								TestEqual("Unexpected number of created rows", CreatedRows.Num(), RowCount);
 							});
 						}
 					});
 
+			
 					Describe("", [this]()
 					{
+						/**
+						 * Prerequisites for this batch of tests are:
+						 * - A test table is registered
+						 * - An add and remove observer is registered for FTestInt64Column
+						 */
 						BeforeEach([this]()
 						{
-							CreateTestRows(1);
+							TestTableHandleA = RegisterTestTableA(TedsInterface);
+							TestNotEqual("Expecting valid table handle", TestTableHandleA, InvalidTableHandle);
+
+							RegisterTestObserver<FTestColumnA>(
+								TEXT("Increment CallCount when FTestColumnA added"),
+								TedsInterface,
+								TypedElementQueryBuilder::FObserver::EEvent::Add,
+								&DataColumnA_AddObserverCallCount,
+								&DataColumnA_AddObserverHandle);
+							TestTrue("Expect valid query observer handle", DataColumnA_AddObserverHandle != InvalidQueryHandle);
+						
+							RegisterTestObserver<FTestColumnC>(
+								TEXT("Increment CallCount when FTestColumnC added"),
+								TedsInterface,
+								TypedElementQueryBuilder::FObserver::EEvent::Add,
+								&DataColumnC_AddObserverCallCount,
+								&DataColumnC_AddObserverHandle);
+							TestTrue("Expect valid query observer handle", DataColumnC_AddObserverHandle != InvalidQueryHandle);
+
+							RegisterTestObserver<FTestColumnB>(
+								TEXT("Increment CallCount when FTestColumnB removed"),
+								TedsInterface,
+								TypedElementQueryBuilder::FObserver::EEvent::Remove,
+								&DataColumnB_RemoveObserverCallCount,
+								&DataColumnB_RemoveObserverHandle);
+							TestTrue("Expect valid query observer handle", DataColumnB_RemoveObserverHandle != InvalidQueryHandle);
+
+							RegisterTestObserver<FTestTagColumnA>(
+								TEXT("Increment CallCount when FTestColumnA added"),
+								TedsInterface,
+								TypedElementQueryBuilder::FObserver::EEvent::Add,
+								&TagColumnA_AddObserverCallCount,
+								&TagColumnA_AddObserverHandle);
+							TestTrue("Expect valid query observer handle", TagColumnA_AddObserverHandle != InvalidQueryHandle);
+												
+							RegisterTestObserver<FTestTagColumnC>(
+								TEXT("Increment CallCount when FTestTagColumnC added"),
+								TedsInterface,
+								TypedElementQueryBuilder::FObserver::EEvent::Add,
+								&TagColumnC_AddObserverCallCount, 
+								&TagColumnC_AddObserverHandle);
+							TestTrue("Expect valid query observer handle", TagColumnC_AddObserverHandle != InvalidQueryHandle);
+
+							RegisterTestObserver<FTestTagColumnB>(
+								TEXT("Increment CallCount when FTestTagColumnB removed"),
+								TedsInterface, TypedElementQueryBuilder::FObserver::EEvent::Remove,
+								&TagColumnB_RemoveObserverCallCount,
+								&TagColumnB_RemoveObserverHandle);
+							TestTrue("Expect valid query observer handle", TagColumnB_RemoveObserverHandle != InvalidQueryHandle);					
+						});
+					
+						AfterEach([this]()
+						{
+							CleanupTestRows();
 						});
 
-						Describe("RemoveRow(TypedElementTableHandle)", [this]()
+						Describe("AddRow(TableHandle)", [this]()
 						{
-							It("should remove row when called", [this]()
+							for (int32 RowCount = 1; RowCount <= 2; ++RowCount)
 							{
-								TypedElementRowHandle Row = CreatedRows[0];
-								TestTrue("Row is available", TedsInterface->IsRowAvailable(Row));
-
-								TedsInterface->RemoveRow(Row);
+								It(FString::Printf(TEXT("should invoke Add observer %d times when called with table containing data column %d times"), RowCount, RowCount), [this, RowCount]()
+								{
+									int64 PreviousQueryCallCount = DataColumnA_AddObserverCallCount;
 								
+									for (int32 CreatedRowIndex = 0; CreatedRowIndex < RowCount; ++CreatedRowIndex)
+									{
+										RowHandle RowHandle = TedsInterface->AddRow(TestTableHandleA);
+										CreatedRows.Add(RowHandle);
+									}
+
+									TestEqual(TEXT("Expect observer to be called correct number of times"), DataColumnA_AddObserverCallCount, PreviousQueryCallCount + RowCount);
+								});
+							}
+						});
 
-								TestFalse("Row is not available", TedsInterface->IsRowAvailable(Row));
+						Describe("BatchAddRow(TableHandle, int32, TypedElementDataStorageCreationCallbackRef)", [this]()
+						{
+							for (int32 RowCount = 1; RowCount <= 2; ++RowCount)
+							{
+								It(FString::Printf(TEXT("should invoke Add observer %d times when called with table containing data column"), RowCount), [this, RowCount]()
+								{
+									int64 PreviousQueryCallCount = DataColumnA_AddObserverCallCount;
+								
+									auto RowInitializeCallback = [this](RowHandle RowHandle)
+									{
+										CreatedRows.Add(RowHandle);
+									};
+																
+									RowHandle RowHandle = TedsInterface->BatchAddRow(TestTableHandleA, RowCount, RowInitializeCallback);
+								
+									TestEqual(TEXT("Expect observer to be called correct number of times"), DataColumnA_AddObserverCallCount, PreviousQueryCallCount + RowCount);
+								});
+							}
+						});
 
-								CreatedRows.RemoveAt(0);
+						Describe("", [this]()
+						{
+							BeforeEach([this]()
+							{
+								CreateTestRows(1);
 							});
-							
-							It("should call observer of removed data column when called on row that has the column", [this]()
+
+							Describe("RemoveRow(TableHandle)", [this]()
 							{
-								int64 PreviousQueryCallCount = DataColumnB_RemoveObserverCallCount;
+								It("should remove row when called", [this]()
+								{
+									RowHandle Row = CreatedRows[0];
+									TestTrue("Row is available", TedsInterface->IsRowAvailable(Row));
+
+									TedsInterface->RemoveRow(Row);
 								
-								TypedElementRowHandle Row = CreatedRows[0];
-								TedsInterface->RemoveRow(Row);
 
-								TestEqual("Expected observer called", DataColumnB_RemoveObserverCallCount, PreviousQueryCallCount + 1);
+									TestFalse("Row is not available", TedsInterface->IsRowAvailable(Row));
 
-								CreatedRows.RemoveAt(0);
-							});
+									CreatedRows.RemoveAt(0);
+								});
 							
-							It("should call observer of removed tag column when called on row that has the column", [this]()
-							{
-								int64 PreviousQueryCallCount = TagColumnB_RemoveObserverCallCount;
+								It("should call observer of removed data column when called on row that has the column", [this]()
+								{
+									int64 PreviousQueryCallCount = DataColumnB_RemoveObserverCallCount;
 								
-								TypedElementRowHandle Row = CreatedRows[0];
-								TedsInterface->RemoveRow(Row);
+									RowHandle Row = CreatedRows[0];
+									TedsInterface->RemoveRow(Row);
+
+									TestEqual("Expected observer called", DataColumnB_RemoveObserverCallCount, PreviousQueryCallCount + 1);
+
+									CreatedRows.RemoveAt(0);
+								});
+							
+								It("should call observer of removed tag column when called on row that has the column", [this]()
+								{
+									int64 PreviousQueryCallCount = TagColumnB_RemoveObserverCallCount;
 								
-								TestEqual("Expected observer called", TagColumnB_RemoveObserverCallCount, PreviousQueryCallCount + 1);
+									RowHandle Row = CreatedRows[0];
+									TedsInterface->RemoveRow(Row);
+								
+									TestEqual("Expected observer called", TagColumnB_RemoveObserverCallCount, PreviousQueryCallCount + 1);
 
-								CreatedRows.RemoveAt(0);
+									CreatedRows.RemoveAt(0);
+								});
 							});
-						});
 						
-						Describe("AddColumn(TypedElementRowHandle, const UScriptStruct*)", [this]()
-						{					
-							It("should add data column when called", [this]()
-							{
-								TypedElementRowHandle Row = CreatedRows[0];
-								TestFalse("Expect row to not have column about to be added",
-									TedsInterface->HasColumns(Row, MakeArrayView({FTestColumnC::StaticStruct()})));
+							Describe("AddColumn(RowHandle, const UScriptStruct*)", [this]()
+							{					
+								It("should add data column when called", [this]()
+								{
+									RowHandle Row = CreatedRows[0];
+									TestFalse("Expect row to not have column about to be added",
+										TedsInterface->HasColumns(Row, MakeArrayView({FTestColumnC::StaticStruct()})));
 
-								TedsInterface->AddColumn(Row, FTestColumnC::StaticStruct());
+									TedsInterface->AddColumn(Row, FTestColumnC::StaticStruct());
 
-								TestTrue("Expect row to have column added",
-									TedsInterface->HasColumns(Row, MakeArrayView({FTestColumnC::StaticStruct()})));
+									TestTrue("Expect row to have column added",
+										TedsInterface->HasColumns(Row, MakeArrayView({FTestColumnC::StaticStruct()})));
 
-							});
+								});
 
-							It("should invoke registered observer when called with data column", [this]()
-							{
-								int64 PreviousQueryCallCount = DataColumnC_AddObserverCallCount;
-								TypedElementRowHandle Row = CreatedRows[0];
-								TedsInterface->AddColumn(Row, FTestColumnC::StaticStruct());
+								It("should invoke registered observer when called with data column", [this]()
+								{
+									int64 PreviousQueryCallCount = DataColumnC_AddObserverCallCount;
+									RowHandle Row = CreatedRows[0];
+									TedsInterface->AddColumn(Row, FTestColumnC::StaticStruct());
 
-								TestEqual("Expect observer to have been called one time", DataColumnC_AddObserverCallCount, PreviousQueryCallCount + 1);
-							});
+									TestEqual("Expect observer to have been called one time", DataColumnC_AddObserverCallCount, PreviousQueryCallCount + 1);
+								});
 
-							It("should add tag column when called", [this]()
-							{
-								TypedElementRowHandle Row = CreatedRows[0];
-								TestFalse("Expect row to not have column about to be added",
-									TedsInterface->HasColumns(Row, MakeArrayView({FTestTagColumnC::StaticStruct()})));
+								It("should add tag column when called", [this]()
+								{
+									RowHandle Row = CreatedRows[0];
+									TestFalse("Expect row to not have column about to be added",
+										TedsInterface->HasColumns(Row, MakeArrayView({FTestTagColumnC::StaticStruct()})));
 
-								TedsInterface->AddColumn(Row, FTestTagColumnC::StaticStruct());
+									TedsInterface->AddColumn(Row, FTestTagColumnC::StaticStruct());
 
-								TestTrue("Expect row to have column added",
-									TedsInterface->HasColumns(Row, MakeArrayView({FTestTagColumnC::StaticStruct()})));
+									TestTrue("Expect row to have column added",
+										TedsInterface->HasColumns(Row, MakeArrayView({FTestTagColumnC::StaticStruct()})));
 
-							});
+								});
 
-							It("should invoke registered observer when called with tag column", [this]()
-							{
-								int64 PreviousQueryCallCount = TagColumnC_AddObserverCallCount;
-								TypedElementRowHandle Row = CreatedRows[0];
-								TedsInterface->AddColumn(Row, FTestTagColumnC::StaticStruct());
+								It("should invoke registered observer when called with tag column", [this]()
+								{
+									int64 PreviousQueryCallCount = TagColumnC_AddObserverCallCount;
+									RowHandle Row = CreatedRows[0];
+									TedsInterface->AddColumn(Row, FTestTagColumnC::StaticStruct());
 
-								TestEqual("Expect observer to have been called one time", TagColumnC_AddObserverCallCount, PreviousQueryCallCount + 1);
+									TestEqual("Expect observer to have been called one time", TagColumnC_AddObserverCallCount, PreviousQueryCallCount + 1);
+								});
 							});
-						});
 						
-						Describe("RemoveColumn(TypedElementRowHandle, TConstArrayView<const UScriptStruct*>)", [this]()
-						{
-							It("should remove a single column when called", [this]()
+							Describe("RemoveColumn(RowHandle, TConstArrayView<const UScriptStruct*>)", [this]()
 							{
-								TypedElementRowHandle Row = CreatedRows[0];
-								TestTrue("Expected to have column about to be removed",
-									TedsInterface->HasColumns(Row, MakeArrayView({FTestColumnB::StaticStruct()})));
+								It("should remove a single column when called", [this]()
+								{
+									RowHandle Row = CreatedRows[0];
+									TestTrue("Expected to have column about to be removed",
+										TedsInterface->HasColumns(Row, MakeArrayView({FTestColumnB::StaticStruct()})));
 								
-								TedsInterface->RemoveColumn(Row, FTestColumnB::StaticStruct());
+									TedsInterface->RemoveColumn(Row, FTestColumnB::StaticStruct());
 
-								TestFalse("Expected to no longer have removed column",
-									TedsInterface->HasColumns(Row, MakeArrayView({FTestColumnB::StaticStruct()})));
-							});
+									TestFalse("Expected to no longer have removed column",
+										TedsInterface->HasColumns(Row, MakeArrayView({FTestColumnB::StaticStruct()})));
+								});
 
-							It("should invoke registered data column observer when called", [this]()
-							{
-								int64 PreviousQueryCallCount = DataColumnB_RemoveObserverCallCount;
-								TypedElementRowHandle Row = CreatedRows[0];
+								It("should invoke registered data column observer when called", [this]()
+								{
+									int64 PreviousQueryCallCount = DataColumnB_RemoveObserverCallCount;
+									RowHandle Row = CreatedRows[0];
 								
-								TedsInterface->RemoveColumn(Row, FTestColumnB::StaticStruct());
+									TedsInterface->RemoveColumn(Row, FTestColumnB::StaticStruct());
 								
-								TestEqual("Expect observer to have been called one time", DataColumnB_RemoveObserverCallCount, PreviousQueryCallCount + 1);
-							});
+									TestEqual("Expect observer to have been called one time", DataColumnB_RemoveObserverCallCount, PreviousQueryCallCount + 1);
+								});
 
-							It("should remove a single tag column when called", [this]()
-							{
-								TypedElementRowHandle Row = CreatedRows[0];
-								TestTrue("Expected to have column about to be removed",
-									TedsInterface->HasColumns(Row, MakeArrayView({FTestTagColumnB::StaticStruct()})));
+								It("should remove a single tag column when called", [this]()
+								{
+									RowHandle Row = CreatedRows[0];
+									TestTrue("Expected to have column about to be removed",
+										TedsInterface->HasColumns(Row, MakeArrayView({FTestTagColumnB::StaticStruct()})));
 														
-								TedsInterface->RemoveColumn(Row, FTestTagColumnB::StaticStruct());
+									TedsInterface->RemoveColumn(Row, FTestTagColumnB::StaticStruct());
 
-								TestFalse("Expected to no longer have removed column",
-									TedsInterface->HasColumns(Row, MakeArrayView({FTestTagColumnB::StaticStruct()})));
-							});
+									TestFalse("Expected to no longer have removed column",
+										TedsInterface->HasColumns(Row, MakeArrayView({FTestTagColumnB::StaticStruct()})));
+								});
 
-							It("should invoke registered tag column observer when called", [this]()
-							{
-								int64 PreviousQueryCallCount = TagColumnB_RemoveObserverCallCount;
-								TypedElementRowHandle Row = CreatedRows[0];
+								It("should invoke registered tag column observer when called", [this]()
+								{
+									int64 PreviousQueryCallCount = TagColumnB_RemoveObserverCallCount;
+									RowHandle Row = CreatedRows[0];
 								
-								TedsInterface->RemoveColumn(Row, FTestTagColumnB::StaticStruct());
+									TedsInterface->RemoveColumn(Row, FTestTagColumnB::StaticStruct());
 								
-								TestEqual("Expect observer to have been called one time", TagColumnB_RemoveObserverCallCount, PreviousQueryCallCount + 1);
+									TestEqual("Expect observer to have been called one time", TagColumnB_RemoveObserverCallCount, PreviousQueryCallCount + 1);
+								});
 							});
-						});
 
-						Describe("AddColumns(TypedElementRowHandle, TConstArrayView<const UScriptStruct*>)", [this]()
-						{
-							It("should add a single column when called with single data column", [this]()
+							Describe("AddColumns(RowHandle, TConstArrayView<const UScriptStruct*>)", [this]()
 							{
-								TArray<const UScriptStruct*> Columns = {FTestColumnC::StaticStruct(), FTestColumnD::StaticStruct()};
+								It("should add a single column when called with single data column", [this]()
+								{
+									TArray<const UScriptStruct*> Columns = {FTestColumnC::StaticStruct(), FTestColumnD::StaticStruct()};
 								
-								TypedElementRowHandle Row = CreatedRows[0];
-								TestFalse("Expected to not have column about to be added",
-									TedsInterface->HasColumns(Row, Columns));
+									RowHandle Row = CreatedRows[0];
+									TestFalse("Expected to not have column about to be added",
+										TedsInterface->HasColumns(Row, Columns));
 
-								TedsInterface->AddColumns(Row, Columns);
+									TedsInterface->AddColumns(Row, Columns);
 
-								TestTrue("Expected to have column about to be added", 
-								TedsInterface->HasColumns(Row, Columns));
-							});
+									TestTrue("Expected to have column about to be added", 
+									TedsInterface->HasColumns(Row, Columns));
+								});
 
-							It("should add multiple column when called with multiple data columns", [this]()
-							{
-								TArray<const UScriptStruct*> Columns = {FTestColumnC::StaticStruct(), FTestColumnD::StaticStruct()};
+								It("should add multiple column when called with multiple data columns", [this]()
+								{
+									TArray<const UScriptStruct*> Columns = {FTestColumnC::StaticStruct(), FTestColumnD::StaticStruct()};
 								
-								TypedElementRowHandle Row = CreatedRows[0];
-								TestFalse("Expected to not have column about to be added", 
-								TedsInterface->HasColumns(Row, Columns));
+									RowHandle Row = CreatedRows[0];
+									TestFalse("Expected to not have column about to be added", 
+									TedsInterface->HasColumns(Row, Columns));
 
-								TedsInterface->AddColumns(Row, Columns);
+									TedsInterface->AddColumns(Row, Columns);
 
-								TestTrue("Expected to have column about to be added",
-									TedsInterface->HasColumns(Row, Columns));
-							});
+									TestTrue("Expected to have column about to be added",
+										TedsInterface->HasColumns(Row, Columns));
+								});
 
-							It("should invoke registered data column observer when called", [this]()
-							{
-								TArray<const UScriptStruct*> Columns = {FTestColumnC::StaticStruct()};
+								It("should invoke registered data column observer when called", [this]()
+								{
+									TArray<const UScriptStruct*> Columns = {FTestColumnC::StaticStruct()};
 								
-								int64 PreviousQueryCallCount = DataColumnC_AddObserverCallCount;
-								TypedElementRowHandle Row = CreatedRows[0];
+									int64 PreviousQueryCallCount = DataColumnC_AddObserverCallCount;
+									RowHandle Row = CreatedRows[0];
 								
-								TedsInterface->AddColumns(Row, Columns);
+									TedsInterface->AddColumns(Row, Columns);
 
-								TestEqual("Expect observer to have been called one time", DataColumnC_AddObserverCallCount, PreviousQueryCallCount + 1);
-							});
+									TestEqual("Expect observer to have been called one time", DataColumnC_AddObserverCallCount, PreviousQueryCallCount + 1);
+								});
 
-							It("should add a single column when called with single tag column", [this]()
-							{
-								TArray<const UScriptStruct*> Columns = {FTestTagColumnC::StaticStruct(), FTestTagColumnD::StaticStruct()};
+								It("should add a single column when called with single tag column", [this]()
+								{
+									TArray<const UScriptStruct*> Columns = {FTestTagColumnC::StaticStruct(), FTestTagColumnD::StaticStruct()};
 														
-								TypedElementRowHandle Row = CreatedRows[0];
-								TestFalse("Expected to not have column about to be added",
-									TedsInterface->HasColumns(Row, Columns));
+									RowHandle Row = CreatedRows[0];
+									TestFalse("Expected to not have column about to be added",
+										TedsInterface->HasColumns(Row, Columns));
 
-								TedsInterface->AddColumns(Row, Columns);
+									TedsInterface->AddColumns(Row, Columns);
 
-								TestTrue("Expected to have column about to be added", 
-								TedsInterface->HasColumns(Row, Columns));
-							});
+									TestTrue("Expected to have column about to be added", 
+									TedsInterface->HasColumns(Row, Columns));
+								});
 
-							It("should add multiple column when called with multiple tag columns", [this]()
-							{
-								TArray<const UScriptStruct*> Columns = {FTestTagColumnC::StaticStruct(), FTestTagColumnD::StaticStruct()};
+								It("should add multiple column when called with multiple tag columns", [this]()
+								{
+									TArray<const UScriptStruct*> Columns = {FTestTagColumnC::StaticStruct(), FTestTagColumnD::StaticStruct()};
 														
-								TypedElementRowHandle Row = CreatedRows[0];
-								TestFalse("Expected to not have column about to be added", 
-								TedsInterface->HasColumns(Row, Columns));
+									RowHandle Row = CreatedRows[0];
+									TestFalse("Expected to not have column about to be added", 
+									TedsInterface->HasColumns(Row, Columns));
 
-								TedsInterface->AddColumns(Row, Columns);
+									TedsInterface->AddColumns(Row, Columns);
 
-								TestTrue("Expected to have column about to be added",
-									TedsInterface->HasColumns(Row, Columns));
-							});
+									TestTrue("Expected to have column about to be added",
+										TedsInterface->HasColumns(Row, Columns));
+								});
 
-							It("should invoke registered tag column added observer when called", [this]()
-							{
-								TArray<const UScriptStruct*> Columns = {FTestTagColumnC::StaticStruct()};
+								It("should invoke registered tag column added observer when called", [this]()
+								{
+									TArray<const UScriptStruct*> Columns = {FTestTagColumnC::StaticStruct()};
 														
-								int64 PreviousQueryCallCount = TagColumnC_AddObserverCallCount;
-								TypedElementRowHandle Row = CreatedRows[0];
-								TedsInterface->AddColumns(Row, Columns);
+									int64 PreviousQueryCallCount = TagColumnC_AddObserverCallCount;
+									RowHandle Row = CreatedRows[0];
+									TedsInterface->AddColumns(Row, Columns);
 
-								TestEqual("Expect observer to have been called one time", TagColumnC_AddObserverCallCount, PreviousQueryCallCount + 1);
+									TestEqual("Expect observer to have been called one time", TagColumnC_AddObserverCallCount, PreviousQueryCallCount + 1);
+								});
 							});
-						});
 
-						Describe("RemoveColumns(TypedElementRowHandle, TConstArrayView<const UScriptStruct*>)", [this]()
-						{
-							It("should remove a single column when RemoveColumns called with single column", [this]()
+							Describe("RemoveColumns(RowHandle, TConstArrayView<const UScriptStruct*>)", [this]()
 							{
-								TArray<const UScriptStruct*> Columns = {FTestColumnB::StaticStruct()};
+								It("should remove a single column when RemoveColumns called with single column", [this]()
+								{
+									TArray<const UScriptStruct*> Columns = {FTestColumnB::StaticStruct()};
 								
-								TypedElementRowHandle Row = CreatedRows[0];
-								TestTrue("Expected to not column about to be removed",
-									TedsInterface->HasColumns(Row, Columns));
+									RowHandle Row = CreatedRows[0];
+									TestTrue("Expected to not column about to be removed",
+										TedsInterface->HasColumns(Row, Columns));
 
-								TedsInterface->RemoveColumns(Row, Columns);
+									TedsInterface->RemoveColumns(Row, Columns);
 
-								TestFalse("Expected to not have column that was removed", 
-								TedsInterface->HasColumns(Row, MakeArrayView({FTestColumnB::StaticStruct()})));
-							});
+									TestFalse("Expected to not have column that was removed", 
+									TedsInterface->HasColumns(Row, MakeArrayView({FTestColumnB::StaticStruct()})));
+								});
 
-							It("should remove all columns when RemoveColumns called with multiple columns", [this]()
-							{
-								TArray<const UScriptStruct*> Columns = {FTestColumnA::StaticStruct(), FTestColumnB::StaticStruct()};
+								It("should remove all columns when RemoveColumns called with multiple columns", [this]()
+								{
+									TArray<const UScriptStruct*> Columns = {FTestColumnA::StaticStruct(), FTestColumnB::StaticStruct()};
 								
-								TypedElementRowHandle Row = CreatedRows[0];
-								TestTrue("Expected to have columns about to be removed", 
-								TedsInterface->HasColumns(Row, Columns));
+									RowHandle Row = CreatedRows[0];
+									TestTrue("Expected to have columns about to be removed", 
+									TedsInterface->HasColumns(Row, Columns));
 
-								TedsInterface->RemoveColumns(Row, Columns);
+									TedsInterface->RemoveColumns(Row, Columns);
 
-								TestFalse("Expected to have column about to be added",
-									TedsInterface->HasColumns(Row, Columns));
-							});
+									TestFalse("Expected to have column about to be added",
+										TedsInterface->HasColumns(Row, Columns));
+								});
 
-							It("should invoke registered observer when RemoveColumns called", [this]()
-							{
-								TArray<const UScriptStruct*> Columns = {FTestColumnB::StaticStruct()};
+								It("should invoke registered observer when RemoveColumns called", [this]()
+								{
+									TArray<const UScriptStruct*> Columns = {FTestColumnB::StaticStruct()};
 								
-								int64 PreviousQueryCallCount = DataColumnB_RemoveObserverCallCount;
-								TypedElementRowHandle Row = CreatedRows[0];
-								TedsInterface->RemoveColumns(Row, Columns);
+									int64 PreviousQueryCallCount = DataColumnB_RemoveObserverCallCount;
+									RowHandle Row = CreatedRows[0];
+									TedsInterface->RemoveColumns(Row, Columns);
 
-								TestEqual("Expect observer to have been called one time", DataColumnB_RemoveObserverCallCount, PreviousQueryCallCount + 1);
+									TestEqual("Expect observer to have been called one time", DataColumnB_RemoveObserverCallCount, PreviousQueryCallCount + 1);
+								});
 							});
 						});
 					});
-				});
 				
+				});
 			});
-		});
-	}
-}
+		}
+	} // namespace Tests
+} // namespace UE::Editor::DataStorage
 
 #endif
\ No newline at end of file
diff --git a/Engine/Plugins/Experimental/EditorDataStorage/Source/TedsUI/Private/Processors/WidgetReferenceColumnUpdateProcessor.cpp b/Engine/Plugins/Experimental/EditorDataStorage/Source/TedsUI/Private/Processors/WidgetReferenceColumnUpdateProcessor.cpp
--- a/Engine/Plugins/Experimental/EditorDataStorage/Source/TedsUI/Private/Processors/WidgetReferenceColumnUpdateProcessor.cpp
+++ b/Engine/Plugins/Experimental/EditorDataStorage/Source/TedsUI/Private/Processors/WidgetReferenceColumnUpdateProcessor.cpp
@@ -15,12 +15,13 @@ void UWidgetReferenceColumnUpdateFactory::RegisterDeleteRowOnWidgetDeleteQuery(I
 {
 	using namespace TypedElementQueryBuilder;
 	using DSI = ITypedElementDataStorageInterface;
+	using namespace UE::Editor;
 
 	DataStorage.RegisterQuery(
     	Select(
     		TEXT("Delete row with deleted widget"),
     		FPhaseAmble(FPhaseAmble::ELocation::Preamble, DSI::EQueryTickPhase::FrameEnd),
-    		[](DSI::IQueryContext& Context, TypedElementRowHandle Row, const FTypedElementSlateWidgetReferenceColumn& WidgetReference)
+    		[](DSI::IQueryContext& Context, DataStorage::RowHandle Row, const FTypedElementSlateWidgetReferenceColumn& WidgetReference)
     		{
     			if (!WidgetReference.TedsWidget.IsValid())
     			{
@@ -38,12 +39,13 @@ void UWidgetReferenceColumnUpdateFactory::RegisterDeleteColumnOnWidgetDeleteQuer
 {
 	using namespace TypedElementQueryBuilder;
 	using DSI = ITypedElementDataStorageInterface;
+	using namespace UE::Editor;
 
 	DataStorage.RegisterQuery(
 		Select(
 			TEXT("Delete widget columns for deleted widget"),
 			FPhaseAmble(FPhaseAmble::ELocation::Preamble, DSI::EQueryTickPhase::FrameEnd),
-			[](DSI::IQueryContext& Context, TypedElementRowHandle Row, const FTypedElementSlateWidgetReferenceColumn& WidgetReference)
+			[](DSI::IQueryContext& Context, DataStorage::RowHandle Row, const FTypedElementSlateWidgetReferenceColumn& WidgetReference)
 			{
 				if (!WidgetReference.TedsWidget.IsValid())
 				{
diff --git a/Engine/Plugins/Experimental/EditorDataStorage/Source/TedsUI/Private/Widgets/CounterWidget.cpp b/Engine/Plugins/Experimental/EditorDataStorage/Source/TedsUI/Private/Widgets/CounterWidget.cpp
--- a/Engine/Plugins/Experimental/EditorDataStorage/Source/TedsUI/Private/Widgets/CounterWidget.cpp
+++ b/Engine/Plugins/Experimental/EditorDataStorage/Source/TedsUI/Private/Widgets/CounterWidget.cpp
@@ -125,7 +125,7 @@ void UCounterWidgetFactory::SetupMainWindowIntegrations(TSharedPtr<SWindow> Pare
 
 		TArray<TSharedRef<SWidget>> Widgets;
 		UiInterface->ConstructWidgets(WigetPurpose, {},
-			[&Widgets](const TSharedRef<SWidget>& NewWidget, TypedElementRowHandle Row)
+			[&Widgets](const TSharedRef<SWidget>& NewWidget, UE::Editor::DataStorage::RowHandle Row)
 			{
 				Widgets.Add(NewWidget);
 			});
@@ -174,7 +174,7 @@ TSharedPtr<SWidget> FCounterWidgetConstructor::CreateWidget(const TypedElementDa
 		.Justification(ETextJustify::Center);
 }
 
-bool FCounterWidgetConstructor::SetColumns(ITypedElementDataStorageInterface* DataStorage, TypedElementRowHandle Row)
+bool FCounterWidgetConstructor::SetColumns(ITypedElementDataStorageInterface* DataStorage, UE::Editor::DataStorage::RowHandle Row)
 {
 	FCounterWidgetColumn* CounterColumn = DataStorage->GetColumn<FCounterWidgetColumn>(Row);
 	checkf(CounterColumn, TEXT("Added a new FCounterWidgetColumn to the Typed Elements Data Storage, but didn't get a valid pointer back."));
diff --git a/Engine/Plugins/Experimental/EditorDataStorage/Source/TedsUI/Private/Widgets/LabelWidget.cpp b/Engine/Plugins/Experimental/EditorDataStorage/Source/TedsUI/Private/Widgets/LabelWidget.cpp
--- a/Engine/Plugins/Experimental/EditorDataStorage/Source/TedsUI/Private/Widgets/LabelWidget.cpp
+++ b/Engine/Plugins/Experimental/EditorDataStorage/Source/TedsUI/Private/Widgets/LabelWidget.cpp
@@ -23,62 +23,69 @@
 // ULabelWidgetFactory
 //
 
-static void UpdateTextWidget(const TWeakPtr<SWidget>& Widget, const FTypedElementLabelColumn& Label, const uint64* HashValue)
+namespace UE::Editor::DataStorage
 {
-	TSharedPtr<SWidget> WidgetPointer = Widget.Pin();
-	if (ensureMsgf(WidgetPointer, TEXT("Referenced widget is not valid. A constructed widget may not have been cleaned up. This can "
-		"also happen if this processor is running in the same phase as the processors responsible for cleaning up old "
-		"references.")))
+	namespace UI::Private
 	{
-		if (TSharedPtr<ITypedElementUiTextCapability> Text = WidgetPointer->GetMetaData<ITypedElementUiTextCapability>())
+		static void UpdateTextWidget(const TWeakPtr<SWidget>& Widget, const FTypedElementLabelColumn& Label, const uint64* HashValue)
 		{
-			Text->SetText(FText::FromString(Label.Label));
-		}
-
-		if (TSharedPtr<ITypedElementUiTooltipCapability> ToolTip = WidgetPointer->GetMetaData<ITypedElementUiTooltipCapability>())
-		{
-			if (!HashValue)
-			{
-				ToolTip->SetToolTipText(FText::FromString(Label.Label));
-			}
-			else
+			TSharedPtr<SWidget> WidgetPointer = Widget.Pin();
+			if (ensureMsgf(WidgetPointer, TEXT("Referenced widget is not valid. A constructed widget may not have been cleaned up. This can "
+				"also happen if this processor is running in the same phase as the processors responsible for cleaning up old "
+				"references.")))
 			{
-				ToolTip->SetToolTipText(FText::FromString(FString::Format(TEXT("{0}\nHash: {1}"), { Label.Label, *HashValue })));
+				if (TSharedPtr<ITypedElementUiTextCapability> Text = WidgetPointer->GetMetaData<ITypedElementUiTextCapability>())
+				{
+					Text->SetText(FText::FromString(Label.Label));
+				}
+
+				if (TSharedPtr<ITypedElementUiTooltipCapability> ToolTip = WidgetPointer->GetMetaData<ITypedElementUiTooltipCapability>())
+				{
+					if (!HashValue)
+					{
+						ToolTip->SetToolTipText(FText::FromString(Label.Label));
+					}
+					else
+					{
+						ToolTip->SetToolTipText(FText::FromString(FString::Format(TEXT("{0}\nHash: {1}"), { Label.Label, *HashValue })));
+					}
+				}
 			}
 		}
-	}
-}
 
-static void SyncColumnsToWidget(
-	ITypedElementDataStorageInterface* DataStorage, 
-	TypedElementRowHandle TargetRow, 
-	FTypedElementU64IntValueCacheColumn& TextHash,
-	const TWeakPtr<SWidget>& Widget,
-	bool bShowHashInToolTip)
-{
-	if (const FTypedElementLabelColumn* LabelColumn = DataStorage->GetColumn<FTypedElementLabelColumn>(TargetRow))
-	{
-		if (const FTypedElementLabelHashColumn* LabelHashColumn = DataStorage->GetColumn<FTypedElementLabelHashColumn>(TargetRow))
+		static void SyncColumnsToWidget(
+			ITypedElementDataStorageInterface* DataStorage,
+			RowHandle TargetRow,
+			FTypedElementU64IntValueCacheColumn& TextHash,
+			const TWeakPtr<SWidget>& Widget,
+			bool bShowHashInToolTip)
 		{
-			if (LabelHashColumn->LabelHash != TextHash.Value)
+			if (const FTypedElementLabelColumn* LabelColumn = DataStorage->GetColumn<FTypedElementLabelColumn>(TargetRow))
 			{
-				UpdateTextWidget(Widget, *LabelColumn, bShowHashInToolTip ? &LabelHashColumn->LabelHash : nullptr);
-				TextHash.Value = LabelHashColumn->LabelHash;
+				if (const FTypedElementLabelHashColumn* LabelHashColumn = DataStorage->GetColumn<FTypedElementLabelHashColumn>(TargetRow))
+				{
+					if (LabelHashColumn->LabelHash != TextHash.Value)
+					{
+						UpdateTextWidget(Widget, *LabelColumn, bShowHashInToolTip ? &LabelHashColumn->LabelHash : nullptr);
+						TextHash.Value = LabelHashColumn->LabelHash;
+					}
+				}
+				else
+				{
+					UpdateTextWidget(Widget, *LabelColumn, nullptr);
+				}
 			}
 		}
-		else
-		{
-			UpdateTextWidget(Widget, *LabelColumn, nullptr);
-		}
-	}
-}
+	} // namespace UI::Private
+} // namespace UE::Editor::DataStorage
 
 void ULabelWidgetFactory::RegisterQueries(ITypedElementDataStorageInterface& DataStorage)
 {
 	using namespace TypedElementQueryBuilder;
 	using namespace TypedElementDataStorage;
+	using namespace UE::Editor;
 	
-	TypedElementQueryHandle UpdateLabelWidget = DataStorage.RegisterQuery(
+	DataStorage::QueryHandle UpdateLabelWidget = DataStorage.RegisterQuery(
 		Select()
 			.ReadOnly<FTypedElementLabelColumn>()
 		.Where()
@@ -86,7 +93,7 @@ void ULabelWidgetFactory::RegisterQueries(ITypedElementDataStorageInterface& Dat
 			.None<FTypedElementLabelHashColumn>()
 		.Compile());
 
-	TypedElementQueryHandle UpdateLabelAndHashWidget = DataStorage.RegisterQuery(
+	DataStorage::QueryHandle UpdateLabelAndHashWidget = DataStorage.RegisterQuery(
 		Select()
 			.ReadOnly<FTypedElementLabelColumn, FTypedElementLabelHashColumn>()
 		.Where()
@@ -108,14 +115,14 @@ void ULabelWidgetFactory::RegisterQueries(ITypedElementDataStorageInterface& Dat
 				Context.RunSubquery(0, Target.Row, CreateSubqueryCallbackBinding(
 					[&Widget](const FTypedElementLabelColumn& Label)
 					{
-						UpdateTextWidget(Widget.Widget, Label, nullptr);
+						DataStorage::UI::Private::UpdateTextWidget(Widget.Widget, Label, nullptr);
 					}));
 				Context.RunSubquery(1, Target.Row, CreateSubqueryCallbackBinding(
 					[&Widget, &TextHash, &Config](const FTypedElementLabelColumn& Label, const FTypedElementLabelHashColumn& Hash)
 					{
 						if (Hash.LabelHash != TextHash.Value)
 						{
-							UpdateTextWidget(Widget.Widget, Label, Config.bShowHashInTooltip ? &Hash.LabelHash : nullptr);
+							DataStorage::UI::Private::UpdateTextWidget(Widget.Widget, Label, Config.bShowHashInTooltip ? &Hash.LabelHash : nullptr);
 							TextHash.Value = Hash.LabelHash;
 						}
 					}));
@@ -129,10 +136,8 @@ void ULabelWidgetFactory::RegisterQueries(ITypedElementDataStorageInterface& Dat
 void ULabelWidgetFactory::RegisterWidgetConstructors(ITypedElementDataStorageInterface& DataStorage,
 	ITypedElementDataStorageUiInterface& DataStorageUi) const
 {
-	using namespace TypedElementDataStorage;
-
 	DataStorageUi.RegisterWidgetFactory<FLabelWidgetConstructor>(FName(TEXT("General.Cell")), 
-		FColumn<FTypedElementLabelColumn>() || (FColumn<FTypedElementLabelColumn>() && FColumn<FTypedElementLabelHashColumn>()));
+		TypedElementDataStorage::FColumn<FTypedElementLabelColumn>() || (TypedElementDataStorage::FColumn<FTypedElementLabelColumn>() && TypedElementDataStorage::FColumn<FTypedElementLabelHashColumn>()));
 }
 
 
@@ -162,7 +167,7 @@ TConstArrayView<const UScriptStruct*> FLabelWidgetConstructor::GetAdditionalColu
 }
 
 TSharedPtr<SWidget> FLabelWidgetConstructor::Construct(
-	TypedElementRowHandle Row,
+	UE::Editor::DataStorage::RowHandle Row,
 	ITypedElementDataStorageInterface* DataStorage,
 	ITypedElementDataStorageUiInterface* DataStorageUi,
 	const TypedElementDataStorage::FMetaDataView& Arguments)
@@ -228,7 +233,7 @@ TSharedPtr<SWidget> FLabelWidgetConstructor::Construct(
 	return nullptr;
 }
 
-bool FLabelWidgetConstructor::SetColumns(ITypedElementDataStorageInterface* DataStorage, TypedElementRowHandle Row)
+bool FLabelWidgetConstructor::SetColumns(ITypedElementDataStorageInterface* DataStorage, UE::Editor::DataStorage::RowHandle Row)
 {
 	DataStorage->AddColumn(Row, FLabelWidgetColumn{ .bShowHashInTooltip = (MatchedColumnTypes.Num() == 2) });
 	return true;
@@ -237,10 +242,10 @@ bool FLabelWidgetConstructor::SetColumns(ITypedElementDataStorageInterface* Data
 bool FLabelWidgetConstructor::FinalizeWidget(
 	ITypedElementDataStorageInterface* DataStorage,
 	ITypedElementDataStorageUiInterface* DataStorageUi,
-	TypedElementRowHandle Row,
+	UE::Editor::DataStorage::RowHandle Row,
 	const TSharedPtr<SWidget>& Widget)
 {
-	SyncColumnsToWidget(
+	UE::Editor::DataStorage::UI::Private::SyncColumnsToWidget(
 		DataStorage, 
 		DataStorage->GetColumn<FTypedElementRowReferenceColumn>(Row)->Row,
 		*DataStorage->GetColumn<FTypedElementU64IntValueCacheColumn>(Row),
@@ -248,7 +253,7 @@ bool FLabelWidgetConstructor::FinalizeWidget(
 	return true;
 }
 
-bool FLabelWidgetConstructor::IsWidgetSelected(ITypedElementDataStorageInterface* DataStorage, TypedElementRowHandle UiRow)
+bool FLabelWidgetConstructor::IsWidgetSelected(ITypedElementDataStorageInterface* DataStorage, UE::Editor::DataStorage::RowHandle UiRow)
 {
 	if (const FExternalWidgetSelectionColumn* ExternalWidgetSelectionColumn = DataStorage->GetColumn<FExternalWidgetSelectionColumn>(UiRow))
 	{
diff --git a/Engine/Plugins/Experimental/EditorDataStorage/Source/TedsUI/Private/Widgets/PackagePathWidget.cpp b/Engine/Plugins/Experimental/EditorDataStorage/Source/TedsUI/Private/Widgets/PackagePathWidget.cpp
--- a/Engine/Plugins/Experimental/EditorDataStorage/Source/TedsUI/Private/Widgets/PackagePathWidget.cpp
+++ b/Engine/Plugins/Experimental/EditorDataStorage/Source/TedsUI/Private/Widgets/PackagePathWidget.cpp
@@ -14,12 +14,10 @@
 void UPackagePathWidgetFactory::RegisterWidgetConstructors(ITypedElementDataStorageInterface& DataStorage,
 	ITypedElementDataStorageUiInterface& DataStorageUi) const
 {
-	using namespace TypedElementDataStorage;
-
 	DataStorageUi.RegisterWidgetFactory<FPackagePathWidgetConstructor>(FName(TEXT("General.Cell")),
-		FColumn<FTypedElementPackagePathColumn>());
+		TypedElementDataStorage::FColumn<FTypedElementPackagePathColumn>());
 	DataStorageUi.RegisterWidgetFactory<FLoadedPackagePathWidgetConstructor>(FName(TEXT("General.Cell")),
-		FColumn<FTypedElementPackageLoadedPathColumn>());
+		TypedElementDataStorage::FColumn<FTypedElementPackageLoadedPathColumn>());
 }
 
 
@@ -48,10 +46,10 @@ TSharedPtr<SWidget> FPackagePathWidgetConstructor::CreateWidget(const TypedEleme
 bool FPackagePathWidgetConstructor::FinalizeWidget(
 	ITypedElementDataStorageInterface* DataStorage,
 	ITypedElementDataStorageUiInterface* DataStorageUi,
-	TypedElementRowHandle Row,
+	UE::Editor::DataStorage::RowHandle Row,
 	const TSharedPtr<SWidget>& Widget)
 {
-	TypedElementRowHandle TargetRow = DataStorage->GetColumn<FTypedElementRowReferenceColumn>(Row)->Row;
+	UE::Editor::DataStorage::RowHandle TargetRow = DataStorage->GetColumn<FTypedElementRowReferenceColumn>(Row)->Row;
 	if (const FTypedElementPackagePathColumn* Path = DataStorage->GetColumn<FTypedElementPackagePathColumn>(TargetRow))
 	{
 		STextBlock* TextWidget = static_cast<STextBlock*>(Widget.Get());
@@ -80,10 +78,10 @@ FLoadedPackagePathWidgetConstructor::FLoadedPackagePathWidgetConstructor()
 bool FLoadedPackagePathWidgetConstructor::FinalizeWidget(
 	ITypedElementDataStorageInterface* DataStorage,
 	ITypedElementDataStorageUiInterface* DataStorageUi,
-	TypedElementRowHandle Row,
+	UE::Editor::DataStorage::RowHandle Row,
 	const TSharedPtr<SWidget>& Widget)
 {
-	TypedElementRowHandle TargetRow = DataStorage->GetColumn<FTypedElementRowReferenceColumn>(Row)->Row;
+	UE::Editor::DataStorage::RowHandle TargetRow = DataStorage->GetColumn<FTypedElementRowReferenceColumn>(Row)->Row;
 	if (const FTypedElementPackageLoadedPathColumn* Path = DataStorage->GetColumn<FTypedElementPackageLoadedPathColumn>(TargetRow))
 	{
 		STextBlock* TextWidget = static_cast<STextBlock*>(Widget.Get());
diff --git a/Engine/Plugins/Experimental/EditorDataStorage/Source/TedsUI/Public/Widgets/CounterWidget.h b/Engine/Plugins/Experimental/EditorDataStorage/Source/TedsUI/Public/Widgets/CounterWidget.h
--- a/Engine/Plugins/Experimental/EditorDataStorage/Source/TedsUI/Public/Widgets/CounterWidget.h
+++ b/Engine/Plugins/Experimental/EditorDataStorage/Source/TedsUI/Public/Widgets/CounterWidget.h
@@ -57,11 +57,11 @@ public:
 
 	FText ToolTipText{ NSLOCTEXT("TypedElementUI_CounterWidget", "Tooltip", "Shows the total number found in the editor.") };
 	FText LabelText{ NSLOCTEXT("TypedElementUI_CounterWidget", "Label", "Counted") };
-	TypedElementQueryHandle Query;
+	UE::Editor::DataStorage::QueryHandle Query;
 
 protected:
 	TSharedPtr<SWidget> CreateWidget(const TypedElementDataStorage::FMetaDataView& Arguments) override;
-	bool SetColumns(ITypedElementDataStorageInterface* DataStorage, TypedElementRowHandle Row) override;
+	bool SetColumns(ITypedElementDataStorageInterface* DataStorage, UE::Editor::DataStorage::RowHandle Row) override;
 };
 
 USTRUCT(meta = (DisplayName = "Counter widget"))
@@ -70,5 +70,5 @@ struct TEDSUI_API FCounterWidgetColumn : public FEditorDataStorageColumn
 	GENERATED_BODY()
 
 	FTextFormat LabelTextFormatter;
-	TypedElementQueryHandle Query;
+	UE::Editor::DataStorage::QueryHandle Query;
 };
\ No newline at end of file
