Icon override support in TEDS.

A new column was added to allow setting a custom icon for a row. Actors have been setup to make use of by extracting their custom icon into the column. This change does not include usage of the column on the Slate side, which will come in future CLs.

diff --git a/Engine/Plugins/Experimental/EditorDataStorage/Source/TedsCore/Private/Compatibility/TypedElementActorIconOverrideQueries.cpp b/Engine/Plugins/Experimental/EditorDataStorage/Source/TedsCore/Private/Compatibility/TypedElementActorIconOverrideQueries.cpp
new file mode 100644
--- /dev/null
+++ b/Engine/Plugins/Experimental/EditorDataStorage/Source/TedsCore/Private/Compatibility/TypedElementActorIconOverrideQueries.cpp
@@ -0,0 +1,65 @@
+// Copyright Epic Games, Inc. All Rights Reserved.
+
+#include "Compatibility/TypedElementActorIconOverrideQueries.h"
+
+#include "Elements/Columns/TypedElementCompatibilityColumns.h"
+#include "Elements/Columns/TypedElementIconOverrideColumns.h"
+#include "Elements/Columns/TypedElementMiscColumns.h"
+#include "Elements/Framework/TypedElementQueryBuilder.h"
+#include "GameFramework/Actor.h"
+
+void UActorIconOverrideDataStorageFactory::RegisterQueries(ITypedElementDataStorageInterface& DataStorage)
+{
+	using namespace TypedElementDataStorage;
+	using namespace TypedElementQueryBuilder;
+
+	DataStorage.RegisterQuery(
+		Select(
+			TEXT("Add icon override column to actor"),
+			FProcessor(EQueryTickPhase::PrePhysics, DataStorage.GetQueryTickGroupName(EQueryTickGroups::SyncExternalToDataStorage))
+				.SetExecutionMode(EExecutionMode::GameThread),
+			[](IQueryContext& Context, RowHandle Row, const FTypedElementUObjectColumn& ActorColumn)
+			{
+				if (const AActor* Actor = Cast<AActor>(ActorColumn.Object))
+				{
+					if (FName IconName = Actor->GetCustomIconName(); !IconName.IsNone())
+					{
+						Context.AddColumn(Row, FTypedElementIconOverrideColumn{ .IconName = IconName });
+					}
+				}
+			}
+		)
+		.Where()
+			.All<FTypedElementSyncFromWorldTag, FTypedElementActorTag>()
+			.None<FTypedElementIconOverrideColumn>()
+		.Compile()
+	);
+
+	DataStorage.RegisterQuery(
+		Select(
+			TEXT("Update/remove icon override column to actor"),
+			FProcessor(EQueryTickPhase::PrePhysics, DataStorage.GetQueryTickGroupName(EQueryTickGroups::SyncExternalToDataStorage))
+				.SetExecutionMode(EExecutionMode::GameThread),
+			[]( IQueryContext& Context, 
+				RowHandle Row, 
+				const FTypedElementUObjectColumn& ActorColumn, 
+				FTypedElementIconOverrideColumn& IconColumn)
+			{
+				if (const AActor* Actor = Cast<AActor>(ActorColumn.Object))
+				{
+					if (FName IconName = Actor->GetCustomIconName(); !IconName.IsNone())
+					{
+						IconColumn.IconName = IconName;
+					}
+					else
+					{
+						Context.RemoveColumns<FTypedElementIconOverrideColumn>(Row);
+					}
+				}
+			}
+		)
+		.Where()
+			.All<FTypedElementSyncFromWorldTag, FTypedElementActorTag>()
+		.Compile()
+	);
+}
diff --git a/Engine/Plugins/Experimental/EditorDataStorage/Source/TedsCore/Private/Compatibility/TypedElementActorIconOverrideQueries.h b/Engine/Plugins/Experimental/EditorDataStorage/Source/TedsCore/Private/Compatibility/TypedElementActorIconOverrideQueries.h
new file mode 100644
--- /dev/null
+++ b/Engine/Plugins/Experimental/EditorDataStorage/Source/TedsCore/Private/Compatibility/TypedElementActorIconOverrideQueries.h
@@ -0,0 +1,19 @@
+// Copyright Epic Games, Inc. All Rights Reserved.
+
+#pragma once
+
+#include "Elements/Interfaces/TypedElementDataStorageFactory.h"
+#include "UObject/ObjectMacros.h"
+
+#include "TypedElementActorIconOverrideQueries.generated.h"
+
+UCLASS()
+class UActorIconOverrideDataStorageFactory : public UTypedElementDataStorageFactory
+{
+	GENERATED_BODY()
+
+public:
+	virtual ~UActorIconOverrideDataStorageFactory() override = default;
+
+	virtual void RegisterQueries(ITypedElementDataStorageInterface& DataStorage) override;
+};
diff --git a/Engine/Source/Runtime/TypedElementFramework/Public/Elements/Columns/TypedElementIconOverrideColumns.h b/Engine/Source/Runtime/TypedElementFramework/Public/Elements/Columns/TypedElementIconOverrideColumns.h
new file mode 100644
--- /dev/null
+++ b/Engine/Source/Runtime/TypedElementFramework/Public/Elements/Columns/TypedElementIconOverrideColumns.h
@@ -0,0 +1,21 @@
+// Copyright Epic Games, Inc. All Rights Reserved.
+
+#pragma once
+
+#include "Elements/Interfaces/TypedElementDataStorageInterface.h"
+#include "UObject/ObjectMacros.h"
+
+#include "TypedElementIconOverrideColumns.generated.h"
+
+/**
+ * When present this column overrides any icon that would normally be used to represent the row.
+ */
+USTRUCT(meta = (DisplayName = "Icon override"))
+struct FTypedElementIconOverrideColumn final : public FEditorDataStorageColumn
+{
+	GENERATED_BODY()
+
+	/** Name of the replacement icon. */
+	UPROPERTY()
+	FName IconName;
+};
\ No newline at end of file
