TEDS UI:
Add a new FSimpleWidgetConstructor class to reduce the boilerplate required to create TEDS widgets
Create a new "General.RowLabel" purpose family, and add label widgets for UObjects, Assets and a specific one for the Scene Outliner

diff --git a/Engine/Plugins/Experimental/EditorDataStorage/Source/TedsUI/Private/Widgets/LabelWidget.cpp b/Engine/Plugins/Experimental/EditorDataStorage/Source/TedsUI/Private/Widgets/LabelWidget.cpp
--- a/Engine/Plugins/Experimental/EditorDataStorage/Source/TedsUI/Private/Widgets/LabelWidget.cpp
+++ b/Engine/Plugins/Experimental/EditorDataStorage/Source/TedsUI/Private/Widgets/LabelWidget.cpp
@@ -8,6 +8,7 @@
 #include "Elements/Columns/TypedElementMiscColumns.h"
 #include "Elements/Columns/TypedElementSlateWidgetColumns.h"
 #include "Elements/Columns/TypedElementValueCacheColumns.h"
+#include "Elements/Framework/TypedElementAttributeBinding.h"
 #include "Elements/Framework/TypedElementQueryBuilder.h"
 #include "Elements/Interfaces/TypedElementDataStorageInterface.h"
 #include "Elements/Interfaces/Capabilities/TypedElementUiEditableCapability.h"
@@ -23,123 +24,24 @@
 // ULabelWidgetFactory
 //
 
-namespace UE::Editor::DataStorage
-{
-	namespace UI::Private
-	{
-		static void UpdateTextWidget(const TWeakPtr<SWidget>& Widget, const FTypedElementLabelColumn& Label, const uint64* HashValue)
-		{
-			TSharedPtr<SWidget> WidgetPointer = Widget.Pin();
-			if (ensureMsgf(WidgetPointer, TEXT("Referenced widget is not valid. A constructed widget may not have been cleaned up. This can "
-				"also happen if this processor is running in the same phase as the processors responsible for cleaning up old "
-				"references.")))
-			{
-				if (TSharedPtr<ITypedElementUiTextCapability> Text = WidgetPointer->GetMetaData<ITypedElementUiTextCapability>())
-				{
-					Text->SetText(FText::FromString(Label.Label));
-				}
-
-				if (TSharedPtr<ITypedElementUiTooltipCapability> ToolTip = WidgetPointer->GetMetaData<ITypedElementUiTooltipCapability>())
-				{
-					if (!HashValue)
-					{
-						ToolTip->SetToolTipText(FText::FromString(Label.Label));
-					}
-					else
-					{
-						ToolTip->SetToolTipText(FText::FromString(FString::Format(TEXT("{0}\nHash: {1}"), { Label.Label, *HashValue })));
-					}
-				}
-			}
-		}
-
-		static void SyncColumnsToWidget(
-			ITypedElementDataStorageInterface* DataStorage,
-			RowHandle TargetRow,
-			FTypedElementU64IntValueCacheColumn& TextHash,
-			const TWeakPtr<SWidget>& Widget,
-			bool bShowHashInToolTip)
-		{
-			if (const FTypedElementLabelColumn* LabelColumn = DataStorage->GetColumn<FTypedElementLabelColumn>(TargetRow))
-			{
-				if (const FTypedElementLabelHashColumn* LabelHashColumn = DataStorage->GetColumn<FTypedElementLabelHashColumn>(TargetRow))
-				{
-					if (LabelHashColumn->LabelHash != TextHash.Value)
-					{
-						UpdateTextWidget(Widget, *LabelColumn, bShowHashInToolTip ? &LabelHashColumn->LabelHash : nullptr);
-						TextHash.Value = LabelHashColumn->LabelHash;
-					}
-				}
-				else
-				{
-					UpdateTextWidget(Widget, *LabelColumn, nullptr);
-				}
-			}
-		}
-	} // namespace UI::Private
-} // namespace UE::Editor::DataStorage
-
-void ULabelWidgetFactory::RegisterQueries(ITypedElementDataStorageInterface& DataStorage)
-{
-	using namespace TypedElementQueryBuilder;
-	using namespace TypedElementDataStorage;
-	using namespace UE::Editor;
-	
-	DataStorage::QueryHandle UpdateLabelWidget = DataStorage.RegisterQuery(
-		Select()
-			.ReadOnly<FTypedElementLabelColumn>()
-		.Where()
-			.Any<FTypedElementSyncFromWorldTag, FTypedElementSyncBackToWorldTag>()
-			.None<FTypedElementLabelHashColumn>()
-		.Compile());
-
-	DataStorage::QueryHandle UpdateLabelAndHashWidget = DataStorage.RegisterQuery(
-		Select()
-			.ReadOnly<FTypedElementLabelColumn, FTypedElementLabelHashColumn>()
-		.Where()
-			.Any<FTypedElementSyncFromWorldTag, FTypedElementSyncBackToWorldTag>()
-		.Compile());
-
-	DataStorage.RegisterQuery(
-		Select(
-			TEXT("Sync label to widget"),
-			FProcessor(EQueryTickPhase::FrameEnd, DataStorage.GetQueryTickGroupName(EQueryTickGroups::SyncWidgets))
-				.SetExecutionMode(EExecutionMode::GameThread),
-			[](
-				IQueryContext& Context, 
-				FTypedElementSlateWidgetReferenceColumn& Widget,
-				FTypedElementU64IntValueCacheColumn& TextHash,
-				const FLabelWidgetColumn& Config,
-				const FTypedElementRowReferenceColumn& Target)
-			{
-				Context.RunSubquery(0, Target.Row, CreateSubqueryCallbackBinding(
-					[&Widget](const FTypedElementLabelColumn& Label)
-					{
-						DataStorage::UI::Private::UpdateTextWidget(Widget.Widget, Label, nullptr);
-					}));
-				Context.RunSubquery(1, Target.Row, CreateSubqueryCallbackBinding(
-					[&Widget, &TextHash, &Config](const FTypedElementLabelColumn& Label, const FTypedElementLabelHashColumn& Hash)
-					{
-						if (Hash.LabelHash != TextHash.Value)
-						{
-							DataStorage::UI::Private::UpdateTextWidget(Widget.Widget, Label, Config.bShowHashInTooltip ? &Hash.LabelHash : nullptr);
-							TextHash.Value = Hash.LabelHash;
-						}
-					}));
-			})
-		.DependsOn()
-			.SubQuery({ UpdateLabelWidget, UpdateLabelAndHashWidget })
-		.Compile()
-	);
-}
-
 void ULabelWidgetFactory::RegisterWidgetConstructors(ITypedElementDataStorageInterface& DataStorage,
 	ITypedElementDataStorageUiInterface& DataStorageUi) const
 {
 	DataStorageUi.RegisterWidgetFactory<FLabelWidgetConstructor>(FName(TEXT("General.Cell")), 
 		TypedElementDataStorage::FColumn<FTypedElementLabelColumn>() || (TypedElementDataStorage::FColumn<FTypedElementLabelColumn>() && TypedElementDataStorage::FColumn<FTypedElementLabelHashColumn>()));
+
+	DataStorageUi.RegisterWidgetFactory<FLabelWidgetConstructor>(FName(TEXT("General.RowLabel")), 
+	TypedElementDataStorage::FColumn<FTypedElementLabelColumn>() || (TypedElementDataStorage::FColumn<FTypedElementLabelColumn>() && TypedElementDataStorage::FColumn<FTypedElementLabelHashColumn>()));
 }
 
+void ULabelWidgetFactory::RegisterWidgetPurposes(ITypedElementDataStorageUiInterface& DataStorageUi) const
+{
+	DataStorageUi.RegisterWidgetPurpose(FName(TEXT("General.RowLabel")), ITypedElementDataStorageUiInterface::EPurposeType::UniqueByNameAndColumn,
+	LOCTEXT("GeneralRowLabelPurpose", "Specific purpose to request a widget to display a user facing display name for a row."));
+	
+	DataStorageUi.RegisterWidgetPurpose(FName(TEXT("General.RowLabel.Default")), ITypedElementDataStorageUiInterface::EPurposeType::UniqueByName,
+	LOCTEXT("GeneralRowLabelDefaultPurpose", "Default purpose to request a widget to display a user facing display name for a row."));
+}
 
 
 //
@@ -147,12 +49,7 @@ void ULabelWidgetFactory::RegisterWidgetConstructors(ITypedElementDataStorageInt
 //
 
 FLabelWidgetConstructor::FLabelWidgetConstructor()
-	: Super(FLabelWidgetConstructor::StaticStruct())
-{
-}
-
-FLabelWidgetConstructor::FLabelWidgetConstructor(const UScriptStruct* InTypeInfo)
-	: Super(InTypeInfo)
+	: Super(StaticStruct())
 {
 }
 
@@ -161,107 +58,19 @@ TConstArrayView<const UScriptStruct*> FLabelWidgetConstructor::GetAdditionalColu
 	static const TTypedElementColumnTypeList<
 		FTypedElementRowReferenceColumn,
 		FTypedElementU64IntValueCacheColumn,
-		FLabelWidgetColumn,
 		FExternalWidgetSelectionColumn> Columns;
 	return Columns;
 }
 
-TSharedPtr<SWidget> FLabelWidgetConstructor::Construct(
-	UE::Editor::DataStorage::RowHandle Row,
-	ITypedElementDataStorageInterface* DataStorage,
-	ITypedElementDataStorageUiInterface* DataStorageUi,
+TSharedPtr<SWidget> FLabelWidgetConstructor::CreateWidget(ITypedElementDataStorageInterface* DataStorage,
+	ITypedElementDataStorageUiInterface* DataStorageUi, RowHandle TargetRow, RowHandle WidgetRow,
 	const TypedElementDataStorage::FMetaDataView& Arguments)
 {
-	TSharedPtr<SWidget> Result;
-	const bool* IsEditable = Arguments.FindForColumn<FTypedElementLabelColumn>(TypedElementDataStorage::IsEditableName).TryGetExact<bool>();
-	if (IsEditable && *IsEditable)
-	{
-		if (FTypedElementRowReferenceColumn* TargetRowColumn = DataStorage->GetColumn<FTypedElementRowReferenceColumn>(Row))
-		{
-			TSharedPtr<SInlineEditableTextBlock> TextBlock = SNew(SInlineEditableTextBlock)
-				.OnTextCommitted_Lambda(
-					[DataStorage, TargetRow = TargetRowColumn->Row](const FText& NewText, ETextCommit::Type CommitInfo)
-					{
-						// This callback happens on the game thread so it's safe to directly call into the data storage.
-						FString NewLabelText = NewText.ToString();
-						if (FTypedElementLabelHashColumn* LabelHashColumn = DataStorage->GetColumn<FTypedElementLabelHashColumn>(TargetRow))
-						{
-							LabelHashColumn->LabelHash = CityHash64(reinterpret_cast<const char*>(*NewLabelText), NewLabelText.Len() * sizeof(**NewLabelText));
-						}
-						if (FTypedElementLabelColumn* LabelColumn = DataStorage->GetColumn<FTypedElementLabelColumn>(TargetRow))
-						{
-							LabelColumn->Label = MoveTemp(NewLabelText);
-						}
-						DataStorage->AddColumn<FTypedElementSyncBackToWorldTag>(TargetRow);
-					})
-				.OnVerifyTextChanged_Lambda([](const FText& Label, FText& ErrorMessage)
-					{
-						// Note: The use of actor specific functionality should be minimized, but this function acts generic enough that the 
-						// use of actor is just in names.
-						return FActorEditorUtils::ValidateActorName(Label, ErrorMessage);
-					})
-				.IsSelected_Static(&FLabelWidgetConstructor::IsWidgetSelected, DataStorage, Row);;
-			TextBlock->AddMetadata(MakeShared<TTypedElementUiEditableCapability<SInlineEditableTextBlock>>(*TextBlock));
-			TextBlock->AddMetadata(MakeShared<TTypedElementUiTextCapability<SInlineEditableTextBlock>>(*TextBlock));
-			TextBlock->AddMetadata(MakeShared<TTypedElementUiTooltipCapability<SInlineEditableTextBlock>>(*TextBlock));
-			TextBlock->AddMetadata(MakeShared<TTypedElementUiStyleOverrideCapability<SInlineEditableTextBlock>>(*TextBlock));
-			Result = TextBlock;
-		}
-	}
-	else
-	{
-		TSharedPtr<STextBlock> TextBlock = SNew(STextBlock)
-			.IsEnabled(false);
-		TextBlock->AddMetadata(MakeShared<TTypedElementUiTextCapability<STextBlock>>(*TextBlock));
-		TextBlock->AddMetadata(MakeShared<TTypedElementUiTooltipCapability<STextBlock>>(*TextBlock));
-		Result = TextBlock;
-	}
-	
-	if (Result)
-	{
-		DataStorage->GetColumn<FTypedElementSlateWidgetReferenceColumn>(Row)->Widget = Result;
-		if (SetColumns(DataStorage, Row))
-		{
-			if (FinalizeWidget(DataStorage, DataStorageUi, Row, Result))
-			{
-				AddDefaultWidgetColumns(Row, DataStorage);
-				return Result;
-			}
-		}
-	}
+	UE::EditorDataStorage::FAttributeBinder Binder(TargetRow, DataStorage);
 
-	return nullptr;
+	return SNew(STextBlock)
+		.Text(Binder.BindText(&FTypedElementLabelColumn::Label))
+		.ToolTipText(Binder.BindText(&FTypedElementLabelColumn::Label));
 }
 
-bool FLabelWidgetConstructor::SetColumns(ITypedElementDataStorageInterface* DataStorage, UE::Editor::DataStorage::RowHandle Row)
-{
-	DataStorage->AddColumn(Row, FLabelWidgetColumn{ .bShowHashInTooltip = (MatchedColumnTypes.Num() == 2) });
-	return true;
-}
-
-bool FLabelWidgetConstructor::FinalizeWidget(
-	ITypedElementDataStorageInterface* DataStorage,
-	ITypedElementDataStorageUiInterface* DataStorageUi,
-	UE::Editor::DataStorage::RowHandle Row,
-	const TSharedPtr<SWidget>& Widget)
-{
-	UE::Editor::DataStorage::UI::Private::SyncColumnsToWidget(
-		DataStorage, 
-		DataStorage->GetColumn<FTypedElementRowReferenceColumn>(Row)->Row,
-		*DataStorage->GetColumn<FTypedElementU64IntValueCacheColumn>(Row),
-		Widget, MatchedColumnTypes.Num() == 2);
-	return true;
-}
-
-bool FLabelWidgetConstructor::IsWidgetSelected(ITypedElementDataStorageInterface* DataStorage, UE::Editor::DataStorage::RowHandle UiRow)
-{
-	if (const FExternalWidgetSelectionColumn* ExternalWidgetSelectionColumn = DataStorage->GetColumn<FExternalWidgetSelectionColumn>(UiRow))
-	{
-		if (ExternalWidgetSelectionColumn->IsSelected.IsBound())
-		{
-			return ExternalWidgetSelectionColumn->IsSelected.Execute();
-		}
-	}
-	return false;
-}
 #undef LOCTEXT_NAMESPACE
diff --git a/Engine/Plugins/Experimental/EditorDataStorage/Source/TedsUI/Public/Widgets/LabelWidget.h b/Engine/Plugins/Experimental/EditorDataStorage/Source/TedsUI/Public/Widgets/LabelWidget.h
--- a/Engine/Plugins/Experimental/EditorDataStorage/Source/TedsUI/Public/Widgets/LabelWidget.h
+++ b/Engine/Plugins/Experimental/EditorDataStorage/Source/TedsUI/Public/Widgets/LabelWidget.h
@@ -20,13 +20,13 @@ class ULabelWidgetFactory : public UTypedElementDataStorageFactory
 public:
 	~ULabelWidgetFactory() override = default;
 
-	TEDSUI_API void RegisterQueries(ITypedElementDataStorageInterface& DataStorage) override;
 	TEDSUI_API void RegisterWidgetConstructors(ITypedElementDataStorageInterface& DataStorage,
 		ITypedElementDataStorageUiInterface& DataStorageUi) const override;
+	TEDSUI_API virtual void RegisterWidgetPurposes(ITypedElementDataStorageUiInterface& DataStorageUi) const override;
 };
 
 USTRUCT()
-struct FLabelWidgetConstructor : public FTypedElementWidgetConstructor
+struct FLabelWidgetConstructor : public FSimpleWidgetConstructor
 {
 	GENERATED_BODY()
 
@@ -34,26 +34,14 @@ public:
 	TEDSUI_API FLabelWidgetConstructor();
 	~FLabelWidgetConstructor() override = default;
 
-	TEDSUI_API TConstArrayView<const UScriptStruct*> GetAdditionalColumnsList() const override;
+	TEDSUI_API virtual TConstArrayView<const UScriptStruct*> GetAdditionalColumnsList() const override;
 
-	TEDSUI_API TSharedPtr<SWidget> Construct(
-		UE::Editor::DataStorage::RowHandle Row,
+	TEDSUI_API virtual TSharedPtr<SWidget> CreateWidget(
 		ITypedElementDataStorageInterface* DataStorage,
 		ITypedElementDataStorageUiInterface* DataStorageUi,
+		RowHandle TargetRow,
+		RowHandle WidgetRow, 
 		const TypedElementDataStorage::FMetaDataView& Arguments) override;
 
 protected:
-	explicit FLabelWidgetConstructor(const UScriptStruct* InTypeInfo);
-	bool SetColumns(ITypedElementDataStorageInterface* DataStorage, UE::Editor::DataStorage::RowHandle Row) override;
-	bool FinalizeWidget(ITypedElementDataStorageInterface* DataStorage, ITypedElementDataStorageUiInterface* DataStorageUi,
-		UE::Editor::DataStorage::RowHandle Row, const TSharedPtr<SWidget>& Widget) override;
-	static bool IsWidgetSelected(ITypedElementDataStorageInterface* DataStorage, UE::Editor::DataStorage::RowHandle UiRow);
-};
-
-USTRUCT(meta = (DisplayName = "Label widget"))
-struct FLabelWidgetColumn : public FEditorDataStorageColumn
-{
-	GENERATED_BODY()
-
-	bool bShowHashInTooltip { false };
 };
\ No newline at end of file
diff --git a/Engine/Plugins/Experimental/EditorDataStorageFeatures/Source/TedsAssetData/Private/Widgets/AssetDataLabelWidget.cpp b/Engine/Plugins/Experimental/EditorDataStorageFeatures/Source/TedsAssetData/Private/Widgets/AssetDataLabelWidget.cpp
new file mode 100644
--- /dev/null
+++ b/Engine/Plugins/Experimental/EditorDataStorageFeatures/Source/TedsAssetData/Private/Widgets/AssetDataLabelWidget.cpp
@@ -0,0 +1,98 @@
+// Copyright Epic Games, Inc. All Rights Reserved.
+
+#include "Widgets/AssetDataLabelWidget.h"
+
+#include "TedsAssetDataColumns.h"
+#include "Elements/Framework/TypedElementAttributeBinding.h"
+#include "Widgets/Images/SImage.h"
+#include "Widgets/Layout/SSpacer.h"
+#include "Widgets/Text/STextBlock.h"
+#include "Widgets/SBoxPanel.h"
+
+#define LOCTEXT_NAMESPACE "FAssetDataLabelWidgetConstructor"
+
+void UAssetDataLabelWidgetFactory::RegisterWidgetConstructors(ITypedElementDataStorageInterface& DataStorage,
+                                                              ITypedElementDataStorageUiInterface& DataStorageUi) const
+{
+	DataStorageUi.RegisterWidgetFactory<FAssetDataLabelWidgetConstructor>(
+		TEXT("General.RowLabel"),
+		TypedElementDataStorage::FColumn<FItemNameColumn_Experimental>() && TypedElementDataStorage::FColumn<FAssetTag>());
+}
+
+FAssetDataLabelWidgetConstructor::FAssetDataLabelWidgetConstructor()
+	: FSimpleWidgetConstructor(StaticStruct())
+{
+}
+
+
+TSharedPtr<SWidget> FAssetDataLabelWidgetConstructor::CreateWidget(ITypedElementDataStorageInterface* DataStorage,
+	ITypedElementDataStorageUiInterface* DataStorageUi, RowHandle TargetRow, RowHandle WidgetRow,
+	const TypedElementDataStorage::FMetaDataView& Arguments)
+{
+	if (DataStorage->IsRowAvailable(TargetRow))
+	{
+		UE::EditorDataStorage::FAttributeBinder Binder(TargetRow, DataStorage);
+		
+		return SNew(SHorizontalBox)
+			+SHorizontalBox::Slot()
+				.AutoWidth()
+			[
+				// TODO: Asset color in TEDS
+				SNew(SImage)
+					.Image(FAppStyle::GetBrush(FName(TEXT("ContentBrowser.ColumnViewAssetIcon"))))
+					.ColorAndOpacity(FSlateColor::UseForeground())
+			]
+			+SHorizontalBox::Slot()
+				.AutoWidth()
+			[
+				SNew(SSpacer)
+					.Size(FVector2D(5.0f, 0.0f))
+			]
+			+SHorizontalBox::Slot()
+				.FillWidth(1.0f)
+			[
+				SNew(STextBlock)
+				.Text(Binder.BindText(&FItemNameColumn_Experimental::Name))
+				.ToolTipText_Lambda([DataStorage, DataRow = TargetRow]()
+					{
+						return FAssetDataLabelWidgetConstructor::ConstructToolTip(DataStorage, DataRow);
+					})
+			];
+	}
+	else
+	{
+		return SNullWidget::NullWidget;
+	}
+}
+
+FText FAssetDataLabelWidgetConstructor::ConstructToolTip(
+	ITypedElementDataStorageInterface* DataStorage, UE::Editor::DataStorage::RowHandle DataRow)
+{
+	TStringBuilder<1024> ToolTip;
+	if (FItemNameColumn_Experimental* ItemName = DataStorage->GetColumn<FItemNameColumn_Experimental>(DataRow))
+	{
+		ToolTip.Append(ItemName->Name.ToString());
+		ToolTip.AppendChar(TEXT('\n'));
+	}
+	if (FAssetPathColumn_Experimental* AssetPath = DataStorage->GetColumn<FAssetPathColumn_Experimental>(DataRow))
+	{
+		ToolTip.AppendChar(TEXT('\n'));
+		ToolTip.Append(LOCTEXT("AssetPath", "Asset path: ").ToString());
+		ToolTip.Append(AssetPath->Path.ToString());
+	}
+	if (FVersePathColumn* VersePath = DataStorage->GetColumn<FVersePathColumn>(DataRow))
+	{
+		ToolTip.AppendChar(TEXT('\n'));
+		ToolTip.Append(LOCTEXT("VersePath", "Verse path: ").ToString());
+		ToolTip.Append(VersePath->VersePath.AsStringView());
+	}
+	if (FVirtualPathColumn_Experimental* VirtualPath = DataStorage->GetColumn<FVirtualPathColumn_Experimental>(DataRow))
+	{
+		ToolTip.AppendChar(TEXT('\n'));
+		ToolTip.Append(LOCTEXT("VirtualPath", "Virtual path: ").ToString());
+		ToolTip.Append(VirtualPath->VirtualPath.ToString());
+	}
+	return FText::FromString(ToolTip.ToString());
+}
+
+#undef LOCTEXT_NAMESPACE
\ No newline at end of file
diff --git a/Engine/Plugins/Experimental/EditorDataStorageFeatures/Source/TedsAssetData/Public/Widgets/AssetDataLabelWidget.h b/Engine/Plugins/Experimental/EditorDataStorageFeatures/Source/TedsAssetData/Public/Widgets/AssetDataLabelWidget.h
new file mode 100644
--- /dev/null
+++ b/Engine/Plugins/Experimental/EditorDataStorageFeatures/Source/TedsAssetData/Public/Widgets/AssetDataLabelWidget.h
@@ -0,0 +1,47 @@
+// Copyright Epic Games, Inc. All Rights Reserved.
+
+#pragma once
+
+#include "Elements/Interfaces/TypedElementDataStorageFactory.h"
+#include "Elements/Interfaces/TypedElementDataStorageUiInterface.h"
+#include "Internationalization/Text.h"
+#include "UObject/ObjectMacros.h"
+
+#include "AssetDataLabelWidget.generated.h"
+
+class ITypedElementDataStorageInterface;
+class UScriptStruct;
+
+UCLASS()
+class UAssetDataLabelWidgetFactory : public UTypedElementDataStorageFactory
+{
+	GENERATED_BODY()
+
+public:
+	~UAssetDataLabelWidgetFactory() override = default;
+
+	TEDSASSETDATA_API void RegisterWidgetConstructors(ITypedElementDataStorageInterface& DataStorage,
+		ITypedElementDataStorageUiInterface& DataStorageUi) const override;
+};
+
+// Label widget for assets in TEDS
+USTRUCT()
+struct FAssetDataLabelWidgetConstructor : public FSimpleWidgetConstructor
+{
+	GENERATED_BODY()
+
+public:
+	TEDSASSETDATA_API FAssetDataLabelWidgetConstructor();
+	~FAssetDataLabelWidgetConstructor() override = default;
+
+	TEDSASSETDATA_API virtual TSharedPtr<SWidget> CreateWidget(
+		ITypedElementDataStorageInterface* DataStorage,
+		ITypedElementDataStorageUiInterface* DataStorageUi,
+		RowHandle TargetRow,
+		RowHandle WidgetRow, 
+		const TypedElementDataStorage::FMetaDataView& Arguments) override;
+
+protected:
+	
+	static FText ConstructToolTip(ITypedElementDataStorageInterface* DataStorage, UE::Editor::DataStorage::RowHandle DataRow);
+};
\ No newline at end of file
diff --git a/Engine/Plugins/Experimental/EditorDataStorageFeatures/Source/TedsAssetData/TedsAssetData.build.cs b/Engine/Plugins/Experimental/EditorDataStorageFeatures/Source/TedsAssetData/TedsAssetData.build.cs
--- a/Engine/Plugins/Experimental/EditorDataStorageFeatures/Source/TedsAssetData/TedsAssetData.build.cs
+++ b/Engine/Plugins/Experimental/EditorDataStorageFeatures/Source/TedsAssetData/TedsAssetData.build.cs
@@ -26,6 +26,8 @@ public class TedsAssetData : ModuleRules
 					"Engine",
 					"TypedElementFramework",
 					"UnrealEd",
+					"SlateCore",
+					"Slate"
 				});
 
 			PublicDependencyModuleNames.AddRange(
diff --git a/Engine/Plugins/Experimental/EditorDataStorageFeatures/Source/TedsOutliner/Private/TedsOutlinerImpl.cpp b/Engine/Plugins/Experimental/EditorDataStorageFeatures/Source/TedsOutliner/Private/TedsOutlinerImpl.cpp
--- a/Engine/Plugins/Experimental/EditorDataStorageFeatures/Source/TedsOutliner/Private/TedsOutlinerImpl.cpp
+++ b/Engine/Plugins/Experimental/EditorDataStorageFeatures/Source/TedsOutliner/Private/TedsOutlinerImpl.cpp
@@ -16,6 +16,8 @@
 #include "TedsOutlinerFilter.h"
 #include "TedsOutlinerItem.h"
 #include "Columns/TedsOutlinerColumns.h"
+#include "Elements/Columns/TypedElementCompatibilityColumns.h"
+#include "Elements/Columns/TypedElementOverrideColumns.h"
 #include "Elements/Columns/TypedElementSelectionColumns.h"
 #include "Elements/Columns/TypedElementSlateWidgetColumns.h"
 #include "Filters/FilterBase.h"
@@ -62,8 +64,8 @@ void FTedsOutlinerImpl::CreateLabelWidgetConstructors()
 
 		bool bFoundWidget = false;
 
-		// We also want to look at the ItemLabel purpose for the label
-		TArray<FName> ItemLabelCellWidgetPurposes{TEXT("SceneOutliner.ItemLabel.Cell")};
+		// We also want to look at the RowLabel purpose for the label
+		TArray<FName> ItemLabelCellWidgetPurposes{TEXT("SceneOutliner.RowLabel")};
 		ItemLabelCellWidgetPurposes.Append(CellWidgetPurposes);
 		
 		for(const FName& WidgetPurpose : ItemLabelCellWidgetPurposes)
@@ -92,22 +94,9 @@ void FTedsOutlinerImpl::CreateLabelWidgetConstructors()
 		return OutWidgetConstructorPtr;
 	};
 	
-
-	TypedElementDataStorage::QueryHandle TypeColumnQueryHandle = Storage->RegisterQuery(
-																	Select()
-																		.ReadOnly<FTypedElementClassTypeInfoColumn>()
-																	.Compile()
-																	);
-
-
-	if (TSharedPtr<FTypedElementWidgetConstructor> TypeColumnWidgetConstructor = CreateWidgetConstructorForQuery(Storage->GetQueryDescription(TypeColumnQueryHandle)))
-	{
-		QueryToWidgetConstructorMap.Emplace(TypeColumnQueryHandle, TypeColumnWidgetConstructor);
-	}
-
 	TypedElementDataStorage::QueryHandle LabelColumnQueryHandle = Storage->RegisterQuery(
 																	Select()
-																		.ReadWrite<FTypedElementLabelColumn>()
+																		.ReadWrite<FTypedElementLabelColumn, FTypedElementClassTypeInfoColumn>()
 																	.Compile()
 																	);
 
diff --git a/Engine/Plugins/Experimental/EditorDataStorageFeatures/Source/TedsOutliner/Private/Widgets/OutlinerLabelWidget.cpp b/Engine/Plugins/Experimental/EditorDataStorageFeatures/Source/TedsOutliner/Private/Widgets/OutlinerLabelWidget.cpp
new file mode 100644
--- /dev/null
+++ b/Engine/Plugins/Experimental/EditorDataStorageFeatures/Source/TedsOutliner/Private/Widgets/OutlinerLabelWidget.cpp
@@ -0,0 +1,230 @@
+// Copyright Epic Games, Inc. All Rights Reserved.
+
+#include "Widgets/OutlinerLabelWidget.h"
+
+#include "ActorEditorUtils.h"
+#include "Elements/Columns/TypedElementCompatibilityColumns.h"
+#include "Elements/Columns/TypedElementLabelColumns.h"
+#include "Elements/Framework/TypedElementAttributeBinding.h"
+#include "TedsTableViewerUtils.h"
+#include "Columns/SlateDelegateColumns.h"
+#include "Elements/Columns/TypedElementMiscColumns.h"
+#include "Elements/Columns/TypedElementOverrideColumns.h"
+#include "Elements/Columns/TypedElementTypeInfoColumns.h"
+#include "Elements/Interfaces/Capabilities/TypedElementUiEditableCapability.h"
+#include "Elements/Interfaces/Capabilities/TypedElementUiStyleOverrideCapability.h"
+#include "Widgets/Images/SImage.h"
+#include "Widgets/Layout/SSpacer.h"
+#include "Widgets/Text/STextBlock.h"
+#include "Widgets/SBoxPanel.h"
+#include "Widgets/Images/SLayeredImage.h"
+#include "Widgets/Text/SInlineEditableTextBlock.h"
+
+#define LOCTEXT_NAMESPACE "FOutlinerLabelWidgetConstructor"
+
+namespace UE::OutlinerLabelWidget::Local
+{
+	const FSlateBrush* GetOverrideBadgeFirstLayer(const EOverriddenState& OverriddenState)
+	{
+		switch(OverriddenState)
+		{
+		case EOverriddenState::Added:
+			return FAppStyle::GetBrush("SceneOutliner.OverrideAddedBase");
+			
+		case EOverriddenState::AllOverridden:
+			// Not implemented yet
+			break;
+			
+		case EOverriddenState::HasOverrides:
+			return FAppStyle::GetBrush("SceneOutliner.OverrideInsideBase");
+			
+		case EOverriddenState::NoOverrides:
+			// No icon for no overrides
+			break;
+			
+		case EOverriddenState::SubObjectsHasOverrides:
+			// Not implemented yet
+			break;
+		}
+		
+		return FAppStyle::GetBrush("NoBrush");
+	}
+	
+	const FSlateBrush* GetOverrideBadgeSecondLayer(const EOverriddenState& OverriddenState)
+	{
+		switch(OverriddenState)
+		{
+		case EOverriddenState::Added:
+			return FAppStyle::GetBrush("SceneOutliner.OverrideAdded");
+			
+		case EOverriddenState::AllOverridden:
+			// Not implemented yet
+				break;
+			
+		case EOverriddenState::HasOverrides:
+			return FAppStyle::GetBrush("SceneOutliner.OverrideInside");
+			
+		case EOverriddenState::NoOverrides:
+			// No icon for no overrides
+				break;
+			
+		case EOverriddenState::SubObjectsHasOverrides:
+			// Not implemented yet
+				break;
+		}
+
+		return FAppStyle::GetBrush("NoBrush");
+	}
+
+	FText GetOverrideTooltip(const EOverriddenState& OverriddenState)
+	{
+		switch(OverriddenState)
+		{
+		case EOverriddenState::Added:
+			return LOCTEXT("OverrideAddedTooltip", "This entity has been added.");
+			
+		case EOverriddenState::AllOverridden:
+			// Not implemented yet
+				break;
+			
+		case EOverriddenState::HasOverrides:
+			return LOCTEXT("OverrideInsideTooltip", "At least one property or child has an override.");
+			
+		case EOverriddenState::NoOverrides:
+			// No icon for no overrides
+				break;
+			
+		case EOverriddenState::SubObjectsHasOverrides:
+			// Not implemented yet
+				break;
+		}
+
+		return FText::GetEmpty();
+	}
+}
+
+void UOutlinerLabelWidgetFactory::RegisterWidgetConstructors(ITypedElementDataStorageInterface& DataStorage,
+	ITypedElementDataStorageUiInterface& DataStorageUi) const
+{
+	DataStorageUi.RegisterWidgetFactory<FOutlinerLabelWidgetConstructor>(
+		TEXT("SceneOutliner.RowLabel"),
+		TypedElementDataStorage::FColumn<FTypedElementLabelColumn>() && TypedElementDataStorage::FColumn<FTypedElementClassTypeInfoColumn>());
+}
+
+FOutlinerLabelWidgetConstructor::FOutlinerLabelWidgetConstructor()
+	: FSimpleWidgetConstructor(StaticStruct())
+{
+}
+
+TSharedPtr<SWidget> FOutlinerLabelWidgetConstructor::CreateWidget(ITypedElementDataStorageInterface* DataStorage,
+	ITypedElementDataStorageUiInterface* DataStorageUi, RowHandle TargetRow, RowHandle WidgetRow,
+	const TypedElementDataStorage::FMetaDataView& Arguments)
+{
+	if(DataStorage->IsRowAvailable(TargetRow))
+	{
+		UE::EditorDataStorage::FAttributeBinder Binder(TargetRow, DataStorage);
+
+		TSharedRef<SLayeredImage> LayeredImageWidget = SNew(SLayeredImage)
+					.Image(UE::Editor::DataStorage::TableViewerUtils::GetIconForRow(DataStorage, TargetRow))
+					.ToolTipText(Binder.BindData(&FObjectOverrideColumn::OverriddenState, [](const EOverriddenState& OverriddenState)
+					{
+						return UE::OutlinerLabelWidget::Local::GetOverrideTooltip(OverriddenState);
+					}))
+					.ColorAndOpacity(FSlateColor::UseForeground());
+
+		LayeredImageWidget->AddLayer(Binder.BindData(&FObjectOverrideColumn::OverriddenState, [](const EOverriddenState& OverriddenState)
+		{
+			return UE::OutlinerLabelWidget::Local::GetOverrideBadgeFirstLayer(OverriddenState);
+		}));
+		
+		LayeredImageWidget->AddLayer(Binder.BindData(&FObjectOverrideColumn::OverriddenState, [](const EOverriddenState& OverriddenState)
+		{
+			return UE::OutlinerLabelWidget::Local::GetOverrideBadgeSecondLayer(OverriddenState);
+		}));
+
+		return SNew(SHorizontalBox)
+			+SHorizontalBox::Slot()
+				.AutoWidth()
+			[
+				LayeredImageWidget
+			]
+			+SHorizontalBox::Slot()
+				.AutoWidth()
+			[
+				SNew(SSpacer)
+					.Size(FVector2D(5.0f, 0.0f))
+			]
+			+SHorizontalBox::Slot()
+				.FillWidth(1.0f)
+			[
+				CreateLabel(DataStorage, DataStorageUi, TargetRow, WidgetRow, Arguments)
+			];
+	}
+	else
+	{
+		return SNew(STextBlock)
+			.Text(LOCTEXT("MissingRowReferenceColumn", "Unable to retrieve row reference."));
+	}
+	
+	
+}
+
+TSharedRef<SWidget> FOutlinerLabelWidgetConstructor::CreateLabel(ITypedElementDataStorageInterface* DataStorage,
+	ITypedElementDataStorageUiInterface* DataStorageUi, RowHandle TargetRow, RowHandle WidgetRow,
+	const TypedElementDataStorage::FMetaDataView& Arguments)
+{
+	TSharedRef<SWidget> Result = SNullWidget::NullWidget;
+	
+	const bool* IsEditable = Arguments.FindForColumn<FTypedElementLabelColumn>(TypedElementDataStorage::IsEditableName).TryGetExact<bool>();
+
+	if (IsEditable && *IsEditable)
+	{
+		UE::EditorDataStorage::FAttributeBinder TargetRowBinder(TargetRow, DataStorage);
+		UE::EditorDataStorage::FAttributeBinder WidgetRowBinder(WidgetRow, DataStorage);
+			
+		TSharedPtr<SInlineEditableTextBlock> TextBlock = SNew(SInlineEditableTextBlock)
+			.OnTextCommitted_Lambda(
+				[DataStorage, TargetRow](const FText& NewText, ETextCommit::Type CommitInfo)
+				{
+					// This callback happens on the game thread so it's safe to directly call into the data storage.
+					FString NewLabelText = NewText.ToString();
+					if (FTypedElementLabelHashColumn* LabelHashColumn = DataStorage->GetColumn<FTypedElementLabelHashColumn>(TargetRow))
+					{
+						LabelHashColumn->LabelHash = CityHash64(reinterpret_cast<const char*>(*NewLabelText), NewLabelText.Len() * sizeof(**NewLabelText));
+					}
+					if (FTypedElementLabelColumn* LabelColumn = DataStorage->GetColumn<FTypedElementLabelColumn>(TargetRow))
+					{
+						LabelColumn->Label = MoveTemp(NewLabelText);
+					}
+					DataStorage->AddColumn<FTypedElementSyncBackToWorldTag>(TargetRow);
+				})
+			.OnVerifyTextChanged_Lambda([](const FText& Label, FText& ErrorMessage)
+				{
+					// Note: The use of actor specific functionality should be minimized, but this function acts generic enough that the 
+					// use of actor is just in names.
+					return FActorEditorUtils::ValidateActorName(Label, ErrorMessage);
+				})
+			.Text(TargetRowBinder.BindText(&FTypedElementLabelColumn::Label))
+			.ToolTipText(TargetRowBinder.BindText(&FTypedElementLabelColumn::Label))
+			.IsSelected(WidgetRowBinder.BindEvent(&FExternalWidgetSelectionColumn::IsSelected));
+
+		TextBlock->AddMetadata(MakeShared<TTypedElementUiEditableCapability<SInlineEditableTextBlock>>(*TextBlock));
+		TextBlock->AddMetadata(MakeShared<TTypedElementUiStyleOverrideCapability<SInlineEditableTextBlock>>(*TextBlock));
+		Result = TextBlock.ToSharedRef();
+	}
+	else
+	{
+		UE::EditorDataStorage::FAttributeBinder TargetRowBinder(TargetRow, DataStorage);
+		
+		TSharedPtr<STextBlock> TextBlock = SNew(STextBlock)
+			.IsEnabled(false)
+			.Text(TargetRowBinder.BindText(&FTypedElementLabelColumn::Label))
+			.ToolTipText(TargetRowBinder.BindText(&FTypedElementLabelColumn::Label));
+		
+		Result = TextBlock.ToSharedRef();
+	}
+
+	return Result;
+}
+
+#undef LOCTEXT_NAMESPACE //"FOutlinerLabelWidgetConstructor"
diff --git a/Engine/Plugins/Experimental/EditorDataStorageFeatures/Source/TedsOutliner/Private/Widgets/OutlinerLabelWidget.h b/Engine/Plugins/Experimental/EditorDataStorageFeatures/Source/TedsOutliner/Private/Widgets/OutlinerLabelWidget.h
new file mode 100644
--- /dev/null
+++ b/Engine/Plugins/Experimental/EditorDataStorageFeatures/Source/TedsOutliner/Private/Widgets/OutlinerLabelWidget.h
@@ -0,0 +1,49 @@
+// Copyright Epic Games, Inc. All Rights Reserved.
+
+#pragma once
+
+#include "Elements/Interfaces/TypedElementDataStorageFactory.h"
+#include "Elements/Interfaces/TypedElementDataStorageUiInterface.h"
+#include "Internationalization/Text.h"
+#include "UObject/ObjectMacros.h"
+
+#include "OutlinerLabelWidget.generated.h"
+
+class ITypedElementDataStorageInterface;
+class UScriptStruct;
+
+UCLASS()
+class UOutlinerLabelWidgetFactory : public UTypedElementDataStorageFactory
+{
+	GENERATED_BODY()
+
+public:
+	~UOutlinerLabelWidgetFactory() override = default;
+
+	TEDSOUTLINER_API void RegisterWidgetConstructors(ITypedElementDataStorageInterface& DataStorage,
+		ITypedElementDataStorageUiInterface& DataStorageUi) const override;
+};
+
+// Label widget for the Scene Outliner that shows an icon (with optional override information) + a text label
+USTRUCT()
+struct FOutlinerLabelWidgetConstructor : public FSimpleWidgetConstructor
+{
+	GENERATED_BODY()
+
+public:
+	TEDSOUTLINER_API FOutlinerLabelWidgetConstructor();
+	~FOutlinerLabelWidgetConstructor() override = default;
+
+	TEDSOUTLINER_API virtual TSharedPtr<SWidget> CreateWidget(
+		ITypedElementDataStorageInterface* DataStorage,
+		ITypedElementDataStorageUiInterface* DataStorageUi,
+		RowHandle TargetRow,
+		RowHandle WidgetRow, 
+		const TypedElementDataStorage::FMetaDataView& Arguments) override;
+
+protected:
+
+	TSharedRef<SWidget> CreateLabel(ITypedElementDataStorageInterface* DataStorage,
+	ITypedElementDataStorageUiInterface* DataStorageUi, RowHandle TargetRow, RowHandle WidgetRow,
+	const TypedElementDataStorage::FMetaDataView& Arguments);
+};
\ No newline at end of file
diff --git a/Engine/Plugins/Experimental/EditorDataStorageFeatures/Source/TedsOutliner/Private/Widgets/OverrideWidget.cpp b/Engine/Plugins/Experimental/EditorDataStorageFeatures/Source/TedsOutliner/Private/Widgets/OverrideWidget.cpp
deleted file mode 100644
diff --git a/Engine/Plugins/Experimental/EditorDataStorageFeatures/Source/TedsOutliner/Private/Widgets/OverrideWidget.h b/Engine/Plugins/Experimental/EditorDataStorageFeatures/Source/TedsOutliner/Private/Widgets/OverrideWidget.h
deleted file mode 100644
diff --git a/Engine/Plugins/Experimental/EditorDataStorageFeatures/Source/TedsOutliner/Private/Widgets/TypeInfoWidget.cpp b/Engine/Plugins/Experimental/EditorDataStorageFeatures/Source/TedsOutliner/Private/Widgets/TypeInfoWidget.cpp
--- a/Engine/Plugins/Experimental/EditorDataStorageFeatures/Source/TedsOutliner/Private/Widgets/TypeInfoWidget.cpp
+++ b/Engine/Plugins/Experimental/EditorDataStorageFeatures/Source/TedsOutliner/Private/Widgets/TypeInfoWidget.cpp
@@ -9,6 +9,7 @@
 #include "Elements/Columns/TypedElementTypeInfoColumns.h"
 #include "Elements/Interfaces/Capabilities/TypedElementUiTextCapability.h"
 #include "Styling/SlateIconFinder.h"
+#include "TedsTableViewerUtils.h"
 #include "Widgets/Text/STextBlock.h"
 #include "Widgets/Images/SImage.h"
 
@@ -20,8 +21,6 @@ void UTypeInfoWidgetFactory::RegisterWidgetConstructors(ITypedElementDataStorage
 
 }
 
-TMap<FName, const FSlateBrush*> FTypeInfoWidgetConstructor::CachedIconMap;
-
 FTypeInfoWidgetConstructor::FTypeInfoWidgetConstructor()
 	: Super(FTypeInfoWidgetConstructor::StaticStruct())
 	, bUseIcon(false)
@@ -81,7 +80,7 @@ bool FTypeInfoWidgetConstructor::FinalizeWidget(ITypedElementDataStorageInterfac
 
 			SImage* WidgetInstance = static_cast<SImage*>(Widget.Get());
 			
-			WidgetInstance->SetImage(GetIconForRow(DataStorage, Row, TypeInfoColumn));
+			WidgetInstance->SetImage(UE::Editor::DataStorage::TableViewerUtils::GetIconForRow(DataStorage, Row));
 		}
 		else
 		{
@@ -123,48 +122,3 @@ bool FTypeInfoWidgetConstructor::FinalizeWidget(ITypedElementDataStorageInterfac
 
 	return true;
 }
-
-const FSlateBrush* FTypeInfoWidgetConstructor::GetIconForRow(ITypedElementDataStorageInterface* DataStorage, UE::Editor::DataStorage::RowHandle Row, const FTypedElementClassTypeInfoColumn* TypeInfoColumn)
-{
-	/* The logic here is very similar to SActorTreeLabel::GetIcon in ActorTreeItem.cpp which allows the actor to specify
-	 * an override for the icon, and has a fallback to the class icon if not.
-	 */
-	
-	const UClass* Type = TypeInfoColumn->TypeInfo.Get();
-	FName IconName;
-
-	// Allow the actor the first chance to provide an icon override
-	if (DataStorage->HasColumns<FTypedElementActorTag>(Row))
-	{
-		if (FTypedElementUObjectColumn* ActorStore = DataStorage->GetColumn<FTypedElementUObjectColumn>(Row))
-		{
-			if (const AActor* Actor = Cast<AActor>(ActorStore->Object))
-			{
-				IconName = Actor->GetCustomIconName();
-			}
-		}
-	}
-
-	if(IconName == NAME_None)
-	{
-		IconName = Type->GetFName();
-	}
-
-	// Check the cache if we already found an icon for this class
-	if(const FSlateBrush** CachedBrush = CachedIconMap.Find(IconName))
-	{
-		if(*CachedBrush)
-		{
-			return *CachedBrush;
-		}
-	}
-
-	if(const FSlateBrush* FoundSlateBrush = FSlateIconFinder::FindIconForClass(Type).GetOptionalIcon())
-	{
-		CachedIconMap.Add(IconName, FoundSlateBrush);
-		return FoundSlateBrush;
-	}
-
-	// Fallback to the regular actor icon if we haven't found any specific icon
-	return FSlateIconFinder::FindIconForClass(AActor::StaticClass()).GetOptionalIcon();;
-}
diff --git a/Engine/Plugins/Experimental/EditorDataStorageFeatures/Source/TedsOutliner/Public/Widgets/TypeInfoWidget.h b/Engine/Plugins/Experimental/EditorDataStorageFeatures/Source/TedsOutliner/Public/Widgets/TypeInfoWidget.h
--- a/Engine/Plugins/Experimental/EditorDataStorageFeatures/Source/TedsOutliner/Public/Widgets/TypeInfoWidget.h
+++ b/Engine/Plugins/Experimental/EditorDataStorageFeatures/Source/TedsOutliner/Public/Widgets/TypeInfoWidget.h
@@ -40,14 +40,6 @@ protected:
 
 protected:
 
-	// Get the icon for a given row, checking the cache to see if we already have one
-	static const FSlateBrush* GetIconForRow(ITypedElementDataStorageInterface* DataStorage, UE::Editor::DataStorage::RowHandle Row, const FTypedElementClassTypeInfoColumn* TypeInfoColumn);
-	
-protected:
-
-	// Cache to avoid looking up the icon for a class every time
-	static TMap<FName, const FSlateBrush*> CachedIconMap;
-
 	// Whether the widget created by this constructor should be icon or text
 	bool bUseIcon;
 };
diff --git a/Engine/Plugins/Experimental/EditorDataStorageFeatures/Source/TedsTableViewer/Private/TedsTableViewerUtils.cpp b/Engine/Plugins/Experimental/EditorDataStorageFeatures/Source/TedsTableViewer/Private/TedsTableViewerUtils.cpp
--- a/Engine/Plugins/Experimental/EditorDataStorageFeatures/Source/TedsTableViewer/Private/TedsTableViewerUtils.cpp
+++ b/Engine/Plugins/Experimental/EditorDataStorageFeatures/Source/TedsTableViewer/Private/TedsTableViewerUtils.cpp
@@ -2,9 +2,14 @@
 
 #include "TedsTableViewerUtils.h"
 
+#include "GameFramework/Actor.h"
+#include "Elements/Columns/TypedElementCompatibilityColumns.h"
+#include "Elements/Columns/TypedElementIconOverrideColumns.h"
 #include "Elements/Columns/TypedElementMiscColumns.h"
+#include "Elements/Columns/TypedElementTypeInfoColumns.h"
 #include "Elements/Common/TypedElementDataStorageLog.h"
 #include "Elements/Interfaces/TypedElementDataStorageUiInterface.h"
+#include "Styling/SlateIconFinder.h"
 
 namespace UE::Editor::DataStorage::TableViewerUtils
 {
@@ -139,6 +144,61 @@ namespace UE::Editor::DataStorage::TableViewerUtils
 		}
 		return nullptr;
 	}
+
+	const FSlateBrush* GetIconForRow(ITypedElementDataStorageInterface* DataStorage, RowHandle Row)
+	{
+		static TMap<FName, const FSlateBrush*> CachedIconMap;
+
+		auto FindCachedIcon = [](const FName& IconName) -> const FSlateBrush*
+		{
+			if (const FSlateBrush** CachedBrush = CachedIconMap.Find(IconName))
+			{
+				if (*CachedBrush)
+				{
+					return *CachedBrush;
+				}
+			}
+			return nullptr;
+		};
+
+		// Look for any icon overrides
+		if(const FTypedElementIconOverrideColumn* IconOverrideColumn = DataStorage->GetColumn<FTypedElementIconOverrideColumn>(Row))
+		{
+			const FName IconName = IconOverrideColumn->IconName;
+
+			if(const FSlateBrush* CachedBrush = FindCachedIcon(IconName))
+			{
+				return CachedBrush;
+			}
+			else if(const FSlateBrush* CustomBrush = FSlateIconFinder::FindIcon(IconName).GetOptionalIcon())
+			{
+				CachedIconMap.Add(IconName, CustomBrush);
+				return CustomBrush;
+			}
+		}
+		// Otherwise find the icon from the type information if available
+		else if (const FTypedElementClassTypeInfoColumn* TypeInfoColumn = DataStorage->GetColumn<FTypedElementClassTypeInfoColumn>(Row))
+		{
+			if (const UClass* Type = TypeInfoColumn->TypeInfo.Get())
+			{
+				const FName IconName = Type->GetFName();
+
+				if(const FSlateBrush* CachedBrush = FindCachedIcon(IconName))
+				{
+					return CachedBrush;
+				}
+				else if(const FSlateBrush* TypeBrush = FSlateIconFinder::FindIconBrushForClass(Type))
+				{
+					CachedIconMap.Add(IconName, TypeBrush);
+					return TypeBrush;
+				}
+			}
+		}
+		
+		// Fallback to the regular actor icon if we haven't found any specific icon
+		return FSlateIconFinder::FindIconForClass(AActor::StaticClass()).GetOptionalIcon();;
+
+	}
 }
 
 void UTypedElementTableViewerFactory::RegisterTables(ITypedElementDataStorageInterface& DataStorage)
diff --git a/Engine/Plugins/Experimental/EditorDataStorageFeatures/Source/TedsTableViewer/Private/Widgets/UObjectLabelWidget.cpp b/Engine/Plugins/Experimental/EditorDataStorageFeatures/Source/TedsTableViewer/Private/Widgets/UObjectLabelWidget.cpp
new file mode 100644
--- /dev/null
+++ b/Engine/Plugins/Experimental/EditorDataStorageFeatures/Source/TedsTableViewer/Private/Widgets/UObjectLabelWidget.cpp
@@ -0,0 +1,70 @@
+// Copyright Epic Games, Inc. All Rights Reserved.
+
+#include "Widgets/UObjectLabelWidget.h"
+
+#include "Elements/Columns/TypedElementCompatibilityColumns.h"
+#include "Elements/Columns/TypedElementLabelColumns.h"
+#include "Elements/Framework/TypedElementAttributeBinding.h"
+#include "TedsTableViewerUtils.h"
+#include "Widgets/Images/SImage.h"
+#include "Widgets/Layout/SSpacer.h"
+#include "Widgets/Text/STextBlock.h"
+#include "Widgets/SBoxPanel.h"
+
+
+#define LOCTEXT_NAMESPACE "FUObjectLabelWidgetConstructor"
+
+void UUObjectLabelWidgetFactory::RegisterWidgetConstructors(ITypedElementDataStorageInterface& DataStorage,
+	ITypedElementDataStorageUiInterface& DataStorageUi) const
+{
+	DataStorageUi.RegisterWidgetFactory<FUObjectLabelWidgetConstructor>(
+		TEXT("General.RowLabel"),
+		TypedElementDataStorage::FColumn<FTypedElementLabelColumn>() && TypedElementDataStorage::FColumn<FTypedElementUObjectColumn>());
+}
+
+FUObjectLabelWidgetConstructor::FUObjectLabelWidgetConstructor()
+	: FSimpleWidgetConstructor(StaticStruct())
+{
+}
+
+TSharedPtr<SWidget> FUObjectLabelWidgetConstructor::CreateWidget(ITypedElementDataStorageInterface* DataStorage,
+	ITypedElementDataStorageUiInterface* DataStorageUi, RowHandle TargetRow, RowHandle WidgetRow,
+	const TypedElementDataStorage::FMetaDataView& Arguments)
+{
+	if(DataStorage->IsRowAvailable(TargetRow))
+	{
+		UE::EditorDataStorage::FAttributeBinder Binder(TargetRow, DataStorage);
+
+		// Once TEDS UI has widget combining functionality, the binder can be used to create the type info widget and label widget and combine them
+		return SNew(SHorizontalBox)
+			+SHorizontalBox::Slot()
+				.AutoWidth()
+			[
+				SNew(SImage)
+					.Image(UE::Editor::DataStorage::TableViewerUtils::GetIconForRow(DataStorage, TargetRow))
+					.ColorAndOpacity(FSlateColor::UseForeground())
+			]
+			+SHorizontalBox::Slot()
+				.AutoWidth()
+			[
+				SNew(SSpacer)
+					.Size(FVector2D(5.0f, 0.0f))
+			]
+			+SHorizontalBox::Slot()
+				.FillWidth(1.0f)
+			[
+				SNew(STextBlock)
+					.Text(Binder.BindText(&FTypedElementLabelColumn::Label))
+					.ToolTipText(Binder.BindText(&FTypedElementLabelColumn::Label))
+			];
+	}
+	else
+	{
+		return SNew(STextBlock)
+			.Text(LOCTEXT("MissingRowReferenceColumn", "Unable to retrieve row reference."));
+	}
+	
+	
+}
+
+#undef LOCTEXT_NAMESPACE // "FUObjectLabelWidgetConstructor"
diff --git a/Engine/Plugins/Experimental/EditorDataStorageFeatures/Source/TedsTableViewer/Public/TedsTableViewerUtils.h b/Engine/Plugins/Experimental/EditorDataStorageFeatures/Source/TedsTableViewer/Public/TedsTableViewerUtils.h
--- a/Engine/Plugins/Experimental/EditorDataStorageFeatures/Source/TedsTableViewer/Public/TedsTableViewerUtils.h
+++ b/Engine/Plugins/Experimental/EditorDataStorageFeatures/Source/TedsTableViewer/Public/TedsTableViewerUtils.h
@@ -3,6 +3,7 @@
 #pragma once
 
 #include "Containers/ArrayView.h"
+#include "Elements/Common/TypedElementHandles.h"
 #include "Elements/Interfaces/TypedElementDataStorageFactory.h"
 #include "UObject/Class.h"
 
@@ -12,9 +13,11 @@ namespace TypedElementDataStorage
 {
 	class FMetaDataView;
 }
+
 class ITypedElementDataStorageInterface;
 class ITypedElementDataStorageUiInterface;
 struct FTypedElementWidgetConstructor;
+struct FSlateBrush;
 
 // Util library for functions shared by the Teds Table Viewer and the Teds Outliner
 namespace UE::Editor::DataStorage::TableViewerUtils
@@ -31,6 +34,9 @@ namespace UE::Editor::DataStorage::TableViewerUtils
 
 	// Create a copy of the provided column types array after discarding invalid entries
 	TEDSTABLEVIEWER_API TArray<TWeakObjectPtr<const UScriptStruct>> CreateVerifiedColumnTypeArray(TConstArrayView<TWeakObjectPtr<const UScriptStruct>> ColumnTypes);
+	
+	// Scans the row for columns and tries to find the best matching icon.
+	TEDSTABLEVIEWER_API const FSlateBrush* GetIconForRow(ITypedElementDataStorageInterface* DataStorage, RowHandle Row);
 } // namespace UE::Editor::DataStorage::TableViewerUtils
 
 UCLASS()
diff --git a/Engine/Plugins/Experimental/EditorDataStorageFeatures/Source/TedsTableViewer/Public/Widgets/UObjectLabelWidget.h b/Engine/Plugins/Experimental/EditorDataStorageFeatures/Source/TedsTableViewer/Public/Widgets/UObjectLabelWidget.h
new file mode 100644
--- /dev/null
+++ b/Engine/Plugins/Experimental/EditorDataStorageFeatures/Source/TedsTableViewer/Public/Widgets/UObjectLabelWidget.h
@@ -0,0 +1,45 @@
+// Copyright Epic Games, Inc. All Rights Reserved.
+
+#pragma once
+
+#include "Elements/Interfaces/TypedElementDataStorageFactory.h"
+#include "Elements/Interfaces/TypedElementDataStorageUiInterface.h"
+#include "Internationalization/Text.h"
+#include "UObject/ObjectMacros.h"
+
+#include "UObjectLabelWidget.generated.h"
+
+class ITypedElementDataStorageInterface;
+class UScriptStruct;
+
+UCLASS()
+class UUObjectLabelWidgetFactory : public UTypedElementDataStorageFactory
+{
+	GENERATED_BODY()
+
+public:
+	~UUObjectLabelWidgetFactory() override = default;
+
+	TEDSTABLEVIEWER_API void RegisterWidgetConstructors(ITypedElementDataStorageInterface& DataStorage,
+		ITypedElementDataStorageUiInterface& DataStorageUi) const override;
+};
+
+// Widget to show labels for UObjects in TEDS
+USTRUCT()
+struct FUObjectLabelWidgetConstructor : public FSimpleWidgetConstructor
+{
+	GENERATED_BODY()
+
+public:
+	TEDSTABLEVIEWER_API FUObjectLabelWidgetConstructor();
+	~FUObjectLabelWidgetConstructor() override = default;
+
+	TEDSTABLEVIEWER_API virtual TSharedPtr<SWidget> CreateWidget(
+		ITypedElementDataStorageInterface* DataStorage,
+		ITypedElementDataStorageUiInterface* DataStorageUi,
+		RowHandle TargetRow,
+		RowHandle WidgetRow, 
+		const TypedElementDataStorage::FMetaDataView& Arguments) override;
+
+protected:
+};
\ No newline at end of file
diff --git a/Engine/Source/Runtime/TypedElementFramework/Private/Elements/Interfaces/TypedElementDataStorageUiInterface.cpp b/Engine/Source/Runtime/TypedElementFramework/Private/Elements/Interfaces/TypedElementDataStorageUiInterface.cpp
--- a/Engine/Source/Runtime/TypedElementFramework/Private/Elements/Interfaces/TypedElementDataStorageUiInterface.cpp
+++ b/Engine/Source/Runtime/TypedElementFramework/Private/Elements/Interfaces/TypedElementDataStorageUiInterface.cpp
@@ -141,7 +141,9 @@ TSharedPtr<SWidget> FTypedElementWidgetConstructor::Construct(
 	ITypedElementDataStorageUiInterface* DataStorageUi,
 	const TypedElementDataStorage::FMetaDataView& Arguments)
 {
-	TSharedPtr<SWidget> Widget = CreateWidget(DataStorage, DataStorageUi, Row, Arguments);
+	TypedElementDataStorage::RowHandle TargetRow = GetTargetRow(DataStorage, Row);
+
+	TSharedPtr<SWidget> Widget = CreateWidget(DataStorage, DataStorageUi, TargetRow, Row, Arguments);
 	if (Widget)
 	{
 		DataStorage->GetColumn<FTypedElementSlateWidgetReferenceColumn>(Row)->Widget = Widget;
@@ -165,6 +167,7 @@ TSharedPtr<SWidget> FTypedElementWidgetConstructor::CreateWidget(const TypedElem
 TSharedPtr<SWidget> FTypedElementWidgetConstructor::CreateWidget(
 	ITypedElementDataStorageInterface* DataStorage,
 	ITypedElementDataStorageUiInterface* DataStorageUi,
+	TypedElementDataStorage::RowHandle TargetRow,
 	TypedElementDataStorage::RowHandle UiRow, 
 	const TypedElementDataStorage::FMetaDataView& Arguments)
 {
@@ -220,3 +223,69 @@ void FTypedElementWidgetConstructor::AddDefaultWidgetColumns(RowHandle Row, ITyp
 		}
 	}
 }
+
+FTypedElementWidgetConstructor::RowHandle FTypedElementWidgetConstructor::GetTargetRow(ITypedElementDataStorageInterface* DataStorage, RowHandle WidgetRow) const
+{
+	TypedElementDataStorage::RowHandle TargetRow = TypedElementDataStorage::InvalidRowHandle;
+
+	if(const FTypedElementRowReferenceColumn* RowReferenceColumn = DataStorage->GetColumn<FTypedElementRowReferenceColumn>(WidgetRow))
+	{
+		TargetRow = RowReferenceColumn->Row;
+	}
+
+	return TargetRow;
+}
+
+
+// FSimpleWidgetConstructor
+
+FSimpleWidgetConstructor::FSimpleWidgetConstructor(const UScriptStruct* InTypeInfo)
+	: FTypedElementWidgetConstructor(InTypeInfo)
+{
+}
+
+TSharedPtr<SWidget> FSimpleWidgetConstructor::CreateWidget(ITypedElementDataStorageInterface* DataStorage,
+	ITypedElementDataStorageUiInterface* DataStorageUi, TypedElementDataStorage::RowHandle TargetRow, TypedElementDataStorage::RowHandle WidgetRow,
+	const TypedElementDataStorage::FMetaDataView& Arguments)
+{
+	return nullptr;
+}
+
+bool FSimpleWidgetConstructor::SetColumns(ITypedElementDataStorageInterface* DataStorage, RowHandle Row)
+{
+	return FTypedElementWidgetConstructor::SetColumns(DataStorage, Row);
+}
+
+TSharedPtr<SWidget> FSimpleWidgetConstructor::CreateWidget(const TypedElementDataStorage::FMetaDataView& Arguments)
+{
+	// This function is not needed anymore and only exists so derived classes cannot derive from it anymore
+	return nullptr;
+}
+
+bool FSimpleWidgetConstructor::FinalizeWidget(ITypedElementDataStorageInterface* DataStorage, ITypedElementDataStorageUiInterface* DataStorageUi,
+	RowHandle Row, const TSharedPtr<SWidget>& Widget)
+{
+	// This function is not needed anymore and only exists so derived classes cannot derive from it anymore
+	return true;
+}
+
+TSharedPtr<SWidget> FSimpleWidgetConstructor::Construct(RowHandle WidgetRow, ITypedElementDataStorageInterface* DataStorage,
+	ITypedElementDataStorageUiInterface* DataStorageUi, const TypedElementDataStorage::FMetaDataView& Arguments)
+{
+	const TypedElementDataStorage::RowHandle TargetRow = GetTargetRow(DataStorage, WidgetRow);
+
+	// Set any required columns on the widget row first
+	SetColumns(DataStorage, WidgetRow);
+
+	// Create the actual widget
+	TSharedPtr<SWidget> Widget = CreateWidget(DataStorage, DataStorageUi, TargetRow, WidgetRow, Arguments);
+
+	// If the widget was created, add the default columns we want all widget rows to have (e.g Label)
+	if (Widget)
+	{
+		AddDefaultWidgetColumns(WidgetRow, DataStorage);
+	}
+	
+	return Widget;
+}
+
diff --git a/Engine/Source/Runtime/TypedElementFramework/Public/Elements/Framework/TypedElementAttributeBinding.h b/Engine/Source/Runtime/TypedElementFramework/Public/Elements/Framework/TypedElementAttributeBinding.h
--- a/Engine/Source/Runtime/TypedElementFramework/Public/Elements/Framework/TypedElementAttributeBinding.h
+++ b/Engine/Source/Runtime/TypedElementFramework/Public/Elements/Framework/TypedElementAttributeBinding.h
@@ -89,6 +89,22 @@ namespace UE::EditorDataStorage
 		 */
 		template <typename InRetValType, typename... ParamTypes, typename ColumnType>
 		TDelegate<InRetValType(ParamTypes...)> BindEvent(TDelegate<InRetValType(ParamTypes...)> ColumnType::* InVariable);
+
+		/**
+		 * Directly bind an FString member in a Teds column to an FText attribute as a shortcut
+		 * @param InFStringVariable The FString variable 
+		 * @return A delegate that can be provided to a text widget in Slate (e.g STextBlock)
+		 */
+		template <TDataColumnType ColumnType>
+		TAttribute<FText> BindText(FString ColumnType::* InFStringVariable);
+
+		/**
+		 * Directly bind an FName member in a Teds column to an FText attribute as a shortcut
+		 * @param InFNameVariable The FName variable 
+		 * @return A delegate that can be provided to a text widget in Slate (e.g STextBlock)
+		 */
+		template <TDataColumnType ColumnType>
+		TAttribute<FText> BindText(FName ColumnType::* InFNameVariable);
 		
 	private:
 
diff --git a/Engine/Source/Runtime/TypedElementFramework/Public/Elements/Framework/TypedElementAttributeBinding.inl b/Engine/Source/Runtime/TypedElementFramework/Public/Elements/Framework/TypedElementAttributeBinding.inl
--- a/Engine/Source/Runtime/TypedElementFramework/Public/Elements/Framework/TypedElementAttributeBinding.inl
+++ b/Engine/Source/Runtime/TypedElementFramework/Public/Elements/Framework/TypedElementAttributeBinding.inl
@@ -202,6 +202,24 @@ namespace UE::EditorDataStorage
 		});
 	}
 	
+	template <TDataColumnType ColumnType>
+	TAttribute<FText> FAttributeBinder::BindText(FString ColumnType::* InFStringVariable)
+	{
+		return BindData(InFStringVariable, [](const FString& InString)
+		{
+			return FText::FromString(InString);
+		});
+	}
+
+	template <TDataColumnType ColumnType>
+	TAttribute<FText> FAttributeBinder::BindText(FName ColumnType::* InFNameVariable)
+	{
+		return BindData(InFNameVariable, [](const FName& InName)
+		{
+			return FText::FromName(InName);
+		});
+	}
+	
 }
 
 	
\ No newline at end of file
diff --git a/Engine/Source/Runtime/TypedElementFramework/Public/Elements/Interfaces/TypedElementDataStorageUiInterface.h b/Engine/Source/Runtime/TypedElementFramework/Public/Elements/Interfaces/TypedElementDataStorageUiInterface.h
--- a/Engine/Source/Runtime/TypedElementFramework/Public/Elements/Interfaces/TypedElementDataStorageUiInterface.h
+++ b/Engine/Source/Runtime/TypedElementFramework/Public/Elements/Interfaces/TypedElementDataStorageUiInterface.h
@@ -19,6 +19,7 @@ class SWidget;
 /**
  * Base class used to construct Typed Element widgets with.
  * See below for the options to register a constructor with the Data Storage.
+ * In most cases you want to inherit from FSimpleWidgetConstructor instead which has a simpler pipeline to create widgets
  */
 USTRUCT()
 struct FTypedElementWidgetConstructor
@@ -90,7 +91,8 @@ protected:
 	TYPEDELEMENTFRAMEWORK_API virtual TSharedPtr<SWidget> CreateWidget(
 		ITypedElementDataStorageInterface* DataStorage,
 		ITypedElementDataStorageUiInterface* DataStorageUi,
-		TypedElementDataStorage::RowHandle UiRow, 
+		TypedElementDataStorage::RowHandle TargetRow,
+		TypedElementDataStorage::RowHandle WidgetRow, 
 		const TypedElementDataStorage::FMetaDataView& Arguments);
 	/** Set any values in columns if needed. The columns provided through GetAdditionalColumnsList() will have already been created. */
 	TYPEDELEMENTFRAMEWORK_API virtual bool SetColumns(ITypedElementDataStorageInterface* DataStorage, RowHandle Row);
@@ -110,6 +112,13 @@ protected:
 
 	/** Add the default misc columns we want a widget row to have. */
 	TYPEDELEMENTFRAMEWORK_API void AddDefaultWidgetColumns(RowHandle Row, ITypedElementDataStorageInterface* DataStorage) const;
+
+	/**
+	 * Helper function to get the actual target row with the data the widget is operating on (if applicable). Returns InvalidRowHandle if there is no
+	 * target row
+	 */
+	TYPEDELEMENTFRAMEWORK_API RowHandle GetTargetRow(ITypedElementDataStorageInterface* DataStorage, RowHandle WidgetRow) const;
+
 protected:
 
 	TArray<TWeakObjectPtr<const UScriptStruct>> MatchedColumnTypes;
@@ -117,6 +126,72 @@ protected:
 	const UScriptStruct* TypeInfo = nullptr;
 };
 
+/**
+ * A simple widget constructor that cuts down on most of the boilerplate, in most cases you want to inherit from this to create your widget constructor
+ * Only requires you to override CreateWidget() to create the actual SWidget
+ */
+USTRUCT()
+struct FSimpleWidgetConstructor : public FTypedElementWidgetConstructor
+{
+	GENERATED_BODY()
+
+	/** Call this constructor with StaticStruct() on your derived class to pass in the type information */
+	TYPEDELEMENTFRAMEWORK_API explicit FSimpleWidgetConstructor(const UScriptStruct* InTypeInfo);
+
+	FSimpleWidgetConstructor() : Super(StaticStruct()) {} //< For compatibility and shouldn't be directly used.
+	
+	virtual ~FSimpleWidgetConstructor() override = default;
+
+	/*
+	 * Required function to create the actual widget instance.
+	 * 
+	 * @param DataStorage A pointer to the TEDS data storage interface
+	 * @param DataStorageUi A pointer to the TEDS data storage UI interface
+	 * @param TargetRow The row for the actual data this widget is being created for (can be InvalidRowHandle if there is no target row attached to this widget)
+	 * @param WidgetRow The row that contains information about the widget itself
+	 * @param Arguments Any metadata arguments that were specified
+	 * @return The actual widget instance
+	 */
+	TYPEDELEMENTFRAMEWORK_API virtual TSharedPtr<SWidget> CreateWidget(
+		ITypedElementDataStorageInterface* DataStorage,
+		ITypedElementDataStorageUiInterface* DataStorageUi,
+		RowHandle TargetRow,
+		RowHandle WidgetRow, 
+		const TypedElementDataStorage::FMetaDataView& Arguments) override;
+
+	/*
+	 * Override this function to add any columns to the WidgetRow before CreateWidget is called
+	 * 
+	 * @param DataStorage A pointer to the TEDS data storage interface
+	 * @param WidgetRow The row that contains information about the widget itself
+	 * @return Whether any columns were added
+	 */
+	TYPEDELEMENTFRAMEWORK_API virtual bool SetColumns(ITypedElementDataStorageInterface* DataStorage, RowHandle WidgetRow) override;
+
+protected:
+
+	/** Old CreateWidget overload that exists for backwards compatibility, you should use the overload that provides the row instead */
+	TYPEDELEMENTFRAMEWORK_API virtual TSharedPtr<SWidget> CreateWidget(const TypedElementDataStorage::FMetaDataView& Arguments) override final;
+	
+	/** 
+	 * Old function in the widget creation pipeline that isn't used anymore. All your logic should go in CreateWidget() itself
+	 */
+	TYPEDELEMENTFRAMEWORK_API virtual bool FinalizeWidget(
+		ITypedElementDataStorageInterface* DataStorage,
+		ITypedElementDataStorageUiInterface* DataStorageUi,
+		RowHandle Row,
+		const TSharedPtr<SWidget>& Widget) override final;
+	
+	/**
+	 * Helper function to call SetColumns() and CreateWidget(), you should not need to override this for simple widget constructors
+	 */
+	TYPEDELEMENTFRAMEWORK_API virtual TSharedPtr<SWidget> Construct(
+		RowHandle WidgetRow, 
+		ITypedElementDataStorageInterface* DataStorage,
+		ITypedElementDataStorageUiInterface* DataStorageUi,
+		const TypedElementDataStorage::FMetaDataView& Arguments) override final;
+};
+
 template<>
 struct TStructOpsTypeTraits<FTypedElementWidgetConstructor> : public TStructOpsTypeTraitsBase2<FTypedElementWidgetConstructor>
 {
