[TEDS] ValueTag rename

* Renames DynamicTags to ValueTags everywhere to avoid confusion with the Dynamic Column feature

diff --git a/Engine/Plugins/Experimental/EditorDataStorage/Source/TedsCore/Private/Debug/TypedElementDatabaseDebugCommands.cpp b/Engine/Plugins/Experimental/EditorDataStorage/Source/TedsCore/Private/Debug/TypedElementDatabaseDebugCommands.cpp
--- a/Engine/Plugins/Experimental/EditorDataStorage/Source/TedsCore/Private/Debug/TypedElementDatabaseDebugCommands.cpp
+++ b/Engine/Plugins/Experimental/EditorDataStorage/Source/TedsCore/Private/Debug/TypedElementDatabaseDebugCommands.cpp
@@ -578,8 +578,8 @@ static FAutoConsoleCommand CVarActivateListDynamicColumnQuery(
 	}),
 	ECVF_Default);
 
-static FAutoConsoleCommand CVarAddDynamicTag(
-	TEXT("TEDS.Debug.DynamicTag.AddColumn"),
+static FAutoConsoleCommand CVarAddValueTag(
+	TEXT("TEDS.Debug.ValueTag.AddColumn"),
 	TEXT("Argument: Row, Tag, Value\n"),
 	FConsoleCommandWithArgsDelegate::CreateLambda([](const TArray<FString>& Args)
 	{
@@ -601,11 +601,11 @@ static FAutoConsoleCommand CVarAddDynamicTag(
 		if constexpr (bUseTemplateSugar)
 		{
 			const FName Tag = *Args[1];
-			DataStorage->AddColumn<FDynamicTag>(Row, Tag, Value); 
+			DataStorage->AddColumn<FValueTag>(Row, Tag, Value); 
 		}
 		else
 		{
-			const FDynamicTag Tag(*Args[1]);
+			const FValueTag Tag(*Args[1]);
 			DataStorage->AddColumn(Row, Tag, Value);
 		}
 		
@@ -613,8 +613,8 @@ static FAutoConsoleCommand CVarAddDynamicTag(
 	}),
 	ECVF_Default);
 
-static FAutoConsoleCommand CVarRemoveDynamicTag(
-	TEXT("TEDS.Debug.DynamicTag.RemoveColumn"),
+static FAutoConsoleCommand CVarRemoveValueTag(
+	TEXT("TEDS.Debug.ValueTag.RemoveColumn"),
 	TEXT("Argument: Row, Group\n"),
 	FConsoleCommandWithArgsDelegate::CreateLambda([](const TArray<FString>& Args)
 	{
@@ -633,18 +633,18 @@ static FAutoConsoleCommand CVarRemoveDynamicTag(
 		{
 			using namespace UE::Editor::DataStorage;
 			const FName Tag = *Args[1];
-			DataStorage->RemoveColumn<FDynamicTag>(Row, Tag);
+			DataStorage->RemoveColumn<FValueTag>(Row, Tag);
 		}
 		else
 		{
-			const UE::Editor::DataStorage::FDynamicTag Tag(*Args[1]);
+			const UE::Editor::DataStorage::FValueTag Tag(*Args[1]);
 			DataStorage->RemoveColumn(Row, Tag);
 		}		
 	}),
 	ECVF_Default);
 	
-static FAutoConsoleCommand CVarMatchDynamicTag(
-	TEXT("TEDS.Debug.DynamicTag.RunQuery"),
+static FAutoConsoleCommand CVarMatchValueTag(
+	TEXT("TEDS.Debug.ValueTag.RunQuery"),
 	TEXT("Argument: Tag, [optional] Value\n"),
 	FConsoleCommandWithArgsDelegate::CreateLambda([](const TArray<FString>& Args)
 	{
@@ -668,8 +668,8 @@ static FAutoConsoleCommand CVarMatchDynamicTag(
 				return DataStorage->RegisterQuery(
 					Select().
 						Where().
-							// Match all rows with a dynamic tag of type Tag (ie. all rows with a dynamic tag of "Color")
-							All<FDynamicTag>(Tag).
+							// Match all rows with a value tag of type Tag (ie. all rows with a value tag of "Color")
+							All<FValueTag>(Tag).
 							All<FTestColumnA>().
 						Compile());
 			}
@@ -679,8 +679,8 @@ static FAutoConsoleCommand CVarMatchDynamicTag(
 				return DataStorage->RegisterQuery(
 					Select().
 					Where().
-						// Match all rows with a dynamic tag of type Tag that has a MatchValue (ie. all rows with dynamic tag "Color" with value "Red")
-						All<FDynamicTag>(Tag, MatchValue).
+						// Match all rows with a value tag of type Tag that has a MatchValue (ie. all rows with value tag "Color" with value "Red")
+						All<FValueTag>(Tag, MatchValue).
 						All<FTestColumnA>().
 					Compile());
 			}
@@ -699,8 +699,8 @@ static FAutoConsoleCommand CVarMatchDynamicTag(
 	}),
 	ECVF_Default);
 
-static FAutoConsoleCommand CVarAddDynamicTagFromEnum(
-	TEXT("TEDS.Debug.DynamicTag.AddWithEnum"),
+static FAutoConsoleCommand CVarAddValueTagFromEnum(
+	TEXT("TEDS.Debug.ValueTag.AddWithEnum"),
 	TEXT("Argument: Row, EnumValue\n"),
 	FConsoleCommandWithArgsDelegate::CreateLambda([](const TArray<FString>& Args)
 	{
@@ -738,8 +738,8 @@ static FAutoConsoleCommand CVarAddDynamicTagFromEnum(
 	}),
 	ECVF_Default);
 
-static FAutoConsoleCommand CVarRemoveDynamicTagFromEnum(
-	TEXT("TEDS.Debug.DynamicTag.RemoveWithEnum"),
+static FAutoConsoleCommand CVarRemoveValueTagFromEnum(
+	TEXT("TEDS.Debug.ValueTag.RemoveWithEnum"),
 	TEXT("Argument: Row\n"),
 	FConsoleCommandWithArgsDelegate::CreateLambda([](const TArray<FString>& Args)
 	{
@@ -759,8 +759,8 @@ static FAutoConsoleCommand CVarRemoveDynamicTagFromEnum(
 	ECVF_Default);
 
 
-static FAutoConsoleCommand CVarMatchDynamicTagFromEnum(
-	TEXT("TEDS.Debug.DynamicTag.RunQueryEnum"),
+static FAutoConsoleCommand CVarMatchValueTagFromEnum(
+	TEXT("TEDS.Debug.ValueTag.RunQueryEnum"),
 	TEXT("Argument: [optional] EnumValue\n"),
 	FConsoleCommandWithArgsDelegate::CreateLambda([](const TArray<FString>& Args)
 	{
@@ -795,7 +795,7 @@ static FAutoConsoleCommand CVarMatchDynamicTagFromEnum(
 				return DataStorage->RegisterQuery(
 					Select().
 						Where().
-							// Match all rows with an enum dynamic tag of the hardcoded enum type
+							// Match all rows with an enum value tag of the hardcoded enum type
 							All<ETedsDebugEnum>().
 						Compile());
 			}
@@ -806,7 +806,7 @@ static FAutoConsoleCommand CVarMatchDynamicTagFromEnum(
 				return DataStorage->RegisterQuery(
 					Select().
 					Where().
-						// Match all rows with a dynamic tag of the hardcoded enum type that has the given value
+						// Match all rows with a value tag of the hardcoded enum type that has the given value
 						// Note, usually this would be written something like:
 						//   All(ETedsDebugEnum::Red).
 						// However it isn't possible to do that when getting the enum value from a string.  API is still exercised
diff --git a/Engine/Plugins/Experimental/EditorDataStorage/Source/TedsCore/Private/DynamicColumnGenerator.cpp b/Engine/Plugins/Experimental/EditorDataStorage/Source/TedsCore/Private/DynamicColumnGenerator.cpp
--- a/Engine/Plugins/Experimental/EditorDataStorage/Source/TedsCore/Private/DynamicColumnGenerator.cpp
+++ b/Engine/Plugins/Experimental/EditorDataStorage/Source/TedsCore/Private/DynamicColumnGenerator.cpp
@@ -114,20 +114,20 @@ namespace UE::Editor::DataStorage
 		return nullptr;
 	}
 	
-	FDynamicTagManager::FDynamicTagManager(FDynamicColumnGenerator& InColumnGenerator)
+	FValueTagManager::FValueTagManager(FDynamicColumnGenerator& InColumnGenerator)
 		: ColumnGenerator(InColumnGenerator)
 	{
 	}
 	
-	FConstSharedStruct FDynamicTagManager::GenerateDynamicTag(const FDynamicTag& InTag, const FName& InValue)
+	FConstSharedStruct FValueTagManager::GenerateValueTag(const FValueTag& InTag, const FName& InValue)
 	{
-		TPair<FDynamicTag, FName> Pair(InTag, InValue);
+		TPair<FValueTag, FName> Pair(InTag, InValue);
 	
 		UE_MT_SCOPED_WRITE_ACCESS(AccessDetector);
 	
 		// Common path
 		{
-			if (FConstSharedStruct* TagStruct = DynamicTagLookup.Find(Pair))
+			if (FConstSharedStruct* TagStruct = ValueTagLookup.Find(Pair))
 			{
 				return *TagStruct;
 			}
@@ -144,13 +144,13 @@ namespace UE::Editor::DataStorage
 	
 			FConstSharedStruct SharedStruct = FConstSharedStruct::Make(ColumnType, reinterpret_cast<const uint8*>(&Overlay));
 			
-			DynamicTagLookup.Emplace(Pair, SharedStruct);
+			ValueTagLookup.Emplace(Pair, SharedStruct);
 	
 			return SharedStruct;
 		}
 	}
 	
-	const UScriptStruct* FDynamicTagManager::GenerateColumnType(const FDynamicTag& Tag)
+	const UScriptStruct* FValueTagManager::GenerateColumnType(const FValueTag& Tag)
 	{
 		const FDynamicColumnGeneratorInfo GeneratedColumnType = ColumnGenerator.GenerateColumn(*FTedsValueTagColumn::StaticStruct(), Tag.GetName());
 		
diff --git a/Engine/Plugins/Experimental/EditorDataStorage/Source/TedsCore/Private/DynamicColumnGenerator.h b/Engine/Plugins/Experimental/EditorDataStorage/Source/TedsCore/Private/DynamicColumnGenerator.h
--- a/Engine/Plugins/Experimental/EditorDataStorage/Source/TedsCore/Private/DynamicColumnGenerator.h
+++ b/Engine/Plugins/Experimental/EditorDataStorage/Source/TedsCore/Private/DynamicColumnGenerator.h
@@ -15,8 +15,7 @@
 class FName;
 class UScriptStruct;
 
-// The template struct that is used to generate the DynamicTag column
-// It is safe to reinterpret a DynamicTag column to this template to access the Value
+// The template struct that is used to generate the ValueTag column
 USTRUCT()
 struct FTedsValueTagColumn : public FTedsSharedColumn
 {
@@ -89,22 +88,17 @@ namespace UE::Editor::DataStorage
 		TMap<FGeneratedColumnKey, int32> GeneratedColumnLookup;
 	};
 
-	class FDynamicTagManager
+	class FValueTagManager
 	{
 	public:
-		struct FDynamicTagStructLayout
-		{
-			FName Tag;
-		};
-		
-		explicit FDynamicTagManager(FDynamicColumnGenerator& InColumnGenerator);
-		FConstSharedStruct GenerateDynamicTag(const FDynamicTag& InTag, const FName& InValue);
-		const UScriptStruct* GenerateColumnType(const FDynamicTag& InTag);
+		explicit FValueTagManager(FDynamicColumnGenerator& InColumnGenerator);
+		FConstSharedStruct GenerateValueTag(const FValueTag& InTag, const FName& InValue);
+		const UScriptStruct* GenerateColumnType(const FValueTag& InTag);
 	private:
 
 		UE_MT_DECLARE_RW_ACCESS_DETECTOR(AccessDetector);
 
-		TMap<TPair<FDynamicTag, FName>, FConstSharedStruct> DynamicTagLookup;
+		TMap<TPair<FValueTag, FName>, FConstSharedStruct> ValueTagLookup;
 
 		FDynamicColumnGenerator& ColumnGenerator;
 	};
diff --git a/Engine/Plugins/Experimental/EditorDataStorage/Source/TedsCore/Private/Queries/TypedElementExtendedQueryStore.cpp b/Engine/Plugins/Experimental/EditorDataStorage/Source/TedsCore/Private/Queries/TypedElementExtendedQueryStore.cpp
--- a/Engine/Plugins/Experimental/EditorDataStorage/Source/TedsCore/Private/Queries/TypedElementExtendedQueryStore.cpp
+++ b/Engine/Plugins/Experimental/EditorDataStorage/Source/TedsCore/Private/Queries/TypedElementExtendedQueryStore.cpp
@@ -640,12 +640,12 @@ namespace UE::Editor::DataStorage
 			return true;
 		}
 
-		Algo::SortBy(Query.ValueTags, [](const TypedElementDataStorage::FQueryDescription::FDynamicTagData& DynamicTagData)
+		Algo::SortBy(Query.ValueTags, [](const TypedElementDataStorage::FQueryDescription::FValueTagData& ValueTagData)
 			{
-				return DynamicTagData.Tag.GetName();
+				return ValueTagData.Tag.GetName();
 			}, FNameFastLess());
 		// Check if there are any duplicate groups. Not yet supported until we can match multiple MatchTags
-		UE::Editor::DataStorage::FDynamicTag PreviousTag = Query.ValueTags[0].Tag;
+		UE::Editor::DataStorage::FValueTag PreviousTag = Query.ValueTags[0].Tag;
 		for (int32 Index = 1, End = Query.ValueTags.Num(); Index < End; ++Index)
 		{
 			if (Query.ValueTags[Index].Tag == PreviousTag)
diff --git a/Engine/Plugins/Experimental/EditorDataStorage/Source/TedsCore/Private/TypedElementDatabase.cpp b/Engine/Plugins/Experimental/EditorDataStorage/Source/TedsCore/Private/TypedElementDatabase.cpp
--- a/Engine/Plugins/Experimental/EditorDataStorage/Source/TedsCore/Private/TypedElementDatabase.cpp
+++ b/Engine/Plugins/Experimental/EditorDataStorage/Source/TedsCore/Private/TypedElementDatabase.cpp
@@ -605,11 +605,11 @@ void UEditorDataStorage::AddColumns(RowHandle Row, TConstArrayView<const UScript
 	}
 }
 
-void UEditorDataStorage::AddColumn(RowHandle Row, const UE::Editor::DataStorage::FDynamicTag& Tag, const FName& InValue)
+void UEditorDataStorage::AddColumn(RowHandle Row, const UE::Editor::DataStorage::FValueTag& Tag, const FName& InValue)
 {
 	if (ActiveEditorEntityManager)
 	{
-		const FConstSharedStruct SharedStruct = Environment->GenerateDynamicTag(Tag, InValue);
+		const FConstSharedStruct SharedStruct = Environment->GenerateValueTag(Tag, InValue);
 
 		FMassEntityHandle Entity = FMassEntityHandle::FromNumber(Row);
 		if (ActiveEditorEntityManager->IsEntityActive(Entity))
@@ -619,15 +619,15 @@ void UEditorDataStorage::AddColumn(RowHandle Row, const UE::Editor::DataStorage:
 	}
 }
 
-void UEditorDataStorage::RemoveColumn(RowHandle Row, const UE::Editor::DataStorage::FDynamicTag& Tag)
+void UEditorDataStorage::RemoveColumn(RowHandle Row, const UE::Editor::DataStorage::FValueTag& Tag)
 {
 	if (ActiveEditorEntityManager)
 	{
-		const UScriptStruct* DynamicTagType = Environment->GenerateColumnType(Tag);
+		const UScriptStruct* ValueTagType = Environment->GenerateColumnType(Tag);
 		FMassEntityHandle Entity = FMassEntityHandle::FromNumber(Row);
 		if (ActiveEditorEntityManager->IsEntityActive(Entity))
 		{
-			UE::Editor::DataStorage::Legacy::FCommandBuffer::Execute_RemoveSharedColumnCommand(*ActiveEditorEntityManager, Row, *DynamicTagType);
+			UE::Editor::DataStorage::Legacy::FCommandBuffer::Execute_RemoveSharedColumnCommand(*ActiveEditorEntityManager, Row, *ValueTagType);
 		}
 	}
 }
diff --git a/Engine/Plugins/Experimental/EditorDataStorage/Source/TedsCore/Private/TypedElementDatabase.h b/Engine/Plugins/Experimental/EditorDataStorage/Source/TedsCore/Private/TypedElementDatabase.h
--- a/Engine/Plugins/Experimental/EditorDataStorage/Source/TedsCore/Private/TypedElementDatabase.h
+++ b/Engine/Plugins/Experimental/EditorDataStorage/Source/TedsCore/Private/TypedElementDatabase.h
@@ -116,12 +116,12 @@ public:
 	virtual bool IsRowAssigned(RowHandle Row) const override;
 
 	virtual void AddColumn(RowHandle Row, const UScriptStruct* ColumnType) override;
-	virtual void AddColumn(RowHandle Row, const UE::Editor::DataStorage::FDynamicTag& Tag, const FName& InValue) override;
+	virtual void AddColumn(RowHandle Row, const UE::Editor::DataStorage::FValueTag& Tag, const FName& InValue) override;
 	virtual void AddColumnData(RowHandle Row, const UScriptStruct* ColumnType,
 		const ColumnCreationCallbackRef& Initializer,
 		ColumnCopyOrMoveCallback Relocator) override;
 	virtual void RemoveColumn(RowHandle Row, const UScriptStruct* ColumnType) override;
-	virtual void RemoveColumn(RowHandle Row, const UE::Editor::DataStorage::FDynamicTag& Tag) override;
+	virtual void RemoveColumn(RowHandle Row, const UE::Editor::DataStorage::FValueTag& Tag) override;
 	virtual void* GetColumnData(RowHandle Row, const UScriptStruct* ColumnType) override;
 	virtual const void* GetColumnData(RowHandle Row, const UScriptStruct* ColumnType) const override;
 	virtual void AddColumns(RowHandle Row, TConstArrayView<const UScriptStruct*> Columns) override;
diff --git a/Engine/Plugins/Experimental/EditorDataStorage/Source/TedsCore/Private/TypedElementDatabaseEnvironment.cpp b/Engine/Plugins/Experimental/EditorDataStorage/Source/TedsCore/Private/TypedElementDatabaseEnvironment.cpp
--- a/Engine/Plugins/Experimental/EditorDataStorage/Source/TedsCore/Private/TypedElementDatabaseEnvironment.cpp
+++ b/Engine/Plugins/Experimental/EditorDataStorage/Source/TedsCore/Private/TypedElementDatabaseEnvironment.cpp
@@ -10,7 +10,7 @@ namespace UE::Editor::DataStorage
 		: DataStorage(InDataStorage)
 		, DirectDeferredCommands(*this)
 		, MementoSystem(InDataStorage)
-		, DynamicTagManager(DynamicColumnGenerator)
+		, ValueTagManager(DynamicColumnGenerator)
 		, MassEntityManager(InMassEntityManager)
 		, MassPhaseManager(InMassPhaseManager)
 	{
@@ -101,14 +101,14 @@ namespace UE::Editor::DataStorage
 		return DynamicColumnGenerator.GenerateColumn(Template, Identifier).Type;
 	}
 
-	FConstSharedStruct FEnvironment::GenerateDynamicTag(const FDynamicTag& Tag, const FName& Value)
+	FConstSharedStruct FEnvironment::GenerateValueTag(const FValueTag& Tag, const FName& Value)
 	{
-		return DynamicTagManager.GenerateDynamicTag(Tag, Value);
+		return ValueTagManager.GenerateValueTag(Tag, Value);
 	}
 
-	const UScriptStruct* FEnvironment::GenerateColumnType(const FDynamicTag& Tag)
+	const UScriptStruct* FEnvironment::GenerateColumnType(const FValueTag& Tag)
 	{
-		return DynamicTagManager.GenerateColumnType(Tag);
+		return ValueTagManager.GenerateColumnType(Tag);
 	}
 
 	void FEnvironment::NextUpdateCycle()
diff --git a/Engine/Plugins/Experimental/EditorDataStorage/Source/TedsCore/Private/TypedElementDatabaseEnvironment.h b/Engine/Plugins/Experimental/EditorDataStorage/Source/TedsCore/Private/TypedElementDatabaseEnvironment.h
--- a/Engine/Plugins/Experimental/EditorDataStorage/Source/TedsCore/Private/TypedElementDatabaseEnvironment.h
+++ b/Engine/Plugins/Experimental/EditorDataStorage/Source/TedsCore/Private/TypedElementDatabaseEnvironment.h
@@ -52,10 +52,10 @@ namespace UE::Editor::DataStorage
 		const UScriptStruct* GenerateDynamicColumn(const UScriptStruct& Template, const FName& Identifier);
 		
 		// Creates or Finds the column type associated with the value tag
-		const UScriptStruct* GenerateColumnType(const FDynamicTag& Tag);
+		const UScriptStruct* GenerateColumnType(const FValueTag& Tag);
 		
 		// Creates an instance of a value tag
-		FConstSharedStruct GenerateDynamicTag(const FDynamicTag& Tag, const FName& Value);
+		FConstSharedStruct GenerateValueTag(const FValueTag& Tag, const FName& Value);
 
 		void NextUpdateCycle();
 		uint64 GetUpdateCycleId() const;
@@ -68,7 +68,7 @@ namespace UE::Editor::DataStorage
 		FExtendedQueryStore Queries;
 		FMementoSystem MementoSystem;
 		FDynamicColumnGenerator DynamicColumnGenerator;
-		FDynamicTagManager DynamicTagManager;
+		FValueTagManager ValueTagManager;
 
 		FMassEntityManager& MassEntityManager;
 		FMassProcessingPhaseManager& MassPhaseManager;
diff --git a/Engine/Source/Runtime/TypedElementFramework/Private/Elements/Common/TypedElementCommonTypes.cpp b/Engine/Source/Runtime/TypedElementFramework/Private/Elements/Common/TypedElementCommonTypes.cpp
--- a/Engine/Source/Runtime/TypedElementFramework/Private/Elements/Common/TypedElementCommonTypes.cpp
+++ b/Engine/Source/Runtime/TypedElementFramework/Private/Elements/Common/TypedElementCommonTypes.cpp
@@ -6,16 +6,16 @@ namespace UE
 {
 	namespace Editor::DataStorage
 	{
-		const FName& FDynamicTag::GetName() const
+		const FName& FValueTag::GetName() const
 		{
 			return Name;
 		}
 		
-		FDynamicTag::FDynamicTag(const FName& InTypeName)
+		FValueTag::FValueTag(const FName& InTypeName)
 			: Name(InTypeName)
 		{}
 		
-		uint32 GetTypeHash(const FDynamicTag& InName)
+		uint32 GetTypeHash(const FValueTag& InName)
 		{
 			return GetTypeHash(InName.Name);
 		}
diff --git a/Engine/Source/Runtime/TypedElementFramework/Private/Elements/Framework/TypedElementQueryBuilder.cpp b/Engine/Source/Runtime/TypedElementFramework/Private/Elements/Framework/TypedElementQueryBuilder.cpp
--- a/Engine/Source/Runtime/TypedElementFramework/Private/Elements/Framework/TypedElementQueryBuilder.cpp
+++ b/Engine/Source/Runtime/TypedElementFramework/Private/Elements/Framework/TypedElementQueryBuilder.cpp
@@ -143,10 +143,10 @@ namespace TypedElementQueryBuilder
 		return *this;
 	}
 
-	FSimpleQuery& FSimpleQuery::All(const UE::Editor::DataStorage::FDynamicTag& Tag, const FName& Value)
+	FSimpleQuery& FSimpleQuery::All(const UE::Editor::DataStorage::FValueTag& Tag, const FName& Value)
 	{
 		Query->ValueTags.Emplace(
-			TypedElementDataStorage::FQueryDescription::FDynamicTagData
+			TypedElementDataStorage::FQueryDescription::FValueTagData
 			{
 				.Tag = Tag,
 				.MatchValue = Value
@@ -157,7 +157,7 @@ namespace TypedElementQueryBuilder
 	FSimpleQuery& FSimpleQuery::All(const UEnum& Enum)
 	{
 		using namespace UE::Editor::DataStorage;
-		return All(FDynamicTag(Enum.GetFName()));
+		return All(FValueTag(Enum.GetFName()));
 	}
 
 	FSimpleQuery& FSimpleQuery::All(const UEnum& Enum, int64 Value)
@@ -169,7 +169,7 @@ namespace TypedElementQueryBuilder
 			return *this;
 		}
 		using namespace UE::Editor::DataStorage;
-		return All(FDynamicTag(Enum.GetFName()), ValueName);
+		return All(FValueTag(Enum.GetFName()), ValueName);
 	}
 	
 	FSimpleQuery& FSimpleQuery::All(const UE::Editor::DataStorage::FDynamicColumnDescription& Description)
@@ -179,7 +179,7 @@ namespace TypedElementQueryBuilder
 		return *this;
 	}
 
-	FSimpleQuery& FSimpleQuery::All(const UE::Editor::DataStorage::FDynamicTag& Tag)
+	FSimpleQuery& FSimpleQuery::All(const UE::Editor::DataStorage::FValueTag& Tag)
 	{
 		return All(Tag, NAME_None);
 	}
diff --git a/Engine/Source/Runtime/TypedElementFramework/Public/Elements/Common/TypedElementCommonTypes.h b/Engine/Source/Runtime/TypedElementFramework/Public/Elements/Common/TypedElementCommonTypes.h
--- a/Engine/Source/Runtime/TypedElementFramework/Public/Elements/Common/TypedElementCommonTypes.h
+++ b/Engine/Source/Runtime/TypedElementFramework/Public/Elements/Common/TypedElementCommonTypes.h
@@ -47,26 +47,26 @@ namespace UE
 		using FTag = FEditorDataStorageTag;
 
 		/**
-		 * Defines a dynamic type for a dynamic tag
+		 * Defines a dynamic type for a value tag
 		 * Example:
-		 *   FDynamicTag ColorTagType(TEXT("Color"));
-		 *   FDynamicTag DirectionTagType(TEXT("Direction"));
-		 * A dynamic tag can take on different values for each type.  This is set up when a tag is added to a row.
+		 *   FValueTag ColorTagType(TEXT("Color"));
+		 *   FValueTag DirectionTagType(TEXT("Direction"));
+		 * A value tag can take on different values for each type.  This is set up when a tag is added to a row.
 		 */
-		class FDynamicTag
+		class FValueTag
 		{
 		public:
-			TYPEDELEMENTFRAMEWORK_API explicit FDynamicTag(const FName& InName);
+			TYPEDELEMENTFRAMEWORK_API explicit FValueTag(const FName& InName);
 			
 			TYPEDELEMENTFRAMEWORK_API const FName& GetName() const;
-			TYPEDELEMENTFRAMEWORK_API bool operator==(const FDynamicTag& Other) const = default;
+			TYPEDELEMENTFRAMEWORK_API bool operator==(const FValueTag& Other) const = default;
 		private:
-			TYPEDELEMENTFRAMEWORK_API friend uint32 GetTypeHash(const FDynamicTag& InName);
+			TYPEDELEMENTFRAMEWORK_API friend uint32 GetTypeHash(const FValueTag& InName);
 			FName Name;
 		};
 
 		template<typename T>
-		concept TValueTagType = std::is_same_v<T, FDynamicTag>;
+		concept TValueTagType = std::is_same_v<T, FValueTag>;
 
 		struct FDynamicColumnDescription
 		{
diff --git a/Engine/Source/Runtime/TypedElementFramework/Public/Elements/Common/TypedElementQueryDescription.h b/Engine/Source/Runtime/TypedElementFramework/Public/Elements/Common/TypedElementQueryDescription.h
--- a/Engine/Source/Runtime/TypedElementFramework/Public/Elements/Common/TypedElementQueryDescription.h
+++ b/Engine/Source/Runtime/TypedElementFramework/Public/Elements/Common/TypedElementQueryDescription.h
@@ -56,10 +56,10 @@ namespace TypedElementDataStorage
 			TWeakObjectPtr<const UScriptStruct> Type;
 		};
 
-		struct FDynamicTagData
+		struct FValueTagData
 		{
 			// The Tag maps to a Mass ConstSharedFragment object
-			UE::Editor::DataStorage::FDynamicTag Tag;
+			UE::Editor::DataStorage::FValueTag Tag;
 
 			// The MatchValue specifies the value that the fragment must have to be matched
 			// If MatchValue is NAME_None, then TEDS will match all values
@@ -104,7 +104,7 @@ namespace TypedElementDataStorage
 		TArray<EOperatorType, TInlineAllocator<NumInlineConditions>> DynamicConditionOperations;
 		TArray<UE::Editor::DataStorage::FDynamicColumnDescription, TInlineAllocator<NumInlineConditions>> DynamicConditionDescriptions;
 
-		TArray<FDynamicTagData> ValueTags;
+		TArray<FValueTagData> ValueTags;
 
 		TArray<TWeakObjectPtr<const UClass>, TInlineAllocator<NumInlineDependencies>> DependencyTypes;
 		TArray<EQueryDependencyFlags, TInlineAllocator<NumInlineDependencies>> DependencyFlags;
diff --git a/Engine/Source/Runtime/TypedElementFramework/Public/Elements/Framework/TypedElementQueryBuilder.h b/Engine/Source/Runtime/TypedElementFramework/Public/Elements/Framework/TypedElementQueryBuilder.h
--- a/Engine/Source/Runtime/TypedElementFramework/Public/Elements/Framework/TypedElementQueryBuilder.h
+++ b/Engine/Source/Runtime/TypedElementFramework/Public/Elements/Framework/TypedElementQueryBuilder.h
@@ -169,7 +169,7 @@ namespace TypedElementQueryBuilder
 		template<UE::Editor::DataStorage::TColumnType T>
 		FSimpleQuery& Any(const FName&);
 
-		// Dynamic Tags not yet supported for Any
+		// ValueTags not yet supported for Any
 		template<UE::Editor::DataStorage::TValueTagType>
 		FSimpleQuery& Any(const FName&) = delete;
 		template<UE::Editor::DataStorage::TValueTagType>
@@ -181,7 +181,7 @@ namespace TypedElementQueryBuilder
 		TYPEDELEMENTFRAMEWORK_API FSimpleQuery& None(const UScriptStruct* Target);
 		TYPEDELEMENTFRAMEWORK_API FSimpleQuery& None(TConstArrayView<const UScriptStruct*> Targets);
 		
-		// Dynamic Tags not yet supported for None
+		// ValueTags not yet supported for None
 		template<UE::Editor::DataStorage::TValueTagType>
 		FSimpleQuery& None(const FName& Tag) = delete;
 		template<UE::Editor::DataStorage::TValueTagType>
@@ -192,13 +192,13 @@ namespace TypedElementQueryBuilder
 		template<UE::Editor::DataStorage::TColumnType T>
 		FSimpleQuery& None(const FName&);
 		
-		// Dynamic Tags
+		// Value Tags
 		// ============
-		// Adds a filter to the query which must match a given DynamicTag.  The value of the DynamicTag will not be checked.
-		TYPEDELEMENTFRAMEWORK_API FSimpleQuery& All(const FDynamicTag& Tag);
-		// Adds a filter to the query which must match a given DynamicTag.  The value of the DynamicTag must also match.
+		// Adds a filter to the query which must match a given ValueTag.  The value of the ValueTag will not be checked.
+		TYPEDELEMENTFRAMEWORK_API FSimpleQuery& All(const FValueTag& Tag);
+		// Adds a filter to the query which must match a given ValueTag.  The value of the ValueTag must also match.
 		// Note: The query can only match a single value.  Multiple value queries are not supported at this time.
-		TYPEDELEMENTFRAMEWORK_API FSimpleQuery& All(const FDynamicTag& Tag, const FName& Value);
+		TYPEDELEMENTFRAMEWORK_API FSimpleQuery& All(const FValueTag& Tag, const FName& Value);
 		
 		TYPEDELEMENTFRAMEWORK_API FSimpleQuery& All(const UEnum& Enum);
 		TYPEDELEMENTFRAMEWORK_API FSimpleQuery& All(const UEnum& Enum, int64 Value);
@@ -208,10 +208,10 @@ namespace TypedElementQueryBuilder
 		template<UE::Editor::DataStorage::TColumnType T>
 		FSimpleQuery& All(const FName&);
 
-		// Adds a filter to the query which must match a given DynamicTag.  The value of the DynamicTag will not be checked.
+		// Adds a filter to the query which must match a given ValueTag.  The value of the DynamicTag will not be checked.
 		template<UE::Editor::DataStorage::TValueTagType>
 		FSimpleQuery& All(const FName& Tag);
-		// Adds a filter to the query which must match a given DynamicTag.  The value of the DynamicTag must also match.
+		// Adds a filter to the query which must match a given ValueTag.  The value of the DynamicTag must also match.
 		// Note: The query can only match a single value.  Multiple value queries are not supported at this time.
 		template<UE::Editor::DataStorage::TValueTagType>
 		FSimpleQuery& All(const FName& Tag, const FName& Value);
diff --git a/Engine/Source/Runtime/TypedElementFramework/Public/Elements/Framework/TypedElementQueryBuilder.inl b/Engine/Source/Runtime/TypedElementFramework/Public/Elements/Framework/TypedElementQueryBuilder.inl
--- a/Engine/Source/Runtime/TypedElementFramework/Public/Elements/Framework/TypedElementQueryBuilder.inl
+++ b/Engine/Source/Runtime/TypedElementFramework/Public/Elements/Framework/TypedElementQueryBuilder.inl
@@ -1201,13 +1201,13 @@ e.g. void(FCachedQueryContext<Subsystem1, const Subsystem2>& Context, RowHandle
 	template<UE::Editor::DataStorage::TValueTagType>
 	FSimpleQuery& FSimpleQuery::All(const FName& Tag)
 	{
-		return All(UE::Editor::DataStorage::FDynamicTag(Tag));
+		return All(UE::Editor::DataStorage::FValueTag(Tag));
 	}
 
 	template<UE::Editor::DataStorage::TValueTagType>
 	FSimpleQuery& FSimpleQuery::All(const FName& Tag, const FName& Value)
 	{
-		return All(UE::Editor::DataStorage::FDynamicTag(Tag), Value);
+		return All(UE::Editor::DataStorage::FValueTag(Tag), Value);
 	}
 	
 	template <TColumnType DynamicColumnTemplate>
diff --git a/Engine/Source/Runtime/TypedElementFramework/Public/Elements/Interfaces/TypedElementDataStorageInterface.h b/Engine/Source/Runtime/TypedElementFramework/Public/Elements/Interfaces/TypedElementDataStorageInterface.h
--- a/Engine/Source/Runtime/TypedElementFramework/Public/Elements/Interfaces/TypedElementDataStorageInterface.h
+++ b/Engine/Source/Runtime/TypedElementFramework/Public/Elements/Interfaces/TypedElementDataStorageInterface.h
@@ -187,17 +187,17 @@ public:
 	void AddColumn(RowHandle Row, ColumnType&& Column);
 
 	/**
-	 * Adds a DynamicTag with the given value to a row
-	 * A row can have multiple DynamicTags, but only one of each tag type.
+	 * Adds a ValueTag with the given value to a row
+	 * A row can have multiple ValueTags, but only one of each tag type.
 	 * Example:
-	 *   AddColumn(Row, FDynamicTag(TEXT("Color"), TEXT("Red));     // Valid
-	 *   AddColumn(Row, FDynamicTag(TEXT("Direction"), TEXT("Up")); // Valid
-	 *   AddColumn(Row, FDynamicTag(TEXT("Color"), TEXT("Blue"));   // Will do nothing since there already exists a Color dynamic tag
-	 * Note: Current support for changing a dynamic tag from one value to another requires that the tag is removed before a new one
+	 *   AddColumn(Row, ValueTag(TEXT("Color"), TEXT("Red));     // Valid
+	 *   AddColumn(Row, ValueTag(TEXT("Direction"), TEXT("Up")); // Valid
+	 *   AddColumn(Row, ValueTag(TEXT("Color"), TEXT("Blue"));   // Will do nothing since there already exists a Color value tag
+	 * Note: Current support for changing a value tag from one value to another requires that the tag is removed before a new one
 	 *       is added.  This will likely change in the future to transparently replace the tag to have consistent behaviour with other usages
 	 *       of AddColumn
 	 */
-	virtual void AddColumn(RowHandle Row, const UE::Editor::DataStorage::FDynamicTag& Tag, const FName& Value) = 0;
+	virtual void AddColumn(RowHandle Row, const UE::Editor::DataStorage::FValueTag& Tag, const FName& Value) = 0;
 
 	template<typename T>
 	void AddColumn(TypedElementDataStorage::RowHandle Row, const FName& Tag) = delete;
@@ -206,7 +206,7 @@ public:
 	void AddColumn(RowHandle Row, const FName& Tag, const FName& Value) = delete;
 	
 	template<>
-	void AddColumn<UE::Editor::DataStorage::FDynamicTag>(RowHandle Row, const FName& Tag, const FName& Value);
+	void AddColumn<UE::Editor::DataStorage::FValueTag>(RowHandle Row, const FName& Tag, const FName& Value);
 
 	template<UE::Editor::DataStorage::TEnumType EnumT>
 	void AddColumn(RowHandle Row, EnumT Value);
@@ -237,10 +237,10 @@ public:
 	void RemoveColumn(RowHandle Row);
 
 	/**
-	 * Removes a dynamic tag from the given row
+	 * Removes a value tag from the given row
 	 * If tag does not exist on row, operation will do nothing.
 	 */
-	virtual void RemoveColumn(RowHandle Row, const UE::Editor::DataStorage::FDynamicTag& Tag) = 0;
+	virtual void RemoveColumn(RowHandle Row, const UE::Editor::DataStorage::FValueTag& Tag) = 0;
 
 	template<typename T>
 	void RemoveColumn(RowHandle Row, const FName& Tag) = delete;
@@ -249,7 +249,7 @@ public:
 	void RemoveColumn(TypedElementDataStorage::RowHandle Row, const FName& Identifier);
 	
 	template<>
-	void RemoveColumn<UE::Editor::DataStorage::FDynamicTag>(RowHandle Row, const FName& Tag);
+	void RemoveColumn<UE::Editor::DataStorage::FValueTag>(RowHandle Row, const FName& Tag);
 
 	/**
 	 * Removes multiple columns from a row. This is typically more efficient than adding columns one
@@ -483,16 +483,16 @@ void ITypedElementDataStorageInterface::AddColumns(RowHandle Row)
 }
 
 template <>
-inline void ITypedElementDataStorageInterface::AddColumn<UE::Editor::DataStorage::FDynamicTag>(RowHandle Row, const FName& Tag, const FName& Value)
+inline void ITypedElementDataStorageInterface::AddColumn<UE::Editor::DataStorage::FValueTag>(RowHandle Row, const FName& Tag, const FName& Value)
 {
-	AddColumn(Row, UE::Editor::DataStorage::FDynamicTag(Tag), Value);
+	AddColumn(Row, UE::Editor::DataStorage::FValueTag(Tag), Value);
 }
 
 template <>
-inline void ITypedElementDataStorageInterface::RemoveColumn<UE::Editor::DataStorage::FDynamicTag>(RowHandle Row, const FName& Tag)
+inline void ITypedElementDataStorageInterface::RemoveColumn<UE::Editor::DataStorage::FValueTag>(RowHandle Row, const FName& Tag)
 {
 	using namespace UE::Editor::DataStorage;
-	RemoveColumn(Row, FDynamicTag(Tag));
+	RemoveColumn(Row, FValueTag(Tag));
 }
 
 template<UE::Editor::DataStorage::TColumnType DynamicColumnTemplate>
@@ -514,7 +514,7 @@ void ITypedElementDataStorageInterface::AddColumn(RowHandle Row, EnumT Value)
 	const FName ValueAsFName = *Enum->GetNameStringByValue(static_cast<int64>(Value));
 	if (ValueAsFName != NAME_None)
 	{
-		AddColumn(Row, UE::Editor::DataStorage::FDynamicTag(Enum->GetFName()), ValueAsFName);
+		AddColumn(Row, UE::Editor::DataStorage::FValueTag(Enum->GetFName()), ValueAsFName);
 	}
 }
 
@@ -579,7 +579,7 @@ template<UE::Editor::DataStorage::TEnumType EnumT>
 void ITypedElementDataStorageInterface::RemoveColumn(RowHandle Row)
 {
 	const UEnum* Enum = StaticEnum<EnumT>();
-	RemoveColumn(Row, UE::Editor::DataStorage::FDynamicTag(Enum->GetFName()));
+	RemoveColumn(Row, UE::Editor::DataStorage::FValueTag(Enum->GetFName()));
 }
 
 template<UE::Editor::DataStorage::TColumnType... Columns>
