[Teds] Update namespace from TypedElementDataStorage to UE::Editor::DataStorage for metadata and query store.

Tests now all roll up to the same categories "Editor > Data Storage"

#rnx

diff --git a/Engine/Plugins/Experimental/Animation/HierarchyTable/Source/Editor/Private/HierarchyTable/Widgets/OverrideWidgetConstructor.cpp b/Engine/Plugins/Experimental/Animation/HierarchyTable/Source/Editor/Private/HierarchyTable/Widgets/OverrideWidgetConstructor.cpp
--- a/Engine/Plugins/Experimental/Animation/HierarchyTable/Source/Editor/Private/HierarchyTable/Widgets/OverrideWidgetConstructor.cpp
+++ b/Engine/Plugins/Experimental/Animation/HierarchyTable/Source/Editor/Private/HierarchyTable/Widgets/OverrideWidgetConstructor.cpp
@@ -15,7 +15,7 @@ FTypedElementWidgetConstructor_Override::FTypedElementWidgetConstructor_Override
 }
 
 TSharedPtr<SWidget> FTypedElementWidgetConstructor_Override::CreateWidget(
-	const TypedElementDataStorage::FMetaDataView& Arguments)
+	const UE::Editor::DataStorage::FMetaDataView& Arguments)
 {
 	return SNew(SBox)
 		.HAlign(HAlign_Left)
@@ -25,7 +25,7 @@ TSharedPtr<SWidget> FTypedElementWidgetConstructor_Override::CreateWidget(
 bool FTypedElementWidgetConstructor_Override::FinalizeWidget(
 	ITypedElementDataStorageInterface* DataStorage,
 	ITypedElementDataStorageUiInterface* DataStorageUi,
-	TypedElementDataStorage::RowHandle Row,
+	UE::Editor::DataStorage::RowHandle Row,
 	const TSharedPtr<SWidget>& Widget)
 {
 	if (!Widget)
@@ -41,7 +41,7 @@ bool FTypedElementWidgetConstructor_Override::FinalizeWidget(
 	TSharedPtr<SBox> WidgetInstance = StaticCastSharedPtr<SBox>(Widget);
 
 	// Row is not actually the row we want, its contained inside of a row reference, I don't know why things are done this way.
-	const TypedElementDataStorage::RowHandle TargetRow = DataStorage->GetColumn<FTypedElementRowReferenceColumn>(Row)->Row;
+	const UE::Editor::DataStorage::RowHandle TargetRow = DataStorage->GetColumn<FTypedElementRowReferenceColumn>(Row)->Row;
 	FTypedElementOverrideColumn* OverrideColumn = DataStorage->GetColumn<FTypedElementOverrideColumn>(TargetRow);
 
 	if (OverrideColumn == nullptr)
diff --git a/Engine/Plugins/Experimental/Animation/HierarchyTable/Source/Editor/Private/HierarchyTable/Widgets/OverrideWidgetConstructor.h b/Engine/Plugins/Experimental/Animation/HierarchyTable/Source/Editor/Private/HierarchyTable/Widgets/OverrideWidgetConstructor.h
--- a/Engine/Plugins/Experimental/Animation/HierarchyTable/Source/Editor/Private/HierarchyTable/Widgets/OverrideWidgetConstructor.h
+++ b/Engine/Plugins/Experimental/Animation/HierarchyTable/Source/Editor/Private/HierarchyTable/Widgets/OverrideWidgetConstructor.h
@@ -16,7 +16,7 @@ public:
 	~FTypedElementWidgetConstructor_Override() override = default;
 
 protected:
-	TSharedPtr<SWidget> CreateWidget(const TypedElementDataStorage::FMetaDataView& Arguments);
+	TSharedPtr<SWidget> CreateWidget(const UE::Editor::DataStorage::FMetaDataView& Arguments);
 	bool FinalizeWidget(ITypedElementDataStorageInterface* DataStorage, ITypedElementDataStorageUiInterface* DataStorageUi,
-		TypedElementDataStorage::RowHandle Row, const TSharedPtr<SWidget>& Widget) override;
+		UE::Editor::DataStorage::RowHandle Row, const TSharedPtr<SWidget>& Widget) override;
 };
\ No newline at end of file
diff --git a/Engine/Plugins/Experimental/Animation/HierarchyTable/Source/Editor/Private/HierarchyTableEditorToolkit.cpp b/Engine/Plugins/Experimental/Animation/HierarchyTable/Source/Editor/Private/HierarchyTableEditorToolkit.cpp
--- a/Engine/Plugins/Experimental/Animation/HierarchyTable/Source/Editor/Private/HierarchyTableEditorToolkit.cpp
+++ b/Engine/Plugins/Experimental/Animation/HierarchyTable/Source/Editor/Private/HierarchyTableEditorToolkit.cpp
@@ -10,6 +10,7 @@
 #include "Elements/Columns/TypedElementTypeInfoColumns.h"
 #include "Elements/Columns/TypedElementSelectionColumns.h"
 #include "Elements/Columns/TypedElementMiscColumns.h"
+#include "Elements/Interfaces/TypedElementDataStorageInterface.h"
 #include "SceneOutlinerPublicTypes.h"
 #include "TedsOutlinerMode.h"
 #include "ReferenceSkeleton.h"
@@ -129,7 +130,7 @@ TSharedRef<SWidget> FHierarchyTableEditorToolkit::CreateTedsOutliner()
 	TArray<UScriptStruct*> HierarchyTableTypeColumns = Handler->GetColumns();
 	HierarchyTableTypeColumns.Add(FTypedElementOverrideColumn::StaticStruct());
 
-	TypedElementDataStorage::FQueryDescription ColumnQueryDescription =
+	FQueryDescription ColumnQueryDescription =
 		Select()
 		.ReadOnly(HierarchyTableTypeColumns)
 		.Compile();
@@ -143,7 +144,7 @@ TSharedRef<SWidget> FHierarchyTableEditorToolkit::CreateTedsOutliner()
 
 	FTedsOutlinerParams Params(nullptr);
 	{
-		TypedElementDataStorage::FQueryDescription RowQueryDescription =
+		FQueryDescription RowQueryDescription =
 			Select()
 			.Where()
 			.All<FTypedElementOverrideColumn>()
@@ -233,9 +234,9 @@ void FHierarchyTableEditorToolkit::AddCurveEntry(const FName CurveName)
 	using namespace UE::Editor::DataStorage;
 	UTypedElementRegistry* Registry = UTypedElementRegistry::GetInstance();
 	ITypedElementDataStorageInterface* DSI = Registry->GetMutableDataStorage();
-	static TypedElementDataStorage::TableHandle Table = DSI->FindTable(FName("Editor_HierarchyTableTable"));
+	static TableHandle Table = DSI->FindTable(FName("Editor_HierarchyTableTable"));
 
-	TypedElementDataStorage::RowHandle* RootBoneHandle = EntryIndexToHandleMap.Find(0);
+	RowHandle* RootBoneHandle = EntryIndexToHandleMap.Find(0);
 	if (!ensure(RootBoneHandle))
 	{
 		return;
@@ -247,7 +248,7 @@ void FHierarchyTableEditorToolkit::AddCurveEntry(const FName CurveName)
 	DSI->AddColumn<FTypedElementMetadataColumn>(Row,
 		{
 			.OwnerTable = HierarchyTable,
-			.Children = TArray<TypedElementDataStorage::RowHandle>(),
+			.Children = TArray<RowHandle>(),
 			.Type = EHierarchyTableMetadataType::Curve,
 			.BoneIndex = INDEX_NONE,
 			.CurveName = CurveName
diff --git a/Engine/Plugins/Experimental/Animation/HierarchyTable/Source/Editor/Private/HierarchyTableWidgetConstructor.cpp b/Engine/Plugins/Experimental/Animation/HierarchyTable/Source/Editor/Private/HierarchyTableWidgetConstructor.cpp
--- a/Engine/Plugins/Experimental/Animation/HierarchyTable/Source/Editor/Private/HierarchyTableWidgetConstructor.cpp
+++ b/Engine/Plugins/Experimental/Animation/HierarchyTable/Source/Editor/Private/HierarchyTableWidgetConstructor.cpp
@@ -16,7 +16,7 @@ TSharedRef<SWidget> FHierarchyTableWidgetConstructor::CreateInternalWidget(FHier
 	return SNullWidget::NullWidget;
 }
 
-TSharedPtr<SWidget> FHierarchyTableWidgetConstructor::CreateWidget(const TypedElementDataStorage::FMetaDataView& Arguments)
+TSharedPtr<SWidget> FHierarchyTableWidgetConstructor::CreateWidget(const UE::Editor::DataStorage::FMetaDataView& Arguments)
 {
 	return SNew(SBox)
 		.HAlign(HAlign_Left)
@@ -24,12 +24,12 @@ TSharedPtr<SWidget> FHierarchyTableWidgetConstructor::CreateWidget(const TypedEl
 }
 
 bool FHierarchyTableWidgetConstructor::FinalizeWidget(ITypedElementDataStorageInterface* DataStorage, ITypedElementDataStorageUiInterface* DataStorageUi,
-	TypedElementDataStorage::RowHandle Row, const TSharedPtr<SWidget>& Widget)
+	UE::Editor::DataStorage::RowHandle Row, const TSharedPtr<SWidget>& Widget)
 {
 	TSharedPtr<SBox> WidgetInstance = StaticCastSharedPtr<SBox>(Widget);
 
 	// Row is not actually the row we want, its contained inside of a row reference, I don't know why things are done this way.
-	TypedElementDataStorage::RowHandle TargetRow = DataStorage->GetColumn<FTypedElementRowReferenceColumn>(Row)->Row;
+	UE::Editor::DataStorage::RowHandle TargetRow = DataStorage->GetColumn<FTypedElementRowReferenceColumn>(Row)->Row;
 	FTypedElementOverrideColumn* OverrideColumn = DataStorage->GetColumn<FTypedElementOverrideColumn>(TargetRow);
 
 	TSharedRef<SWidget> ActualWidget = CreateInternalWidget(OverrideColumn->OwnerEntry);
diff --git a/Engine/Plugins/Experimental/Animation/HierarchyTable/Source/Editor/Public/HierarchyTableWidgetConstructor.h b/Engine/Plugins/Experimental/Animation/HierarchyTable/Source/Editor/Public/HierarchyTableWidgetConstructor.h
--- a/Engine/Plugins/Experimental/Animation/HierarchyTable/Source/Editor/Public/HierarchyTableWidgetConstructor.h
+++ b/Engine/Plugins/Experimental/Animation/HierarchyTable/Source/Editor/Public/HierarchyTableWidgetConstructor.h
@@ -23,9 +23,9 @@ public:
 
 protected:
 	// Begin FHierarchyTableWidgetConstructor
-	TSharedPtr<SWidget> CreateWidget(const TypedElementDataStorage::FMetaDataView& Arguments);
+	TSharedPtr<SWidget> CreateWidget(const UE::Editor::DataStorage::FMetaDataView& Arguments);
 
 	bool FinalizeWidget(ITypedElementDataStorageInterface* DataStorage, ITypedElementDataStorageUiInterface* DataStorageUi,
-		TypedElementDataStorage::RowHandle Row, const TSharedPtr<SWidget>& Widget) override;
+		UE::Editor::DataStorage::RowHandle Row, const TSharedPtr<SWidget>& Widget) override;
 	// End FHierarchyTableWidgetConstructor
 };
\ No newline at end of file
diff --git a/Engine/Plugins/Experimental/EditorDataStorage/Source/TedsCore/Private/Commands/EditorDataStorageCompatibilityCommands.cpp b/Engine/Plugins/Experimental/EditorDataStorage/Source/TedsCore/Private/Commands/EditorDataStorageCompatibilityCommands.cpp
--- a/Engine/Plugins/Experimental/EditorDataStorage/Source/TedsCore/Private/Commands/EditorDataStorageCompatibilityCommands.cpp
+++ b/Engine/Plugins/Experimental/EditorDataStorage/Source/TedsCore/Private/Commands/EditorDataStorageCompatibilityCommands.cpp
@@ -344,57 +344,57 @@ namespace UE::Editor::DataStorage
 	// FGetSourceRowHandle
 	//
 
-	TypedElementDataStorage::RowHandle FGetSourceRowHandle::operator()(const FAddCompatibleUObject& Command)
+	RowHandle FGetSourceRowHandle::operator()(const FAddCompatibleUObject& Command)
 	{
 		return Command.Row;
 	}
 
-	TypedElementDataStorage::RowHandle FGetSourceRowHandle::operator()(const FAddCompatibleExternalObject& Command)
+	RowHandle FGetSourceRowHandle::operator()(const FAddCompatibleExternalObject& Command)
 	{
 		return Command.Row;
 	}
 
-	TypedElementDataStorage::RowHandle FGetSourceRowHandle::operator()(const FCreateMemento& Command)
+	RowHandle FGetSourceRowHandle::operator()(const FCreateMemento& Command)
 	{
 		return Command.TargetRow;
 	}
 
-	TypedElementDataStorage::RowHandle FGetSourceRowHandle::operator()(const FRestoreMemento& Command)
+	RowHandle FGetSourceRowHandle::operator()(const FRestoreMemento& Command)
 	{
 		return Command.TargetRow;
 	}
 
-	TypedElementDataStorage::RowHandle FGetSourceRowHandle::operator()(const FDestroyMemento& Command)
+	RowHandle FGetSourceRowHandle::operator()(const FDestroyMemento& Command)
 	{
 		return Command.MementoRow;
 	}
 
-	TypedElementDataStorage::RowHandle FGetSourceRowHandle::operator()(const FRemoveCompatibleUObject& Command)
+	RowHandle FGetSourceRowHandle::operator()(const FRemoveCompatibleUObject& Command)
 	{
 		return Command.ObjectRow;
 	}
 
-	TypedElementDataStorage::RowHandle FGetSourceRowHandle::operator()(const FRemoveCompatibleExternalObject& Command)
+	RowHandle FGetSourceRowHandle::operator()(const FRemoveCompatibleExternalObject& Command)
 	{
 		return Command.ObjectRow;
 	}
 
-	TypedElementDataStorage::RowHandle FGetSourceRowHandle::Get(const CompatibilityCommandBuffer::TCommandVariant& Command)
+	RowHandle FGetSourceRowHandle::Get(const CompatibilityCommandBuffer::TCommandVariant& Command)
 	{
 		return Visit(FGetSourceRowHandle{}, Command);
 	}
 
-	TypedElementDataStorage::RowHandle FGetSourceRowHandle::operator()(const FAddInteractiveSyncFromWorldTag& Command)
+	RowHandle FGetSourceRowHandle::operator()(const FAddInteractiveSyncFromWorldTag& Command)
 	{
 		return Command.Row;
 	}
 
-	TypedElementDataStorage::RowHandle FGetSourceRowHandle::operator()(const FRemoveInteractiveSyncFromWorldTag& Command)
+	RowHandle FGetSourceRowHandle::operator()(const FRemoveInteractiveSyncFromWorldTag& Command)
 	{
 		return Command.Row;
 	}
 
-	TypedElementDataStorage::RowHandle FGetSourceRowHandle::operator()(const FAddSyncFromWorldTag& Command)
+	RowHandle FGetSourceRowHandle::operator()(const FAddSyncFromWorldTag& Command)
 	{
 		return Command.Row;
 	}
@@ -513,7 +513,7 @@ namespace UE::Editor::DataStorage
 		, MementoSystem(InStorageCompatibility.Environment->GetMementoSystem())
 	{}
 
-	void FCommandProcessor::SetupRow(TypedElementDataStorage::RowHandle Row, UObject* Object)
+	void FCommandProcessor::SetupRow(RowHandle Row, UObject* Object)
 	{
 		Storage.AddColumn(Row, FTypedElementUObjectColumn{ .Object = Object });
 		Storage.AddColumn(Row, FTypedElementUObjectIdColumn
@@ -531,7 +531,7 @@ namespace UE::Editor::DataStorage
 		StorageCompatibility.TriggerOnObjectAdded(Object, Object->GetClass(), Row);
 	}
 
-	void FCommandProcessor::SetupRow(TypedElementDataStorage::RowHandle Row, void* Object, TWeakObjectPtr<UScriptStruct> TypeInfo)
+	void FCommandProcessor::SetupRow(RowHandle Row, void* Object, TWeakObjectPtr<UScriptStruct> TypeInfo)
 	{
 		Storage.AddColumn(Row, FTypedElementExternalObjectColumn{ .Object = Object });
 		Storage.AddColumn(Row, FTypedElementScriptStructTypeInfoColumn{ .TypeInfo = TypeInfo });
@@ -587,7 +587,7 @@ namespace UE::Editor::DataStorage
 		checkf(ObjectPtr, TEXT(
 			"Expected a valid object pointer. If there isn't one here then the filter pass did not correctly clean up this command."));
 		Storage.AddRow(Object.Row, Object.Table, 
-			[ObjectPtr, this](TypedElementDataStorage::RowHandle Row)
+			[ObjectPtr, this](RowHandle Row)
 			{
 				SetupRow(Row, ObjectPtr);
 			});
@@ -614,7 +614,7 @@ namespace UE::Editor::DataStorage
 
 	void FCommandProcessor::operator()(FAddCompatibleExternalObject& Object)
 	{
-		Storage.AddRow(Object.Row, Object.Table, [&Object, this](TypedElementDataStorage::RowHandle Row)
+		Storage.AddRow(Object.Row, Object.Table, [&Object, this](RowHandle Row)
 			{
 				SetupRow(Object.Row, Object.Object, Object.TypeInfo);
 			});
@@ -1120,7 +1120,7 @@ namespace UE::Editor::DataStorage
 		FoldSyncFromWorldTags(Range);
 	}
 
-	CompatibilityCommandBuffer::FOptimizer FCommandOptimizer::CreateRangeOptimizer(TypedElementDataStorage::RowHandle RowCluster)
+	CompatibilityCommandBuffer::FOptimizer FCommandOptimizer::CreateRangeOptimizer(RowHandle RowCluster)
 	{
 		return Optimizer.BranchOnLeft(
 			[RowCluster](const CompatibilityCommandBuffer::TCommandVariant& Command)
@@ -1129,7 +1129,7 @@ namespace UE::Editor::DataStorage
 			});
 	}
 
-	void FCommandOptimizer::RunRightOnRowCluster(TypedElementDataStorage::RowHandle RowCluster)
+	void FCommandOptimizer::RunRightOnRowCluster(RowHandle RowCluster)
 	{
 		CompatibilityCommandBuffer::FOptimizer SubOptimizer = Optimizer.BranchOnRight(
 			[RowCluster](const CompatibilityCommandBuffer::TCommandVariant& Command)
diff --git a/Engine/Plugins/Experimental/EditorDataStorage/Source/TedsCore/Private/Commands/EditorDataStorageCompatibilityCommands.h b/Engine/Plugins/Experimental/EditorDataStorage/Source/TedsCore/Private/Commands/EditorDataStorageCompatibilityCommands.h
--- a/Engine/Plugins/Experimental/EditorDataStorage/Source/TedsCore/Private/Commands/EditorDataStorageCompatibilityCommands.h
+++ b/Engine/Plugins/Experimental/EditorDataStorage/Source/TedsCore/Private/Commands/EditorDataStorageCompatibilityCommands.h
@@ -67,20 +67,20 @@ namespace UE::Editor::DataStorage
 	struct FRegisterTypeTableAssociation final
 	{
 		TWeakObjectPtr<UStruct> TypeInfo;
-		TypedElementDataStorage::TableHandle Table;
+		TableHandle Table;
 	};
 
 	struct FAddCompatibleUObject final
 	{
 		TWeakObjectPtr<UObject> Object;
-		TypedElementDataStorage::RowHandle Row;
-		TypedElementDataStorage::TableHandle Table;
+		RowHandle Row;
+		TableHandle Table;
 	};
 
 	using ObjectAddedCallback = TFunction<void(
 		const void* /*Object*/, 
 		const FObjectTypeInfo&, /*Type information*/
-		TypedElementDataStorage::RowHandle /*Row*/)>;
+		RowHandle /*Row*/)>;
 	struct FRegisterObjectAddedCallback final
 	{
 		ObjectAddedCallback Callback;
@@ -95,13 +95,13 @@ namespace UE::Editor::DataStorage
 	using ObjectRemovedCallback = TFunction<void(
 		const void* /*Object*/, 
 		const FObjectTypeInfo&, /*Type information*/
-		TypedElementDataStorage::RowHandle /*Row*/)>;
+		RowHandle /*Row*/)>;
 
 	struct FBatchAddCompatibleUObject final
 	{
 		TWeakObjectPtr<UObject>* ObjectArray;
-		TypedElementDataStorage::RowHandle* RowArray;
-		TypedElementDataStorage::TableHandle Table;
+		RowHandle* RowArray;
+		TableHandle Table;
 		uint32 Count;
 	};
 
@@ -109,52 +109,52 @@ namespace UE::Editor::DataStorage
 	{
 		void* Object;
 		TWeakObjectPtr<UScriptStruct> TypeInfo;
-		TypedElementDataStorage::RowHandle Row;
-		TypedElementDataStorage::TableHandle Table;
+		RowHandle Row;
+		TableHandle Table;
 	};
 
 	struct FCreateMemento final
 	{
-		TypedElementDataStorage::RowHandle ReservedMementoRow;
-		TypedElementDataStorage::RowHandle TargetRow;
+		RowHandle ReservedMementoRow;
+		RowHandle TargetRow;
 	};
 
 	struct FRestoreMemento final
 	{
-		TypedElementDataStorage::RowHandle MementoRow;
-		TypedElementDataStorage::RowHandle TargetRow;
+		RowHandle MementoRow;
+		RowHandle TargetRow;
 	};
 
 	struct FDestroyMemento final
 	{
-		TypedElementDataStorage::RowHandle MementoRow;
+		RowHandle MementoRow;
 	};
 
 	struct FRemoveCompatibleUObject final
 	{
 		const UObject* Object;
-		TypedElementDataStorage::RowHandle ObjectRow = TypedElementDataStorage::InvalidRowHandle;
+		RowHandle ObjectRow = InvalidRowHandle;
 	};
 
 	struct FRemoveCompatibleExternalObject final
 	{
 		void* Object;
-		TypedElementDataStorage::RowHandle ObjectRow;
+		RowHandle ObjectRow;
 	};
 
 	struct FBatchAddCompatibleExternalObject final
 	{
 		void** ObjectArray;
 		TWeakObjectPtr<UScriptStruct>* TypeInfoArray;
-		TypedElementDataStorage::RowHandle* RowArray;
-		TypedElementDataStorage::TableHandle Table;
+		RowHandle* RowArray;
+		TableHandle Table;
 		uint32 Count;
 	};
 
 	struct FAddSyncFromWorldTag final
 	{
 		TObjectKey<const UObject> Target;
-		TypedElementDataStorage::RowHandle Row;
+		RowHandle Row;
 
 		static UScriptStruct* GetType();
 		static UScriptStruct** GetTypeAddress();
@@ -163,7 +163,7 @@ namespace UE::Editor::DataStorage
 	struct FAddInteractiveSyncFromWorldTag final
 	{
 		TObjectKey<const UObject> Target;
-		TypedElementDataStorage::RowHandle Row;
+		RowHandle Row;
 
 		static UScriptStruct* GetType();
 		static UScriptStruct** GetTypeAddress();
@@ -172,7 +172,7 @@ namespace UE::Editor::DataStorage
 	struct FRemoveInteractiveSyncFromWorldTag final
 	{
 		TObjectKey<const UObject> Target;
-		TypedElementDataStorage::RowHandle Row;
+		RowHandle Row;
 	};
 
 	using CompatibilityCommandBuffer = FCommandBuffer
@@ -254,18 +254,18 @@ namespace UE::Editor::DataStorage
 	struct FGetSourceRowHandle
 	{
 		template<typename T>
-		TypedElementDataStorage::RowHandle operator()(const T&){ return TypedElementDataStorage::InvalidRowHandle; }
-		TypedElementDataStorage::RowHandle operator()(const FAddCompatibleUObject& Command);
-		TypedElementDataStorage::RowHandle operator()(const FAddCompatibleExternalObject& Command);
-		TypedElementDataStorage::RowHandle operator()(const FCreateMemento& Command);
-		TypedElementDataStorage::RowHandle operator()(const FRestoreMemento& Command);
-		TypedElementDataStorage::RowHandle operator()(const FDestroyMemento& Command);
-		TypedElementDataStorage::RowHandle operator()(const FRemoveCompatibleUObject& Command);
-		TypedElementDataStorage::RowHandle operator()(const FRemoveCompatibleExternalObject& Command);
-		TypedElementDataStorage::RowHandle operator()(const FAddInteractiveSyncFromWorldTag& Command);
-		TypedElementDataStorage::RowHandle operator()(const FRemoveInteractiveSyncFromWorldTag& Command);
-		TypedElementDataStorage::RowHandle operator()(const FAddSyncFromWorldTag& Command);
-		static TypedElementDataStorage::RowHandle Get(const CompatibilityCommandBuffer::TCommandVariant& Command);
+		RowHandle operator()(const T&){ return InvalidRowHandle; }
+		RowHandle operator()(const FAddCompatibleUObject& Command);
+		RowHandle operator()(const FAddCompatibleExternalObject& Command);
+		RowHandle operator()(const FCreateMemento& Command);
+		RowHandle operator()(const FRestoreMemento& Command);
+		RowHandle operator()(const FDestroyMemento& Command);
+		RowHandle operator()(const FRemoveCompatibleUObject& Command);
+		RowHandle operator()(const FRemoveCompatibleExternalObject& Command);
+		RowHandle operator()(const FAddInteractiveSyncFromWorldTag& Command);
+		RowHandle operator()(const FRemoveInteractiveSyncFromWorldTag& Command);
+		RowHandle operator()(const FAddSyncFromWorldTag& Command);
+		static RowHandle Get(const CompatibilityCommandBuffer::TCommandVariant& Command);
 	};
 
 	/** 
@@ -301,8 +301,8 @@ namespace UE::Editor::DataStorage
 	{
 		FCommandProcessor(ITypedElementDataStorageInterface& InStorage, UEditorDataStorageCompatibility& InStorageCompatibility);
 
-		void SetupRow(TypedElementDataStorage::RowHandle Row, UObject* Object);
-		void SetupRow(TypedElementDataStorage::RowHandle Row, void* Object, TWeakObjectPtr<UScriptStruct> TypeInfo);
+		void SetupRow(RowHandle Row, UObject* Object);
+		void SetupRow(RowHandle Row, void* Object, TWeakObjectPtr<UScriptStruct> TypeInfo);
 		
 		void operator()(FNopCommand&) {}
 		void operator()(FTypeInfoReinstanced&) {}
@@ -382,9 +382,9 @@ namespace UE::Editor::DataStorage
 		 * Creates a new optimizer starting at the left of the current optimizer that will run only as long as commands on the right use 
 		 * the same row as the provided argument.
 		 */
-		CompatibilityCommandBuffer::FOptimizer CreateRangeOptimizer(TypedElementDataStorage::RowHandle RowCluster);
+		CompatibilityCommandBuffer::FOptimizer CreateRangeOptimizer(RowHandle RowCluster);
 		/** Keeps processing a subset of commands until the right no longer has a command with the provided source row. */
-		void RunRightOnRowCluster(TypedElementDataStorage::RowHandle RowCluster);
+		void RunRightOnRowCluster(RowHandle RowCluster);
 		static void Run(CompatibilityCommandBuffer::FCollection& Commands, FScratchBuffer& ScratchBuffer);
 	};
 } // namespace UE::Editor::DataStorage
diff --git a/Engine/Plugins/Experimental/EditorDataStorage/Source/TedsCore/Private/Compatibility/TypedElementObjectReinstancingManager.cpp b/Engine/Plugins/Experimental/EditorDataStorage/Source/TedsCore/Private/Compatibility/TypedElementObjectReinstancingManager.cpp
--- a/Engine/Plugins/Experimental/EditorDataStorage/Source/TedsCore/Private/Compatibility/TypedElementObjectReinstancingManager.cpp
+++ b/Engine/Plugins/Experimental/EditorDataStorage/Source/TedsCore/Private/Compatibility/TypedElementObjectReinstancingManager.cpp
@@ -47,7 +47,6 @@ void UTedsObjectReinstancingManager::Deinitialize()
 
 void UTedsObjectReinstancingManager::UpdateCompleted()
 {
-	using namespace TypedElementDataStorage;
 	using namespace UE::Editor::DataStorage;
 
 	FMementoSystem& MementoSystem = DataStorage->GetEnvironment()->GetMementoSystem();
@@ -61,10 +60,10 @@ void UTedsObjectReinstancingManager::UpdateCompleted()
 void UTedsObjectReinstancingManager::HandleOnObjectPreRemoved(
 	const void* Object, 
 	const UE::Editor::DataStorage::FObjectTypeInfo& TypeInfo, 
-	TypedElementDataStorage::RowHandle ObjectRow)
+	UE::Editor::DataStorage::RowHandle ObjectRow)
 {
 	// This is the chance to record the old object to memento
-	TypedElementDataStorage::RowHandle Memento = DataStorage->GetEnvironment()->GetMementoSystem().CreateMemento(ObjectRow);
+	UE::Editor::DataStorage::RowHandle Memento = DataStorage->GetEnvironment()->GetMementoSystem().CreateMemento(ObjectRow);
 	OldObjectToMementoMap.Add(Object, Memento);
 }
 
diff --git a/Engine/Plugins/Experimental/EditorDataStorage/Source/TedsCore/Private/Compatibility/TypedElementObjectReinstancingManager.h b/Engine/Plugins/Experimental/EditorDataStorage/Source/TedsCore/Private/Compatibility/TypedElementObjectReinstancingManager.h
--- a/Engine/Plugins/Experimental/EditorDataStorage/Source/TedsCore/Private/Compatibility/TypedElementObjectReinstancingManager.h
+++ b/Engine/Plugins/Experimental/EditorDataStorage/Source/TedsCore/Private/Compatibility/TypedElementObjectReinstancingManager.h
@@ -35,7 +35,7 @@ private:
 	void HandleOnObjectPreRemoved(
 		const void* Object, 
 		const UE::Editor::DataStorage::FObjectTypeInfo& TypeInfo, 
-		TypedElementDataStorage::RowHandle ObjectRow);
+		UE::Editor::DataStorage::RowHandle ObjectRow);
 	void HandleOnObjectsReinstanced(const FCoreUObjectDelegates::FReplacementObjectMap& ObjectReplacementMap);
 
 	UPROPERTY()
@@ -45,9 +45,9 @@ private:
 	
 	// Reverse lookup that holds all populated mementos for recently deleted objects
 	// Entry removed when the memento is removed
-	TMap<const void*, TypedElementDataStorage::RowHandle> OldObjectToMementoMap;
+	TMap<const void*, UE::Editor::DataStorage::RowHandle> OldObjectToMementoMap;
 	
-	TypedElementDataStorage::TableHandle MementoRowBaseTable;
+	UE::Editor::DataStorage::TableHandle MementoRowBaseTable;
 	FDelegateHandle UpdateCompletedCallbackHandle;
 	FDelegateHandle ReinstancingCallbackHandle;
 	FDelegateHandle ObjectRemovedCallbackHandle;
diff --git a/Engine/Plugins/Experimental/EditorDataStorage/Source/TedsCore/Private/Debug/TypedElementDatabaseDebugCommands.cpp b/Engine/Plugins/Experimental/EditorDataStorage/Source/TedsCore/Private/Debug/TypedElementDatabaseDebugCommands.cpp
--- a/Engine/Plugins/Experimental/EditorDataStorage/Source/TedsCore/Private/Debug/TypedElementDatabaseDebugCommands.cpp
+++ b/Engine/Plugins/Experimental/EditorDataStorage/Source/TedsCore/Private/Debug/TypedElementDatabaseDebugCommands.cpp
@@ -224,9 +224,9 @@ static FAutoConsoleCommand CVarCreateRow(
 	FConsoleCommandWithArgsDelegate::CreateLambda([](const TArray<FString>& Args)
 	{
 		ITypedElementDataStorageInterface* DataStorage = UTypedElementRegistry::GetInstance()->GetMutableDataStorage();
-		static TypedElementDataStorage::TableHandle Table = DataStorage->RegisterTable<FTestColumnA>(FName(TEXT("Debug.CreateRow Table")));
+		static UE::Editor::DataStorage::TableHandle Table = DataStorage->RegisterTable<FTestColumnA>(FName(TEXT("Debug.CreateRow Table")));
 
-		const TypedElementDataStorage::RowHandle RowHandle = DataStorage->AddRow(Table);
+		const UE::Editor::DataStorage::RowHandle RowHandle = DataStorage->AddRow(Table);
 
 		UE_LOG(LogEditorDataStorage, Warning, TEXT("Added Row %llu"), static_cast<uint64>(RowHandle));
 	}));
@@ -245,7 +245,7 @@ static FAutoConsoleCommand CVarAddDynamicColumnTag(
 		}
 
 		const uint64 RowAsU64 = FCString::Strtoui64(*Args[0], nullptr, 10);
-		const TypedElementDataStorage::RowHandle Row = RowAsU64;
+		const UE::Editor::DataStorage::RowHandle Row = RowAsU64;
 		const FName Identifier = FName(*Args[1]);
 
 		DataStorage->AddColumn<FTestDynamicTag>(Row, Identifier);
@@ -267,7 +267,7 @@ static FAutoConsoleCommand CVarAddDynamicColumn(
 		}
 
 		const uint64 RowAsU64 = FCString::Strtoui64(*Args[0], nullptr, 10);
-		const TypedElementDataStorage::RowHandle Row = RowAsU64;
+		const UE::Editor::DataStorage::RowHandle Row = RowAsU64;
 		const FName Identifier = FName(*Args[1]);
 		
 		bool bUseDefaultApi = false;
@@ -298,7 +298,7 @@ static FAutoConsoleCommand CVarRemoveDynamicColumn(
 		}
 
 		const uint64 RowAsU64 = FCString::Strtoui64(*Args[0], nullptr, 10);
-		const TypedElementDataStorage::RowHandle Row = RowAsU64;
+		const UE::Editor::DataStorage::RowHandle Row = RowAsU64;
 		const FName Identifier = FName(*Args[1]);
 		
 		DataStorage->RemoveColumn<FTestDynamicColumn>(Row, Identifier);
@@ -320,7 +320,7 @@ static FAutoConsoleCommand CVarAddToDynamicColumn(
 		}
 
 		const uint64 RowAsU64 = FCString::Strtoui64(*Args[0], nullptr, 10);
-		const TypedElementDataStorage::RowHandle Row = RowAsU64;
+		const UE::Editor::DataStorage::RowHandle Row = RowAsU64;
 		const FName TagId = FName(*Args[1]);
 		const uint64 Value = FCString::Strtoui64(*Args[2], nullptr, 10);
 		const uint64 MethodId = Args.Num() >= 4 ? FCString::Strtoui64(*Args[3], nullptr, 10) : 0;
@@ -371,7 +371,7 @@ static FAutoConsoleCommand CVarPrintDynamicColumn(
 			return;
 		}
 		const uint64 RowAsU64 = FCString::Strtoui64(*Args[0], nullptr, 10);
-		const TypedElementDataStorage::RowHandle Row = RowAsU64;
+		const UE::Editor::DataStorage::RowHandle Row = RowAsU64;
 		const FName TagId = FName(*Args[1]);
 
 		FTestDynamicColumn* Column = DataStorage->GetColumn<FTestDynamicColumn>(Row, TagId);
@@ -432,8 +432,8 @@ static FAutoConsoleCommand CVarPrintDynamicColumnWithQuery(
 			Compile());
 
 		TStringBuilder<1024> StringBuilder;
-		TypedElementDataStorage::FQueryResult Result = DataStorage->RunQuery(Query,CreateDirectQueryCallbackBinding(
-			[Identifier, &StringBuilder](TypedElementDataStorage::IDirectQueryContext& Context, const RowHandle* Rows)
+		UE::Editor::DataStorage::FQueryResult Result = DataStorage->RunQuery(Query,CreateDirectQueryCallbackBinding(
+			[Identifier, &StringBuilder](UE::Editor::DataStorage::IDirectQueryContext& Context, const RowHandle* Rows)
 			{
 				const TArrayView<const RowHandle> RowView = MakeConstArrayView(Rows, Context.GetRowCount());
 				// Get pointer to the start of the range of columns to process
@@ -473,7 +473,6 @@ static FAutoConsoleCommand CVarCountDynamicTagWithQuery(
 	FConsoleCommandWithArgsDelegate::CreateLambda([](const TArray<FString>& Args)
 	{
 		using namespace TypedElementQueryBuilder;
-		using DSI = ITypedElementDataStorageInterface;
 		using namespace UE::Editor::DataStorage;
 				
 		// Print column using query
@@ -495,7 +494,7 @@ static FAutoConsoleCommand CVarCountDynamicTagWithQuery(
 			Compile());
 
 		int32 Count = 0;
-		DataStorage->RunQuery(Query, CreateDirectQueryCallbackBinding([Identifier, &Count](TypedElementDataStorage::IDirectQueryContext& Context, const RowHandle* Rows)
+		DataStorage->RunQuery(Query, CreateDirectQueryCallbackBinding([Identifier, &Count](IDirectQueryContext& Context, const RowHandle* Rows)
 		{
 			Count += Context.GetRowCount();
 		}));
@@ -516,7 +515,6 @@ static FAutoConsoleCommand CVarRegisterListDynamicColumnQuery(
 	FConsoleCommandWithArgsDelegate::CreateLambda([](const TArray<FString>& Args)
 	{
 		using namespace TypedElementQueryBuilder;
-		using DSI = ITypedElementDataStorageInterface;
 		using namespace UE::Editor::DataStorage;
 				
 		// Print column using query
@@ -531,9 +529,6 @@ static FAutoConsoleCommand CVarRegisterListDynamicColumnQuery(
 		const FName Identifier(*Args[0]);
 		const FName ActivationGroup(*Args[1]);
 		
-		using namespace TypedElementQueryBuilder;
-		using DSI = ITypedElementDataStorageInterface;
-		using namespace UE::Editor::DataStorage;
 		using namespace TypedElementDataStorage;
 
 		// Lists the rows processed that have 
@@ -592,7 +587,7 @@ static FAutoConsoleCommand CVarAddValueTag(
 		}
 		
 		const uint64 RowAsU64 = FCString::Strtoui64(*Args[0], nullptr, 10);
-		const TypedElementDataStorage::RowHandle Row = RowAsU64;
+		const RowHandle Row = RowAsU64;
 
 		
 		const FName Value(*Args[2]);
@@ -626,7 +621,7 @@ static FAutoConsoleCommand CVarRemoveValueTag(
 		}
 		
 		const uint64 RowAsU64 = FCString::Strtoui64(*Args[0], nullptr, 10);
-		const TypedElementDataStorage::RowHandle Row = RowAsU64;
+		const UE::Editor::DataStorage::RowHandle Row = RowAsU64;
 		
 		constexpr bool bUseTemplateSugar = true;
 		if constexpr (bUseTemplateSugar)
@@ -649,7 +644,6 @@ static FAutoConsoleCommand CVarMatchValueTag(
 	FConsoleCommandWithArgsDelegate::CreateLambda([](const TArray<FString>& Args)
 	{
 		using namespace TypedElementQueryBuilder;
-		using DSI = ITypedElementDataStorageInterface;
 		using namespace UE::Editor::DataStorage;
 		
 		ITypedElementDataStorageInterface* DataStorage = UTypedElementRegistry::GetInstance()->GetMutableDataStorage();
@@ -688,8 +682,8 @@ static FAutoConsoleCommand CVarMatchValueTag(
 
 		uint64 Count = 0;
 		
-		const TypedElementDataStorage::FQueryResult Result = DataStorage->RunQuery(Query, CreateDirectQueryCallbackBinding(
-			[&Count](const DSI::IDirectQueryContext& Context, const RowHandle*)
+		const FQueryResult Result = DataStorage->RunQuery(Query, CreateDirectQueryCallbackBinding(
+			[&Count](const IDirectQueryContext& Context, const RowHandle*)
 			{
 				Count += Context.GetRowCount();
 			}));
@@ -765,7 +759,6 @@ static FAutoConsoleCommand CVarMatchValueTagFromEnum(
 	FConsoleCommandWithArgsDelegate::CreateLambda([](const TArray<FString>& Args)
 	{
 		using namespace TypedElementQueryBuilder;
-		using DSI = ITypedElementDataStorageInterface;
 		using namespace UE::Editor::DataStorage;
 		
 		ITypedElementDataStorageInterface* DataStorage = UTypedElementRegistry::GetInstance()->GetMutableDataStorage();
@@ -828,7 +821,7 @@ static FAutoConsoleCommand CVarMatchValueTagFromEnum(
 		uint64 Count = 0;
 		
 		const TypedElementDataStorage::FQueryResult Result = DataStorage->RunQuery(Query, CreateDirectQueryCallbackBinding(
-			[&Count](const DSI::IDirectQueryContext& Context, const RowHandle*)
+			[&Count](const IDirectQueryContext& Context, const RowHandle*)
 			{
 				Count += Context.GetRowCount();
 			}));
diff --git a/Engine/Plugins/Experimental/EditorDataStorage/Source/TedsCore/Private/Processors/TypedElementProcessorAdaptors.cpp b/Engine/Plugins/Experimental/EditorDataStorage/Source/TedsCore/Private/Processors/TypedElementProcessorAdaptors.cpp
--- a/Engine/Plugins/Experimental/EditorDataStorage/Source/TedsCore/Private/Processors/TypedElementProcessorAdaptors.cpp
+++ b/Engine/Plugins/Experimental/EditorDataStorage/Source/TedsCore/Private/Processors/TypedElementProcessorAdaptors.cpp
@@ -34,7 +34,7 @@ namespace UE::Editor::DataStorage
 			{
 				static_assert(
 					sizeof(RowHandle) == sizeof(FMassEntityHandle) && alignof(RowHandle) == alignof(FMassEntityHandle),
-					"UE::Editor::DataStorage::RowHandle and FMassEntityHandle must be layout compatible.");
+					"RowHandle and FMassEntityHandle must be layout compatible.");
 				TConstArrayView<FMassEntityHandle> Entities = Context.GetEntities();
 				return TConstArrayView<RowHandle>(reinterpret_cast<const RowHandle*>(Entities.GetData()), Entities.Num());
 			}
@@ -211,11 +211,11 @@ namespace UE::Editor::DataStorage
 				}
 			}
 
-			void AddColumns(TConstArrayView<TypedElementDataStorage::RowHandle> Rows, TConstArrayView<UE::Editor::DataStorage::FDynamicColumnDescription> DynamicColumnDescriptions)
+			void AddColumns(TConstArrayView<RowHandle> Rows, TConstArrayView<UE::Editor::DataStorage::FDynamicColumnDescription> DynamicColumnDescriptions)
 			{
 				struct FAddDynamicColumns
 				{
-					TConstArrayView<TypedElementDataStorage::RowHandle> Rows;
+					TConstArrayView<RowHandle> Rows;
 					TConstArrayView<FDynamicColumnDescription> Descriptions;
 					TArrayView<const UScriptStruct*> ResolvedTypes;
 				};
@@ -223,8 +223,8 @@ namespace UE::Editor::DataStorage
 				FScratchBuffer& ScratchBuffer = Environment.GetScratchBuffer();
 				
 				FAddDynamicColumns* CommandData = ScratchBuffer.Emplace<FAddDynamicColumns>();
-				TArrayView<TypedElementDataStorage::RowHandle> ScratchRows = MakeArrayView(
-					ScratchBuffer.EmplaceArray<TypedElementDataStorage::RowHandle>(Rows.Num()),
+				TArrayView<RowHandle> ScratchRows = MakeArrayView(
+					ScratchBuffer.EmplaceArray<RowHandle>(Rows.Num()),
 					DynamicColumnDescriptions.Num());
 				TArrayView<FDynamicColumnDescription> ScratchDescriptions = MakeArrayView(
 					ScratchBuffer.EmplaceArray<FDynamicColumnDescription>(DynamicColumnDescriptions.Num()),
@@ -257,7 +257,7 @@ namespace UE::Editor::DataStorage
 						FMassArchetypeCompositionDescriptor AddDescriptor;
 						TedsColumnsToMassDescriptor(AddDescriptor, CommandData->ResolvedTypes);
 						
-						for (TypedElementDataStorage::RowHandle Row : CommandData->Rows)
+						for (RowHandle Row : CommandData->Rows)
 						{
 							FMassEntityHandle Entity = FMassEntityHandle::FromNumber(Row);
 							if (System.IsEntityValid(Entity))
@@ -337,7 +337,7 @@ namespace UE::Editor::DataStorage
 				return ColumnData;
 			}
 
-			void* AddColumnUninitialized(TypedElementDataStorage::RowHandle Row, const FDynamicColumnDescription& Description)
+			void* AddColumnUninitialized(RowHandle Row, const FDynamicColumnDescription& Description)
 			{
 				return AddColumnUninitialized(Row, Description, [](const UScriptStruct& TypeInfo, void* Destination, void* Source)
 				{
@@ -345,7 +345,7 @@ namespace UE::Editor::DataStorage
 				});
 			}
 
-			void* AddColumnUninitialized(TypedElementDataStorage::RowHandle Row, const FDynamicColumnDescription& Description, ObjectCopyOrMove Relocator)
+			void* AddColumnUninitialized(RowHandle Row, const FDynamicColumnDescription& Description, ObjectCopyOrMove Relocator)
 			{
 				struct FAddDynamicColumn
 				{
@@ -590,7 +590,7 @@ namespace UE::Editor::DataStorage
 						System.BuildEntity(CommandData.Entity, CommandData.Archetype, SharedFragmentValues);
 					});
 		
-				const TypedElementDataStorage::RowHandle TedsRowHandle = EntityHandle.AsNumber();
+				const RowHandle TedsRowHandle = EntityHandle.AsNumber();
 				return TedsRowHandle;
 			}
 
@@ -657,11 +657,11 @@ namespace UE::Editor::DataStorage
 			void RemoveRows(TConstArrayView<RowHandle> Rows) override  { return Implementation.RemoveRows(Rows); }
 			void AddColumns(RowHandle Row, TConstArrayView<const UScriptStruct*> ColumnTypes) override { return Implementation.AddColumns(Row, ColumnTypes); }
 			void AddColumns(TConstArrayView<RowHandle> Rows, TConstArrayView<const UScriptStruct*> ColumnTypes) override  { return Implementation.AddColumns(Rows, ColumnTypes); }
-			void AddColumns(TConstArrayView<TypedElementDataStorage::RowHandle> Rows, TConstArrayView<UE::Editor::DataStorage::FDynamicColumnDescription> DynamicColumnDescriptions) override { Implementation.AddColumns(Rows, DynamicColumnDescriptions); }
+			void AddColumns(TConstArrayView<RowHandle> Rows, TConstArrayView<UE::Editor::DataStorage::FDynamicColumnDescription> DynamicColumnDescriptions) override { Implementation.AddColumns(Rows, DynamicColumnDescriptions); }
 			void* AddColumnUninitialized(RowHandle Row, const UScriptStruct* ColumnType) override  { return Implementation.AddColumnUninitialized(Row, ColumnType); }
 			void* AddColumnUninitialized(RowHandle Row, const UScriptStruct* ObjectType, ObjectCopyOrMove Relocator) override { return Implementation.AddColumnUninitialized(Row, ObjectType, Relocator); }
-			void* AddColumnUninitialized(TypedElementDataStorage::RowHandle Row, const UE::Editor::DataStorage::FDynamicColumnDescription& DynamicColumnDescription) override { return Implementation.AddColumnUninitialized(Row, DynamicColumnDescription); }
-			void* AddColumnUninitialized(TypedElementDataStorage::RowHandle Row, const UE::Editor::DataStorage::FDynamicColumnDescription& DynamicColumnDescription, ObjectCopyOrMove Relocator) override { return Implementation.AddColumnUninitialized(Row, DynamicColumnDescription, Relocator); }
+			void* AddColumnUninitialized(RowHandle Row, const UE::Editor::DataStorage::FDynamicColumnDescription& DynamicColumnDescription) override { return Implementation.AddColumnUninitialized(Row, DynamicColumnDescription); }
+			void* AddColumnUninitialized(RowHandle Row, const UE::Editor::DataStorage::FDynamicColumnDescription& DynamicColumnDescription, ObjectCopyOrMove Relocator) override { return Implementation.AddColumnUninitialized(Row, DynamicColumnDescription, Relocator); }
 			void RemoveColumns(RowHandle Row, TConstArrayView<const UScriptStruct*> ColumnTypes) override { Implementation.RemoveColumns(Row, ColumnTypes); }
 			void RemoveColumns(TConstArrayView<RowHandle> Rows, TConstArrayView<const UScriptStruct*> ColumnTypes) override { Implementation.RemoveColumns(Rows, ColumnTypes); }
 			const UScriptStruct* FindDynamicColumnType(const UE::Editor::DataStorage::FDynamicColumnDescription& Description) const override { return Implementation.FindDynamicColumnType(Description); }
@@ -729,21 +729,21 @@ namespace UE::Editor::DataStorage
 				return Environment.GetIndexTable().FindIndexedRow(Scope, Index);
 			}
 
-			TypedElementDataStorage::FQueryResult RunQuery(QueryHandle Query)
+			UE::Editor::DataStorage::FQueryResult RunQuery(QueryHandle Query)
 			{
 				const FExtendedQueryStore::Handle Handle(Query);
 				// This can be safely called because there's not callback, which means no columns are accessed, even for select queries.
 				return QueryStore.RunQuery(Context.GetEntityManagerChecked(), Handle);
 			}
 
-			TypedElementDataStorage::FQueryResult RunSubquery(int32 SubqueryIndex)
+			UE::Editor::DataStorage::FQueryResult RunSubquery(int32 SubqueryIndex)
 			{
 				return SubqueryIndex < QueryDescription.Subqueries.Num() ?
 					RunQuery(QueryDescription.Subqueries[SubqueryIndex]) :
-					TypedElementDataStorage::FQueryResult{};
+					UE::Editor::DataStorage::FQueryResult{};
 			}
 
-			TypedElementDataStorage::FQueryResult RunSubquery(int32 SubqueryIndex, TypedElementDataStorage::SubqueryCallbackRef Callback)
+			UE::Editor::DataStorage::FQueryResult RunSubquery(int32 SubqueryIndex, UE::Editor::DataStorage::SubqueryCallbackRef Callback)
 			{
 				if (SubqueryIndex < QueryDescription.Subqueries.Num())
 				{
@@ -753,12 +753,12 @@ namespace UE::Editor::DataStorage
 				}
 				else
 				{
-					return TypedElementDataStorage::FQueryResult{};
+					return UE::Editor::DataStorage::FQueryResult{};
 				}
 			}
 
-			TypedElementDataStorage::FQueryResult RunSubquery(int32 SubqueryIndex, TypedElementDataStorage::RowHandle Row,
-				TypedElementDataStorage::SubqueryCallbackRef Callback)
+			UE::Editor::DataStorage::FQueryResult RunSubquery(int32 SubqueryIndex, RowHandle Row,
+				UE::Editor::DataStorage::SubqueryCallbackRef Callback)
 			{
 				if (SubqueryIndex < QueryDescription.Subqueries.Num())
 				{
@@ -768,7 +768,7 @@ namespace UE::Editor::DataStorage
 				}
 				else
 				{
-					return TypedElementDataStorage::FQueryResult{};
+					return UE::Editor::DataStorage::FQueryResult{};
 				}
 			}
 	
@@ -803,11 +803,11 @@ namespace UE::Editor::DataStorage
 			void RemoveRows(TConstArrayView<RowHandle> Rows) override  { return Implementation.RemoveRows(Rows); }
 			void AddColumns(RowHandle Row, TConstArrayView<const UScriptStruct*> ColumnTypes) override { return Implementation.AddColumns(Row, ColumnTypes); }
 			void AddColumns(TConstArrayView<RowHandle> Rows, TConstArrayView<const UScriptStruct*> ColumnTypes) override  { return Implementation.AddColumns(Rows, ColumnTypes); }
-			void AddColumns(TConstArrayView<TypedElementDataStorage::RowHandle> Rows, TConstArrayView<UE::Editor::DataStorage::FDynamicColumnDescription> DynamicColumnDescriptions) override { return Implementation.AddColumns(Rows, DynamicColumnDescriptions); }
+			void AddColumns(TConstArrayView<RowHandle> Rows, TConstArrayView<UE::Editor::DataStorage::FDynamicColumnDescription> DynamicColumnDescriptions) override { return Implementation.AddColumns(Rows, DynamicColumnDescriptions); }
 			void* AddColumnUninitialized(RowHandle Row, const UScriptStruct* ColumnType) override  { return Implementation.AddColumnUninitialized(Row, ColumnType); }
 			void* AddColumnUninitialized(RowHandle Row, const UScriptStruct* ObjectType, ObjectCopyOrMove Relocator) override { return Implementation.AddColumnUninitialized(Row, ObjectType, Relocator); }
-			void* AddColumnUninitialized(TypedElementDataStorage::RowHandle Row, const FDynamicColumnDescription& DynamicColumnDescription, ObjectCopyOrMove Relocator) override { return Implementation.AddColumnUninitialized(Row, DynamicColumnDescription, Relocator); }
-			void* AddColumnUninitialized(TypedElementDataStorage::RowHandle Row, const FDynamicColumnDescription& DynamicColumnDescription) override { return Implementation.AddColumnUninitialized(Row, DynamicColumnDescription); };
+			void* AddColumnUninitialized(RowHandle Row, const FDynamicColumnDescription& DynamicColumnDescription, ObjectCopyOrMove Relocator) override { return Implementation.AddColumnUninitialized(Row, DynamicColumnDescription, Relocator); }
+			void* AddColumnUninitialized(RowHandle Row, const FDynamicColumnDescription& DynamicColumnDescription) override { return Implementation.AddColumnUninitialized(Row, DynamicColumnDescription); };
 			void RemoveColumns(RowHandle Row, TConstArrayView<const UScriptStruct*> ColumnTypes) override { Implementation.RemoveColumns(Row, ColumnTypes); }
 			void RemoveColumns(TConstArrayView<RowHandle> Rows, TConstArrayView<const UScriptStruct*> ColumnTypes) override { Implementation.RemoveColumns(Rows, ColumnTypes); }
 			const UScriptStruct* FindDynamicColumnType(const UE::Editor::DataStorage::FDynamicColumnDescription& Description) const override { return Implementation.FindDynamicColumnType(Description); }
@@ -818,8 +818,8 @@ namespace UE::Editor::DataStorage
 			RowHandle FindIndexedRow(TypedElementDataStorage::IndexHash Index) const override { return Implementation.FindIndexedRow(Index); }
 			TypedElementDataStorage::FQueryResult RunQuery(QueryHandle Query) override { return Implementation.RunQuery(Query); }
 			TypedElementDataStorage::FQueryResult RunSubquery(int32 SubqueryIndex) override { return Implementation.RunSubquery(SubqueryIndex); }
-			TypedElementDataStorage::FQueryResult RunSubquery(int32 SubqueryIndex, TypedElementDataStorage::SubqueryCallbackRef Callback) override { return Implementation.RunSubquery(SubqueryIndex, Callback); }
-			TypedElementDataStorage::FQueryResult RunSubquery(int32 SubqueryIndex, RowHandle Row, TypedElementDataStorage::SubqueryCallbackRef Callback) override { return Implementation.RunSubquery(SubqueryIndex, Row, Callback); }
+			TypedElementDataStorage::FQueryResult RunSubquery(int32 SubqueryIndex, UE::Editor::DataStorage::SubqueryCallbackRef Callback) override { return Implementation.RunSubquery(SubqueryIndex, Callback); }
+			TypedElementDataStorage::FQueryResult RunSubquery(int32 SubqueryIndex, RowHandle Row, UE::Editor::DataStorage::SubqueryCallbackRef Callback) override { return Implementation.RunSubquery(SubqueryIndex, Row, Callback); }
 	
 			FMassQueryContextImplementation Implementation;
 		};
@@ -1043,15 +1043,15 @@ bool FTypedElementQueryProcessorData::PrepareCachedDependenciesOnQuery(
 	return true;
 }
 
-TypedElementDataStorage::FQueryResult FTypedElementQueryProcessorData::Execute(
-	TypedElementDataStorage::DirectQueryCallbackRef& Callback,
-	TypedElementDataStorage::FQueryDescription& Description,
+UE::Editor::DataStorage::FQueryResult FTypedElementQueryProcessorData::Execute(
+	UE::Editor::DataStorage::DirectQueryCallbackRef& Callback,
+	UE::Editor::DataStorage::FQueryDescription& Description,
 	FMassEntityQuery& NativeQuery, 
 	FMassEntityManager& EntityManager,
 	FEnvironment& Environment,
-	TypedElementDataStorage::EDirectQueryExecutionFlags ExecutionFlags)
+	UE::Editor::DataStorage::EDirectQueryExecutionFlags ExecutionFlags)
 {
-	using namespace TypedElementDataStorage;
+	using namespace UE::Editor::DataStorage;
 
 	FQueryResult Result;
 	Result.Completed = FQueryResult::ECompletion::Fully;
@@ -1085,15 +1085,15 @@ TypedElementDataStorage::FQueryResult FTypedElementQueryProcessorData::Execute(
 	return Result;
 }
 
-TypedElementDataStorage::FQueryResult FTypedElementQueryProcessorData::Execute(
-	TypedElementDataStorage::SubqueryCallbackRef& Callback,
-	TypedElementDataStorage::FQueryDescription& Description,
+UE::Editor::DataStorage::FQueryResult FTypedElementQueryProcessorData::Execute(
+	UE::Editor::DataStorage::SubqueryCallbackRef& Callback,
+	UE::Editor::DataStorage::FQueryDescription& Description,
 	FMassEntityQuery& NativeQuery,
 	FMassEntityManager& EntityManager,
 	FEnvironment& Environment,
 	FMassExecutionContext& ParentContext)
 {
-	using namespace TypedElementDataStorage;
+	using namespace UE::Editor::DataStorage;
 
 	ITypedElementDataStorageInterface::FQueryResult Result;
 	Result.Completed = ITypedElementDataStorageInterface::FQueryResult::ECompletion::Fully;
@@ -1120,16 +1120,16 @@ TypedElementDataStorage::FQueryResult FTypedElementQueryProcessorData::Execute(
 	return Result;
 }
 
-TypedElementDataStorage::FQueryResult FTypedElementQueryProcessorData::Execute(
-	TypedElementDataStorage::SubqueryCallbackRef& Callback,
-	TypedElementDataStorage::FQueryDescription& Description,
-	TypedElementDataStorage::RowHandle RowHandle,
+UE::Editor::DataStorage::FQueryResult FTypedElementQueryProcessorData::Execute(
+	UE::Editor::DataStorage::SubqueryCallbackRef& Callback,
+	UE::Editor::DataStorage::FQueryDescription& Description,
+	UE::Editor::DataStorage::RowHandle RowHandle,
 	FMassEntityQuery& NativeQuery,
 	FMassEntityManager& EntityManager,
 	FEnvironment& Environment,
 	FMassExecutionContext& ParentContext)
 {
-	using namespace TypedElementDataStorage;
+	using namespace UE::Editor::DataStorage;
 
 	ITypedElementDataStorageInterface::FQueryResult Result;
 	Result.Completed = ITypedElementDataStorageInterface::FQueryResult::ECompletion::Fully;
@@ -1150,7 +1150,7 @@ TypedElementDataStorage::FQueryResult FTypedElementQueryProcessorData::Execute(
 			[&Result, &Callback, &Description, &Environment, &EntityManager, RowHandle](FMassExecutionContext& Context)
 			{
 				// No need to cache any subsystem dependencies as these are not accessible from a subquery.
-				UE::Editor::DataStorage::Processors::Private::FMassSubqueryContextForwarder QueryContext(Context, Environment);
+				Processors::Private::FMassSubqueryContextForwarder QueryContext(Context, Environment);
 				Callback(Description, QueryContext);
 				Result.Count += Context.GetNumEntities();
 			}
@@ -1162,7 +1162,7 @@ TypedElementDataStorage::FQueryResult FTypedElementQueryProcessorData::Execute(
 
 void FTypedElementQueryProcessorData::Execute(FMassEntityManager& EntityManager, FMassExecutionContext& Context)
 {
-	using namespace TypedElementDataStorage;
+	using namespace UE::Editor::DataStorage;
 	
 	FExtendedQuery* StoredQuery = QueryStore->GetMutable(ParentQuery);
 
@@ -1175,7 +1175,7 @@ void FTypedElementQueryProcessorData::Execute(FMassEntityManager& EntityManager,
 			{
 				if (PrepareCachedDependenciesOnQuery(Description, Context))
 				{
-					UE::Editor::DataStorage::Processors::Private::FMassContextForwarder QueryContext(Description, Context, *QueryStore, *Environment);
+					Processors::Private::FMassContextForwarder QueryContext(Description, Context, *QueryStore, *Environment);
 					Description.Callback.Function(Description, QueryContext);
 				}
 			};
diff --git a/Engine/Plugins/Experimental/EditorDataStorage/Source/TedsCore/Private/Processors/TypedElementProcessorAdaptors.h b/Engine/Plugins/Experimental/EditorDataStorage/Source/TedsCore/Private/Processors/TypedElementProcessorAdaptors.h
--- a/Engine/Plugins/Experimental/EditorDataStorage/Source/TedsCore/Private/Processors/TypedElementProcessorAdaptors.h
+++ b/Engine/Plugins/Experimental/EditorDataStorage/Source/TedsCore/Private/Processors/TypedElementProcessorAdaptors.h
@@ -57,24 +57,24 @@ struct FTypedElementQueryProcessorData
 	FString GetProcessorName() const;
 	void DebugOutputDescription(FOutputDevice& Ar, int32 Indent) const;
 
-	static TypedElementDataStorage::FQueryResult Execute(
-		TypedElementDataStorage::DirectQueryCallbackRef& Callback,
-		TypedElementDataStorage::FQueryDescription& Description,
+	static UE::Editor::DataStorage::FQueryResult Execute(
+		UE::Editor::DataStorage::DirectQueryCallbackRef& Callback,
+		UE::Editor::DataStorage::FQueryDescription& Description,
 		FMassEntityQuery& NativeQuery, 
 		FMassEntityManager& EntityManager,
 		FEnvironment& Environment,
-		TypedElementDataStorage::EDirectQueryExecutionFlags ExecutionFlags);
-	static TypedElementDataStorage::FQueryResult Execute(
-		TypedElementDataStorage::SubqueryCallbackRef& Callback,
-		TypedElementDataStorage::FQueryDescription& Description,
+		UE::Editor::DataStorage::EDirectQueryExecutionFlags ExecutionFlags);
+	static UE::Editor::DataStorage::FQueryResult Execute(
+		UE::Editor::DataStorage::SubqueryCallbackRef& Callback,
+		UE::Editor::DataStorage::FQueryDescription& Description,
 		FMassEntityQuery& NativeQuery,
 		FMassEntityManager& EntityManager,
 		FEnvironment& Environment,
 		FMassExecutionContext& ParentContext);
-	static TypedElementDataStorage::FQueryResult Execute(
-		TypedElementDataStorage::SubqueryCallbackRef& Callback,
-		TypedElementDataStorage::FQueryDescription& Description,
-		TypedElementDataStorage::RowHandle RowHandle,
+	static UE::Editor::DataStorage::FQueryResult Execute(
+		UE::Editor::DataStorage::SubqueryCallbackRef& Callback,
+		UE::Editor::DataStorage::FQueryDescription& Description,
+		UE::Editor::DataStorage::RowHandle RowHandle,
 		FMassEntityQuery& NativeQuery,
 		FMassEntityManager& EntityManager,
 		FEnvironment& Environment,
diff --git a/Engine/Plugins/Experimental/EditorDataStorage/Source/TedsCore/Private/Queries/TypedElementAlertQueries.cpp b/Engine/Plugins/Experimental/EditorDataStorage/Source/TedsCore/Private/Queries/TypedElementAlertQueries.cpp
--- a/Engine/Plugins/Experimental/EditorDataStorage/Source/TedsCore/Private/Queries/TypedElementAlertQueries.cpp
+++ b/Engine/Plugins/Experimental/EditorDataStorage/Source/TedsCore/Private/Queries/TypedElementAlertQueries.cpp
@@ -335,14 +335,14 @@ void UTypedElementAlertQueriesFactory::RegisterOnRemoveQueries(ITypedElementData
 }
 
 void UTypedElementAlertQueriesFactory::AddChildAlertsToHierarchy(
-	TypedElementDataStorage::IQueryContext& Context, UE::Editor::DataStorage::RowHandle Parent, int32 ParentQueryIndex)
+	UE::Editor::DataStorage::IQueryContext& Context, UE::Editor::DataStorage::RowHandle Parent, int32 ParentQueryIndex)
 {
-	using namespace TypedElementDataStorage;
+	using namespace UE::Editor::DataStorage;
 
 	bool bHasParent = true;
 	do
 	{
-		TypedElementDataStorage::RowHandle NextParent = Parent;
+		RowHandle NextParent = Parent;
 		bHasParent = MoveToNextParent(NextParent, Context, ParentQueryIndex);
 
 		// Increment parent's child counter or create a child alert if one doesn't exist.
@@ -358,7 +358,7 @@ void UTypedElementAlertQueriesFactory::AddChildAlertsToHierarchy(
 }
 
 void UTypedElementAlertQueriesFactory::IncrementParents(
-	TypedElementDataStorage::IQueryContext& Context, UE::Editor::DataStorage::RowHandle Row, FTypedElementAlertColumnType AlertType,
+	UE::Editor::DataStorage::IQueryContext& Context, UE::Editor::DataStorage::RowHandle Row, FTypedElementAlertColumnType AlertType,
 	int32 ChildAlertQueryIndex)
 {
 	using namespace TypedElementDataStorage;
@@ -388,7 +388,7 @@ void UTypedElementAlertQueriesFactory::ResetChildAlertCounters(FTypedElementChil
 }
 
 bool UTypedElementAlertQueriesFactory::MoveToNextParent(
-	TypedElementDataStorage::RowHandle& Parent, TypedElementDataStorage::IQueryContext& Context, int32 SubQueryIndex)
+	UE::Editor::DataStorage::RowHandle& Parent, UE::Editor::DataStorage::IQueryContext& Context, int32 SubQueryIndex)
 {
 	using namespace TypedElementDataStorage;
 	using namespace TypedElementQueryBuilder;
diff --git a/Engine/Plugins/Experimental/EditorDataStorage/Source/TedsCore/Private/Queries/TypedElementAlertQueries.h b/Engine/Plugins/Experimental/EditorDataStorage/Source/TedsCore/Private/Queries/TypedElementAlertQueries.h
--- a/Engine/Plugins/Experimental/EditorDataStorage/Source/TedsCore/Private/Queries/TypedElementAlertQueries.h
+++ b/Engine/Plugins/Experimental/EditorDataStorage/Source/TedsCore/Private/Queries/TypedElementAlertQueries.h
@@ -11,7 +11,7 @@
 
 #include "TypedElementAlertQueries.generated.h"
 
-namespace TypedElementDataStorage
+namespace UE::Editor::DataStorage
 {
 	struct IQueryContext;
 }
@@ -37,19 +37,19 @@ private:
 	void RegisterOnRemoveQueries(ITypedElementDataStorageInterface& DataStorage);
 
 	static void AddChildAlertsToHierarchy(
-		TypedElementDataStorage::IQueryContext& Context, TypedElementDataStorage::RowHandle Parent, int32 ParentQueryIndex);
+		UE::Editor::DataStorage::IQueryContext& Context, UE::Editor::DataStorage::RowHandle Parent, int32 ParentQueryIndex);
 
 	static void IncrementParents(
-		TypedElementDataStorage::IQueryContext& Context, TypedElementDataStorage::RowHandle Row, FTypedElementAlertColumnType AlertType,
+		UE::Editor::DataStorage::IQueryContext& Context, UE::Editor::DataStorage::RowHandle Row, FTypedElementAlertColumnType AlertType,
 		int32 ChildAlertQueryIndex);
 	
 	static void ResetChildAlertCounters(FTypedElementChildAlertColumn& ChildAlert);
 
 	static bool MoveToNextParent(
-		TypedElementDataStorage::RowHandle& Parent, TypedElementDataStorage::IQueryContext& Context, int32 SubQueryIndex);
+		UE::Editor::DataStorage::RowHandle& Parent, UE::Editor::DataStorage::IQueryContext& Context, int32 SubQueryIndex);
 
 	static const FName AlertConditionName;
-	TypedElementDataStorage::QueryHandle ChildAlertColumnReadWriteQuery;
-	TypedElementDataStorage::QueryHandle ParentReadOnlyQuery;
+	UE::Editor::DataStorage::QueryHandle ChildAlertColumnReadWriteQuery;
+	UE::Editor::DataStorage::QueryHandle ParentReadOnlyQuery;
 };
 
diff --git a/Engine/Plugins/Experimental/EditorDataStorage/Source/TedsCore/Private/Queries/TypedElementExtendedQueryStore.cpp b/Engine/Plugins/Experimental/EditorDataStorage/Source/TedsCore/Private/Queries/TypedElementExtendedQueryStore.cpp
--- a/Engine/Plugins/Experimental/EditorDataStorage/Source/TedsCore/Private/Queries/TypedElementExtendedQueryStore.cpp
+++ b/Engine/Plugins/Experimental/EditorDataStorage/Source/TedsCore/Private/Queries/TypedElementExtendedQueryStore.cpp
@@ -212,15 +212,15 @@ namespace UE::Editor::DataStorage
 	}
 
 	template<typename CallbackReference>
-	TypedElementDataStorage::FQueryResult FExtendedQueryStore::RunQueryCallbackCommon(
+	FQueryResult FExtendedQueryStore::RunQueryCallbackCommon(
 		FMassEntityManager& EntityManager,
 		FEnvironment& Environment,
 		FMassExecutionContext* ParentContext,
 		Handle Query,
-		TypedElementDataStorage::EDirectQueryExecutionFlags DirectExecutionFlags,
+		EDirectQueryExecutionFlags DirectExecutionFlags,
 		CallbackReference Callback)
 	{
-		using namespace TypedElementDataStorage;
+		using namespace UE::Editor::DataStorage;
 
 		FQueryResult Result;
 		if (FExtendedQuery* QueryData = Get(Query))
@@ -267,24 +267,24 @@ namespace UE::Editor::DataStorage
 		return Result;
 	}
 
-	TypedElementDataStorage::FQueryResult FExtendedQueryStore::RunQuery(FMassEntityManager& EntityManager,
-		FEnvironment& Environment, Handle Query, TypedElementDataStorage::EDirectQueryExecutionFlags DirectExecutionFlags,
-		TypedElementDataStorage::DirectQueryCallbackRef Callback)
+	FQueryResult FExtendedQueryStore::RunQuery(FMassEntityManager& EntityManager,
+		FEnvironment& Environment, Handle Query, EDirectQueryExecutionFlags DirectExecutionFlags,
+		DirectQueryCallbackRef Callback)
 	{
 		return RunQueryCallbackCommon(EntityManager, Environment, nullptr, Query, DirectExecutionFlags, Callback);
 	}
 	
-	TypedElementDataStorage::FQueryResult FExtendedQueryStore::RunQuery(FMassEntityManager& EntityManager,
+	FQueryResult FExtendedQueryStore::RunQuery(FMassEntityManager& EntityManager,
 		FEnvironment& Environment, FMassExecutionContext& ParentContext, Handle Query,
-		TypedElementDataStorage::SubqueryCallbackRef Callback)
+		SubqueryCallbackRef Callback)
 	{
 		return RunQueryCallbackCommon(EntityManager, Environment, &ParentContext, Query, 
-			TypedElementDataStorage::EDirectQueryExecutionFlags::Default, Callback);
+			EDirectQueryExecutionFlags::Default, Callback);
 	}
 
-	TypedElementDataStorage::FQueryResult FExtendedQueryStore::RunQuery(FMassEntityManager& EntityManager,
+	FQueryResult FExtendedQueryStore::RunQuery(FMassEntityManager& EntityManager,
 		FEnvironment& Environment, FMassExecutionContext& ParentContext, Handle Query, RowHandle Row,
-		TypedElementDataStorage::SubqueryCallbackRef Callback)
+		SubqueryCallbackRef Callback)
 	{
 		using ActionType = ITypedElementDataStorageInterface::FQueryDescription::EActionType;
 		using CompletionType = ITypedElementDataStorageInterface::FQueryResult::ECompletion;
@@ -393,21 +393,21 @@ namespace UE::Editor::DataStorage
 			{
 				Query.SelectionTypes.Emplace(DynamicColumnType);
 				Query.SelectionAccessTypes.Emplace(AccessType);
-				Query.SelectionMetaData.Emplace(TypedElementDataStorage::FColumnMetaData(DynamicColumnType, MetadataFlags));
+				Query.SelectionMetaData.Emplace(FColumnMetaData(DynamicColumnType, MetadataFlags));
 			}
 		}
 
 		for (int32 Index = 0, End = Query.DynamicConditionDescriptions.Num(); Index < End; ++Index)
 		{
-			const UE::Editor::DataStorage::FDynamicColumnDescription& Description = Query.DynamicConditionDescriptions[Index];
+			const FDynamicColumnDescription& Description = Query.DynamicConditionDescriptions[Index];
 			if (!ensureMsgf(Description.TemplateType, TEXT("Null template type for dynamic column")))
 			{
 				continue;
 			}
 			const UScriptStruct* DynamicColumnType = Environment.GenerateDynamicColumn(*Description.TemplateType, Description.Identifier);
-			TypedElementDataStorage::FQueryDescription::EOperatorType Operation = Query.DynamicConditionOperations[Index];
+			FQueryDescription::EOperatorType Operation = Query.DynamicConditionOperations[Index];
 			Query.ConditionTypes.Add(Operation);
-			Query.ConditionOperators.Add(TypedElementDataStorage::FQueryDescription::FOperator
+			Query.ConditionOperators.Add(FQueryDescription::FOperator
 			{
 				.Type = DynamicColumnType
 			});
@@ -551,9 +551,9 @@ namespace UE::Editor::DataStorage
 					DSI::EQueryAccessType AccessType = Query.SelectionAccessTypes[SelectionIndex];
 					if (ensureMsgf(Type.IsValid(), TEXT("Provided query selection type can not be null.")) &&
 						ensureMsgf(
-							Type->IsChildOf(UE::Editor::DataStorage::FColumn::StaticStruct()) ||
+							Type->IsChildOf(FColumn::StaticStruct()) ||
 							Type->IsChildOf(FMassFragment::StaticStruct()),
-							TEXT("Provided query selection type '%s' is not based on UE::Editor::DataStorage::FColumn or another supported base type."),
+							TEXT("Provided query selection type '%s' is not based on FColumn or another supported base type."),
 							*Type->GetStructPathName().ToString()))
 					{
 						NativeQuery.AddRequirement(Type.Get(), ConvertToNativeAccessType(AccessType), ConvertToNativePresenceType(AccessType));
@@ -640,12 +640,12 @@ namespace UE::Editor::DataStorage
 			return true;
 		}
 
-		Algo::SortBy(Query.ValueTags, [](const TypedElementDataStorage::FQueryDescription::FValueTagData& ValueTagData)
+		Algo::SortBy(Query.ValueTags, [](const FQueryDescription::FValueTagData& ValueTagData)
 			{
 				return ValueTagData.Tag.GetName();
 			}, FNameFastLess());
 		// Check if there are any duplicate groups. Not yet supported until we can match multiple MatchTags
-		UE::Editor::DataStorage::FValueTag PreviousTag = Query.ValueTags[0].Tag;
+	FValueTag PreviousTag = Query.ValueTags[0].Tag;
 		for (int32 Index = 1, End = Query.ValueTags.Num(); Index < End; ++Index)
 		{
 			if (Query.ValueTags[Index].Tag == PreviousTag)
diff --git a/Engine/Plugins/Experimental/EditorDataStorage/Source/TedsCore/Private/Queries/TypedElementExtendedQueryStore.h b/Engine/Plugins/Experimental/EditorDataStorage/Source/TedsCore/Private/Queries/TypedElementExtendedQueryStore.h
--- a/Engine/Plugins/Experimental/EditorDataStorage/Source/TedsCore/Private/Queries/TypedElementExtendedQueryStore.h
+++ b/Engine/Plugins/Experimental/EditorDataStorage/Source/TedsCore/Private/Queries/TypedElementExtendedQueryStore.h
@@ -101,26 +101,26 @@ namespace UE::Editor::DataStorage
 		 * @description Various functions to run queries.
 		 */
 
-		TypedElementDataStorage::FQueryResult RunQuery(FMassEntityManager& EntityManager, Handle Query);
-		TypedElementDataStorage::FQueryResult RunQuery(
+		FQueryResult RunQuery(FMassEntityManager& EntityManager, Handle Query);
+		FQueryResult RunQuery(
 			FMassEntityManager& EntityManager,
 			FEnvironment& Environment,
 			Handle Query,
-			TypedElementDataStorage::EDirectQueryExecutionFlags DirectExecutionFlags,
-			TypedElementDataStorage::DirectQueryCallbackRef Callback);
-		TypedElementDataStorage::FQueryResult RunQuery(
+			EDirectQueryExecutionFlags DirectExecutionFlags,
+			DirectQueryCallbackRef Callback);
+		FQueryResult RunQuery(
 			FMassEntityManager& EntityManager,
 			FEnvironment& Environment,
 			FMassExecutionContext& ParentContext,
 			Handle Query,
-			TypedElementDataStorage::SubqueryCallbackRef Callback);
-		TypedElementDataStorage::FQueryResult RunQuery(
+			SubqueryCallbackRef Callback);
+		FQueryResult RunQuery(
 			FMassEntityManager& EntityManager,
 			FEnvironment& Environment,
 			FMassExecutionContext& ParentContext,
 			Handle Query,
 			DataStorage::RowHandle Row,
-			TypedElementDataStorage::SubqueryCallbackRef Callback);
+			SubqueryCallbackRef Callback);
 		void RunPhasePreambleQueries(
 			FMassEntityManager& EntityManager,
 			FEnvironment& Environment,
@@ -156,12 +156,12 @@ namespace UE::Editor::DataStorage
 		};
 
 		template<typename CallbackReference>
-		TypedElementDataStorage::FQueryResult RunQueryCallbackCommon(
+		FQueryResult RunQueryCallbackCommon(
 			FMassEntityManager& EntityManager,
 			FEnvironment& Environment,
 			FMassExecutionContext* ParentContext,
 			Handle Query,
-		TypedElementDataStorage::EDirectQueryExecutionFlags ExecutionFlags,
+		EDirectQueryExecutionFlags ExecutionFlags,
 			CallbackReference Callback);
 
 		FMassEntityQuery& SetupNativeQuery(ITypedElementDataStorageInterface::FQueryDescription& Query, FExtendedQuery& StoredQuery);
diff --git a/Engine/Plugins/Experimental/EditorDataStorage/Source/TedsCore/Private/TypedElementDatabase.cpp b/Engine/Plugins/Experimental/EditorDataStorage/Source/TedsCore/Private/TypedElementDatabase.cpp
--- a/Engine/Plugins/Experimental/EditorDataStorage/Source/TedsCore/Private/TypedElementDatabase.cpp
+++ b/Engine/Plugins/Experimental/EditorDataStorage/Source/TedsCore/Private/TypedElementDatabase.cpp
@@ -266,9 +266,9 @@ TSharedPtr<const FMassEntityManager> UEditorDataStorage::GetActiveEditorEntityMa
 	return ActiveEditorEntityManager;
 }
 
-TypedElementDataStorage::TableHandle UEditorDataStorage::RegisterTable(TConstArrayView<const UScriptStruct*> ColumnList, const FName Name)
+UE::Editor::DataStorage::TableHandle UEditorDataStorage::RegisterTable(TConstArrayView<const UScriptStruct*> ColumnList, const FName Name)
 {
-	using namespace TypedElementDataStorage;
+	using namespace UE::Editor::DataStorage;
 
 	if (ActiveEditorEntityManager && !TableNameLookup.Contains(Name))
 	{
@@ -286,7 +286,7 @@ TypedElementDataStorage::TableHandle UEditorDataStorage::RegisterTable(TConstArr
 	return InvalidTableHandle;
 }
 
-UE::Editor::DataStorage::TableHandle UEditorDataStorage::RegisterTable(TypedElementDataStorage::TableHandle SourceTable,
+UE::Editor::DataStorage::TableHandle UEditorDataStorage::RegisterTable(UE::Editor::DataStorage::TableHandle SourceTable,
 	TConstArrayView<const UScriptStruct*> ColumnList, const FName Name)
 {
 	using namespace UE::Editor::DataStorage;
@@ -1048,7 +1048,7 @@ ITypedElementDataStorageInterface::FQueryResult UEditorDataStorage::RunQuery(
 	{
 		const UE::Editor::DataStorage::FExtendedQueryStore::Handle StorageHandle(Query);
 		return Environment->GetQueryStore().RunQuery(*ActiveEditorEntityManager, *Environment, StorageHandle, 
-			TypedElementDataStorage::EDirectQueryExecutionFlags::Default, Callback);
+			UE::Editor::DataStorage::EDirectQueryExecutionFlags::Default, Callback);
 	}
 	else
 	{
@@ -1057,7 +1057,7 @@ ITypedElementDataStorageInterface::FQueryResult UEditorDataStorage::RunQuery(
 }
 
 ITypedElementDataStorageInterface::FQueryResult UEditorDataStorage::RunQuery(
-	QueryHandle Query, TypedElementDataStorage::EDirectQueryExecutionFlags Flags,
+	QueryHandle Query, UE::Editor::DataStorage::EDirectQueryExecutionFlags Flags,
 	ITypedElementDataStorageInterface::DirectQueryCallbackRef Callback)
 {
 	TRACE_CPUPROFILER_EVENT_SCOPE(TEDS.RunQuery);
@@ -1225,7 +1225,7 @@ TSharedPtr<const UE::Editor::DataStorage::FEnvironment> UEditorDataStorage::GetE
 	return Environment;
 }
 
-FMassArchetypeHandle UEditorDataStorage::LookupArchetype(TypedElementDataStorage::TableHandle InTableHandle) const
+FMassArchetypeHandle UEditorDataStorage::LookupArchetype(UE::Editor::DataStorage::TableHandle InTableHandle) const
 {
 	const uint32 TableIndex = InTableHandle;
 	if (Tables.IsValidIndex(TableIndex))
diff --git a/Engine/Plugins/Experimental/EditorDataStorage/Source/TedsCore/Private/TypedElementDatabase.h b/Engine/Plugins/Experimental/EditorDataStorage/Source/TedsCore/Private/TypedElementDatabase.h
--- a/Engine/Plugins/Experimental/EditorDataStorage/Source/TedsCore/Private/TypedElementDatabase.h
+++ b/Engine/Plugins/Experimental/EditorDataStorage/Source/TedsCore/Private/TypedElementDatabase.h
@@ -148,7 +148,7 @@ public:
 	virtual FName GetQueryTickGroupName(EQueryTickGroups Group) const override;
 	virtual FQueryResult RunQuery(QueryHandle Query) override;
 	virtual FQueryResult RunQuery(QueryHandle Query, DirectQueryCallbackRef Callback) override;
-	virtual FQueryResult RunQuery(QueryHandle Query, TypedElementDataStorage::EDirectQueryExecutionFlags Flags, 
+	virtual FQueryResult RunQuery(QueryHandle Query, UE::Editor::DataStorage::EDirectQueryExecutionFlags Flags,
 		DirectQueryCallbackRef Callback) override;
 	virtual void ActivateQueries(FName ActivationName) override;
 
diff --git a/Engine/Plugins/Experimental/EditorDataStorage/Source/TedsCore/Private/TypedElementDatabaseCommandBuffer.cpp b/Engine/Plugins/Experimental/EditorDataStorage/Source/TedsCore/Private/TypedElementDatabaseCommandBuffer.cpp
--- a/Engine/Plugins/Experimental/EditorDataStorage/Source/TedsCore/Private/TypedElementDatabaseCommandBuffer.cpp
+++ b/Engine/Plugins/Experimental/EditorDataStorage/Source/TedsCore/Private/TypedElementDatabaseCommandBuffer.cpp
@@ -9,6 +9,7 @@
 
 namespace UE::Editor::DataStorage::Legacy
 {
+	using namespace UE::Editor::DataStorage;
 	// 
 	// Commands section
 	//
@@ -19,7 +20,7 @@ namespace UE::Editor::DataStorage::Legacy
 	}
 
 	void* FCommandBuffer::GetQueuedDataColumn(
-		TypedElementDataStorage::RowHandle Row, const UScriptStruct* ColumnType)
+		RowHandle Row, const UScriptStruct* ColumnType)
 	{
 		checkf(ColumnType->IsChildOf(FMassFragment::StaticStruct()),
 			TEXT("Trying to get the column '%s' which isn't a data column."), *ColumnType->GetName());
@@ -50,12 +51,12 @@ namespace UE::Editor::DataStorage::Legacy
 		}
 	}
 
-	bool FCommandBuffer::HasColumn(TypedElementDataStorage::RowHandle Row, const UScriptStruct* ColumnType) const
+	bool FCommandBuffer::HasColumn(RowHandle Row, const UScriptStruct* ColumnType) const
 	{
 		return PendingColumns.Contains(PendingColumnMappingKey(Row, ColumnType));
 	}
 
-	void FCommandBuffer::Clear(TypedElementDataStorage::RowHandle Row)
+	void FCommandBuffer::Clear(RowHandle Row)
 	{
 		for (TMap<PendingColumnMappingKey, void*>::TIterator It = PendingColumns.CreateIterator(); It; ++It)
 		{
@@ -86,7 +87,7 @@ namespace UE::Editor::DataStorage::Legacy
 	// Queue section
 	//
 
-	void FCommandBuffer::Queue_AddColumnCommand(TypedElementDataStorage::RowHandle Row, const UScriptStruct* ColumnType)
+	void FCommandBuffer::Queue_AddColumnCommand(RowHandle Row, const UScriptStruct* ColumnType)
 	{
 		AddCommand(Row, FAddColumnCommand
 			{
@@ -96,7 +97,7 @@ namespace UE::Editor::DataStorage::Legacy
 	}
 
 	void* FCommandBuffer::Queue_AddDataColumnCommandUnitialized(
-		TypedElementDataStorage::RowHandle Row, const UScriptStruct* ColumnType, ColumnCopyOrMoveCallback Relocator)
+		RowHandle Row, const UScriptStruct* ColumnType, ColumnCopyOrMoveCallback Relocator)
 	{
 		checkf(ColumnType->IsChildOf(FMassFragment::StaticStruct()),
 			TEXT("Trying to queue a data column creation for '%s' which isn't a data column."), *ColumnType->GetName());
@@ -122,7 +123,7 @@ namespace UE::Editor::DataStorage::Legacy
 		}
 	}
 
-	void FCommandBuffer::Queue_AddColumnsCommand(TypedElementDataStorage::RowHandle Row,
+	void FCommandBuffer::Queue_AddColumnsCommand(RowHandle Row,
 		FMassFragmentBitSet FragmentsToAdd, FMassTagBitSet TagsToAdd)
 	{
 		auto AddColumn = [this, Row](const UScriptStruct* ColumnType)
@@ -140,7 +141,7 @@ namespace UE::Editor::DataStorage::Legacy
 			});
 	}
 
-	void FCommandBuffer::Queue_RemoveColumnCommand(TypedElementDataStorage::RowHandle Row, const UScriptStruct* ColumnType)
+	void FCommandBuffer::Queue_RemoveColumnCommand(RowHandle Row, const UScriptStruct* ColumnType)
 	{
 		AddCommand(Row, FRemoveColumnCommand
 			{
@@ -149,7 +150,7 @@ namespace UE::Editor::DataStorage::Legacy
 		PendingColumns.Remove(PendingColumnMappingKey(Row, ColumnType));
 	}
 
-	void FCommandBuffer::Queue_RemoveColumnsCommand(TypedElementDataStorage::RowHandle Row,
+	void FCommandBuffer::Queue_RemoveColumnsCommand(RowHandle Row,
 		FMassFragmentBitSet FragmentsToRemove, FMassTagBitSet TagsToRemove)
 	{
 		auto RemoveColumn = [this, Row](const UScriptStruct* ColumnType)
@@ -191,7 +192,7 @@ namespace UE::Editor::DataStorage::Legacy
 
 	void FCommandBuffer::Execute_AddColumnCommand(
 		FMassEntityManager& MassEntityManager,
-		TypedElementDataStorage::RowHandle Row,
+		RowHandle Row,
 		const UScriptStruct* ColumnType)
 	{
 		if (ColumnType)
@@ -260,7 +261,7 @@ namespace UE::Editor::DataStorage::Legacy
 
 	void FCommandBuffer::Execute_AddColumnsCommand(
 		FMassEntityManager& MassEntityManager,
-		TypedElementDataStorage::RowHandle Row,
+		RowHandle Row,
 		FMassFragmentBitSet FragmentsToAdd,
 		FMassTagBitSet TagsToAdd)
 	{
@@ -271,7 +272,7 @@ namespace UE::Editor::DataStorage::Legacy
 
 	void FCommandBuffer::Execute_RemoveColumnCommand(
 		FMassEntityManager& MassEntityManager,
-		TypedElementDataStorage::RowHandle Row,
+		RowHandle Row,
 		const UScriptStruct* ColumnType)
 	{
 		if (ColumnType)
@@ -339,7 +340,7 @@ namespace UE::Editor::DataStorage::Legacy
 	//
 
 	template<typename T>
-	void FCommandBuffer::AddCommand(TypedElementDataStorage::RowHandle Row, T&& Args)
+	void FCommandBuffer::AddCommand(RowHandle Row, T&& Args)
 	{
 		FCommand Command;
 		Command.Row = Row;
diff --git a/Engine/Plugins/Experimental/EditorDataStorage/Source/TedsCore/Private/TypedElementDatabaseCompatibility.cpp b/Engine/Plugins/Experimental/EditorDataStorage/Source/TedsCore/Private/TypedElementDatabaseCompatibility.cpp
--- a/Engine/Plugins/Experimental/EditorDataStorage/Source/TedsCore/Private/TypedElementDatabaseCompatibility.cpp
+++ b/Engine/Plugins/Experimental/EditorDataStorage/Source/TedsCore/Private/TypedElementDatabaseCompatibility.cpp
@@ -132,7 +132,7 @@ void UEditorDataStorageCompatibility::RegisterDealiaserCallback(ObjectToRowDeali
 }
 
 void UEditorDataStorageCompatibility::RegisterTypeTableAssociation(
-	TObjectPtr<UStruct> TypeInfo, TypedElementDataStorage::TableHandle Table)
+	TObjectPtr<UStruct> TypeInfo, UE::Editor::DataStorage::TableHandle Table)
 {
 	using namespace UE::Editor::DataStorage;
 	
@@ -296,7 +296,7 @@ UE::Editor::DataStorage::RowHandle UEditorDataStorageCompatibility::FindRowWithC
 {
 	// Thread safety is only needed by FindIndexedRow which internally takes care of it.
 
-	using namespace TypedElementDataStorage;
+	using namespace UE::Editor::DataStorage;
 
 	return (Object && Storage && Storage->IsAvailable()) ? Storage->FindIndexedRow(GenerateIndexHash(Object)) : InvalidRowHandle;
 }
@@ -400,7 +400,6 @@ void UEditorDataStorageCompatibility::RegisterTypeInformationQueries()
 
 bool UEditorDataStorageCompatibility::ShouldAddObject(const UObject* Object) const
 {
-	using namespace TypedElementDataStorage;
 	using namespace UE::Editor::DataStorage;
 	
 	FScopedSharedLock Lock(EGlobalLockScope::Public);
@@ -418,9 +417,8 @@ bool UEditorDataStorageCompatibility::ShouldAddObject(const UObject* Object) con
 	return Include;
 }
 
-TypedElementDataStorage::TableHandle UEditorDataStorageCompatibility::FindBestMatchingTable(const UStruct* TypeInfo) const
+UE::Editor::DataStorage::TableHandle UEditorDataStorageCompatibility::FindBestMatchingTable(const UStruct* TypeInfo) const
 {
-	using namespace TypedElementDataStorage;
 	using namespace UE::Editor::DataStorage;
 
 	FScopedSharedLock Lock(EGlobalLockScope::Public);
@@ -440,12 +438,11 @@ TypedElementDataStorage::TableHandle UEditorDataStorageCompatibility::FindBestMa
 template<bool bEnableTransactions>
 UE::Editor::DataStorage::RowHandle UEditorDataStorageCompatibility::AddCompatibleObjectExplicitTransactionable(UObject* Object)
 {
-	using namespace TypedElementDataStorage;
 	using namespace UE::Editor::DataStorage;
 
 	FScopedExclusiveLock Lock(EGlobalLockScope::Public);
 
-	UE::Editor::DataStorage::RowHandle Result = FindRowWithCompatibleObjectExplicit(Object);
+	RowHandle Result = FindRowWithCompatibleObjectExplicit(Object);
 	if (!Storage->IsRowAvailable(Result))
 	{
 		Result = Storage->ReserveRow();
@@ -473,7 +470,6 @@ UE::Editor::DataStorage::RowHandle UEditorDataStorageCompatibility::AddCompatibl
 template<bool bEnableTransactions>
 void UEditorDataStorageCompatibility::RemoveCompatibleObjectExplicitTransactionable(const UObject* Object)
 {
-	using namespace TypedElementDataStorage;
 	using namespace UE::Editor::DataStorage;
 
 	checkf(Storage,
@@ -504,9 +500,8 @@ void UEditorDataStorageCompatibility::RemoveCompatibleObjectExplicitTransactiona
 
 template<bool bEnableTransactions>
 void UEditorDataStorageCompatibility::RemoveCompatibleObjectExplicitTransactionable(
-	const UObject* Object, TypedElementDataStorage::RowHandle ObjectRow)
+	const UObject* Object, UE::Editor::DataStorage::RowHandle ObjectRow)
 {
-	using namespace TypedElementDataStorage;
 	using namespace UE::Editor::DataStorage;
 
 	checkf(Storage,
@@ -557,7 +552,7 @@ UE::Editor::DataStorage::RowHandle UEditorDataStorageCompatibility::DealiasObjec
 			return Row;
 		}
 	}
-	return TypedElementDataStorage::InvalidRowHandle;
+	return UE::Editor::DataStorage::InvalidRowHandle;
 }
 
 void UEditorDataStorageCompatibility::Tick()
@@ -612,7 +607,6 @@ void UEditorDataStorageCompatibility::FPendingTypeInformationUpdate::AddTypeInfo
 
 void UEditorDataStorageCompatibility::FPendingTypeInformationUpdate::Process(UEditorDataStorageCompatibility& Compatibility)
 {
-	using namespace TypedElementDataStorage;
 	using namespace TypedElementQueryBuilder;
 	using namespace UE::Editor::DataStorage;
 
@@ -635,7 +629,7 @@ void UEditorDataStorageCompatibility::FPendingTypeInformationUpdate::Process(UEd
 				It.RemoveCurrent();
 			}
 		}
-		for (TPair<TWeakObjectPtr<UStruct>, TypedElementDataStorage::TableHandle>& UpdatedEntry : UpdatedTypeInfoScratchBuffer)
+		for (TPair<TWeakObjectPtr<UStruct>, TableHandle>& UpdatedEntry : UpdatedTypeInfoScratchBuffer)
 		{
 			checkf(UpdatedEntry.Key.IsValid(),
 				TEXT("Type info column in data storage has been re-instanced to an object without type information"));
@@ -743,7 +737,7 @@ void UEditorDataStorageCompatibility::PendingRegistration<AddressType>::ProcessE
 	UEditorDataStorageCompatibility& Compatibility, const TFunctionRef<void(UE::Editor::DataStorage::RowHandle, const AddressType&)>& SetupRowCallback)
 {
 	// Thread-safe as it's only called from functions that already lock using.
-	using namespace TypedElementDataStorage;
+	using namespace UE::Editor::DataStorage;
 
 	// Start by removing any entries that are no longer valid.
 	for (auto It = Entries.CreateIterator(); It; ++It)
@@ -825,7 +819,7 @@ void UEditorDataStorageCompatibility::PendingRegistration<AddressType>::ProcessE
 			if (Current->Table != CurrentTable)
 			{
 				StorageInterface.BatchAddRow(CurrentTable, Compatibility.RowScratchBuffer,
-					[&SetupRowCallback, &TableFront](UE::Editor::DataStorage::RowHandle Row)
+					[&SetupRowCallback, &TableFront](RowHandle Row)
 					{
 						SetupRowCallback(Row, TableFront->Address);
 						++TableFront;
@@ -837,7 +831,7 @@ void UEditorDataStorageCompatibility::PendingRegistration<AddressType>::ProcessE
 			Compatibility.RowScratchBuffer.Add(Current->Row);
 		}
 		StorageInterface.BatchAddRow(CurrentTable, Compatibility.RowScratchBuffer,
-			[&SetupRowCallback, &TableFront](UE::Editor::DataStorage::RowHandle Row)
+			[&SetupRowCallback, &TableFront](RowHandle Row)
 			{
 				SetupRowCallback(Row, TableFront->Address);
 				++TableFront;
@@ -953,7 +947,7 @@ void UEditorDataStorageCompatibility::TickPendingExternalObjectRegistration()
 
 void UEditorDataStorageCompatibility::TickObjectSync()
 {
-	using namespace TypedElementDataStorage;
+	using namespace UE::Editor::DataStorage;
 	
 	// Thread safe because it's only called from functions that already lock.
 
@@ -1061,7 +1055,7 @@ void UEditorDataStorageCompatibility::OnObjectModified(UObject* Object)
 }
 
 void UEditorDataStorageCompatibility::TriggerOnObjectAdded(
-	const void* Object, UE::Editor::DataStorage::FObjectTypeInfo TypeInfo, TypedElementDataStorage::RowHandle Row) const
+	const void* Object, UE::Editor::DataStorage::FObjectTypeInfo TypeInfo, UE::Editor::DataStorage::RowHandle Row) const
 {
 	using namespace UE::Editor::DataStorage;
 
@@ -1075,7 +1069,7 @@ void UEditorDataStorageCompatibility::TriggerOnObjectAdded(
 }
 
 void UEditorDataStorageCompatibility::TriggerOnPreObjectRemoved(
-	const void* Object, UE::Editor::DataStorage::FObjectTypeInfo TypeInfo, TypedElementDataStorage::RowHandle Row) const
+	const void* Object, UE::Editor::DataStorage::FObjectTypeInfo TypeInfo, UE::Editor::DataStorage::RowHandle Row) const
 {
 	using namespace UE::Editor::DataStorage;
 
@@ -1115,16 +1109,16 @@ void UEditorDataStorageCompatibility::OnObjectReinstanced(const FCoreUObjectDele
 
 void UEditorDataStorageCompatibility::OnPostGcUnreachableAnalysis()
 {
-	using namespace TypedElementDataStorage;
 	using namespace TypedElementQueryBuilder;
 	using namespace UE::Editor::DataStorage;
+	using namespace UE::Editor::DataStorage::Private;
 
-	if (Private::bIntegrateWithGC)
+	if (bIntegrateWithGC)
 	{
 		TEDS_EVENT_SCOPE(TEXT("Post GC clean up"));
 		FScopedExclusiveLock Lock(EGlobalLockScope::Public);
 
-		if (Private::bUseCommandBuffer)
+		if (bUseCommandBuffer)
 		{
 			Storage->RunQuery(UObjectQuery, CreateDirectQueryCallbackBinding(
 				[this](RowHandle Row, const FTypedElementUObjectIdColumn& ObjectId)
diff --git a/Engine/Plugins/Experimental/EditorDataStorage/Source/TedsCore/Private/TypedElementDatabaseCompatibility.h b/Engine/Plugins/Experimental/EditorDataStorage/Source/TedsCore/Private/TypedElementDatabaseCompatibility.h
--- a/Engine/Plugins/Experimental/EditorDataStorage/Source/TedsCore/Private/TypedElementDatabaseCompatibility.h
+++ b/Engine/Plugins/Experimental/EditorDataStorage/Source/TedsCore/Private/TypedElementDatabaseCompatibility.h
@@ -204,10 +204,10 @@ private:
 	TArray<TPair<UE::Editor::DataStorage::ObjectAddedCallback, FDelegateHandle>> ObjectAddedCallbackList;
 	TArray<TPair<UE::Editor::DataStorage::ObjectRemovedCallback, FDelegateHandle>> PreObjectRemovedCallbackList;
 
-	TypedElementDataStorage::TableHandle StandardActorTable{ TypedElementDataStorage::InvalidTableHandle };
-	TypedElementDataStorage::TableHandle StandardActorWithTransformTable{ TypedElementDataStorage::InvalidTableHandle };
-	TypedElementDataStorage::TableHandle StandardUObjectTable{ TypedElementDataStorage::InvalidTableHandle };
-	TypedElementDataStorage::TableHandle StandardExternalObjectTable{ TypedElementDataStorage::InvalidTableHandle };
+	UE::Editor::DataStorage::TableHandle StandardActorTable{ UE::Editor::DataStorage::InvalidTableHandle };
+	UE::Editor::DataStorage::TableHandle StandardActorWithTransformTable{ UE::Editor::DataStorage::InvalidTableHandle };
+	UE::Editor::DataStorage::TableHandle StandardUObjectTable{ UE::Editor::DataStorage::InvalidTableHandle };
+	UE::Editor::DataStorage::TableHandle StandardExternalObjectTable{ UE::Editor::DataStorage::InvalidTableHandle };
 	ITypedElementDataStorageInterface* Storage{ nullptr };
 
 	/**
diff --git a/Engine/Plugins/Experimental/EditorDataStorage/Source/TedsCore/Private/TypedElementDatabaseEnvironment.cpp b/Engine/Plugins/Experimental/EditorDataStorage/Source/TedsCore/Private/TypedElementDatabaseEnvironment.cpp
--- a/Engine/Plugins/Experimental/EditorDataStorage/Source/TedsCore/Private/TypedElementDatabaseEnvironment.cpp
+++ b/Engine/Plugins/Experimental/EditorDataStorage/Source/TedsCore/Private/TypedElementDatabaseEnvironment.cpp
@@ -76,9 +76,9 @@ namespace UE::Editor::DataStorage
 		return MassEntityManager;
 	}
 
-	FMassArchetypeHandle FEnvironment::LookupMassArchetype(TypedElementDataStorage::TableHandle TableHandle) const
+	FMassArchetypeHandle FEnvironment::LookupMassArchetype(TableHandle Table) const
 	{
-		return DataStorage.LookupArchetype(TableHandle);
+		return DataStorage.LookupArchetype(Table);
 	}
 
 	FMassProcessingPhaseManager& FEnvironment::GetMassPhaseManager()
diff --git a/Engine/Plugins/Experimental/EditorDataStorage/Source/TedsCore/Private/TypedElementDatabaseEnvironment.h b/Engine/Plugins/Experimental/EditorDataStorage/Source/TedsCore/Private/TypedElementDatabaseEnvironment.h
--- a/Engine/Plugins/Experimental/EditorDataStorage/Source/TedsCore/Private/TypedElementDatabaseEnvironment.h
+++ b/Engine/Plugins/Experimental/EditorDataStorage/Source/TedsCore/Private/TypedElementDatabaseEnvironment.h
@@ -39,7 +39,7 @@ namespace UE::Editor::DataStorage
 		FMassEntityManager& GetMassEntityManager();
 		const FMassEntityManager& GetMassEntityManager() const;
 
-		FMassArchetypeHandle LookupMassArchetype(TypedElementDataStorage::TableHandle TableHandle) const;
+		FMassArchetypeHandle LookupMassArchetype(TableHandle Table) const;
 
 		FMassProcessingPhaseManager& GetMassPhaseManager();
 		const FMassProcessingPhaseManager& GetMassPhaseManager() const;
diff --git a/Engine/Plugins/Experimental/EditorDataStorage/Source/TedsCore/Private/TypedElementDatabaseIndexTable.cpp b/Engine/Plugins/Experimental/EditorDataStorage/Source/TedsCore/Private/TypedElementDatabaseIndexTable.cpp
--- a/Engine/Plugins/Experimental/EditorDataStorage/Source/TedsCore/Private/TypedElementDatabaseIndexTable.cpp
+++ b/Engine/Plugins/Experimental/EditorDataStorage/Source/TedsCore/Private/TypedElementDatabaseIndexTable.cpp
@@ -4,7 +4,7 @@
 
 namespace UE::Editor::DataStorage
 {
-	TypedElementDataStorage::RowHandle FIndexTable::FindIndexedRow(
+	RowHandle FIndexTable::FindIndexedRow(
 		EGlobalLockScope LockScope, TypedElementDataStorage::IndexHash Index) const
 	{
 		using namespace TypedElementDataStorage;
@@ -16,7 +16,7 @@ namespace UE::Editor::DataStorage
 	}
 
 	void FIndexTable::BatchIndexRows(EGlobalLockScope LockScope,
-		TConstArrayView<TPair<TypedElementDataStorage::IndexHash, TypedElementDataStorage::RowHandle>> IndexRowPairs)
+		TConstArrayView<TPair<TypedElementDataStorage::IndexHash, RowHandle>> IndexRowPairs)
 	{
 		using namespace TypedElementDataStorage;
 
@@ -32,7 +32,7 @@ namespace UE::Editor::DataStorage
 	}
 
 	void FIndexTable::IndexRow(EGlobalLockScope LockScope,
-		TypedElementDataStorage::IndexHash Index, TypedElementDataStorage::RowHandle Row)
+		TypedElementDataStorage::IndexHash Index, RowHandle Row)
 	{
 
 		FScopedExclusiveLock Lock(LockScope);
@@ -41,7 +41,7 @@ namespace UE::Editor::DataStorage
 
 	void FIndexTable::ReindexRow(EGlobalLockScope LockScope,
 		TypedElementDataStorage::IndexHash OriginalIndex, TypedElementDataStorage::IndexHash NewIndex, 
-		TypedElementDataStorage::RowHandle Row)
+		RowHandle Row)
 	{
 		FScopedExclusiveLock Lock(LockScope);
 	
@@ -55,7 +55,7 @@ namespace UE::Editor::DataStorage
 		RemoveIndexUnguarded(Index);
 	}
 
-	void FIndexTable::RemoveRow(EGlobalLockScope LockScope, TypedElementDataStorage::RowHandle Row)
+	void FIndexTable::RemoveRow(EGlobalLockScope LockScope, RowHandle Row)
 	{
 		using namespace TypedElementDataStorage;
 
@@ -72,7 +72,7 @@ namespace UE::Editor::DataStorage
 		}
 	}
 
-	void FIndexTable::IndexRowUnguarded(TypedElementDataStorage::IndexHash Index, TypedElementDataStorage::RowHandle Row)
+	void FIndexTable::IndexRowUnguarded(TypedElementDataStorage::IndexHash Index, RowHandle Row)
 	{
 		IndexLookupMap.Add(Index, Row);
 		ReverseIndexLookupMap.Add(Row, Index);
diff --git a/Engine/Plugins/Experimental/EditorDataStorage/Source/TedsCore/Private/TypedElementDatabaseIndexTable.h b/Engine/Plugins/Experimental/EditorDataStorage/Source/TedsCore/Private/TypedElementDatabaseIndexTable.h
--- a/Engine/Plugins/Experimental/EditorDataStorage/Source/TedsCore/Private/TypedElementDatabaseIndexTable.h
+++ b/Engine/Plugins/Experimental/EditorDataStorage/Source/TedsCore/Private/TypedElementDatabaseIndexTable.h
@@ -18,34 +18,34 @@ namespace UE::Editor::DataStorage
 	class FIndexTable final
 	{
 	public:
-		TypedElementDataStorage::RowHandle FindIndexedRow(
+		RowHandle FindIndexedRow(
 			EGlobalLockScope LockScope,
 			TypedElementDataStorage::IndexHash Index) const;
 		void IndexRow(
 			EGlobalLockScope LockScope,
 			TypedElementDataStorage::IndexHash Index, 
-			TypedElementDataStorage::RowHandle Row);
+			RowHandle Row);
 		void BatchIndexRows(
 			EGlobalLockScope LockScope,
 			TConstArrayView<TPair<TypedElementDataStorage::IndexHash, 
-			TypedElementDataStorage::RowHandle>> IndexRowPairs);
+			RowHandle>> IndexRowPairs);
 		void ReindexRow(
 			EGlobalLockScope LockScope,
 			TypedElementDataStorage::IndexHash OriginalIndex,
 			TypedElementDataStorage::IndexHash NewIndex,
-			TypedElementDataStorage::RowHandle Row);
+			RowHandle Row);
 		void RemoveIndex(
 			EGlobalLockScope LockScope,
 			TypedElementDataStorage::IndexHash Index);
 		void RemoveRow(
 			EGlobalLockScope LockScope,
-			TypedElementDataStorage::RowHandle Row);
+			RowHandle Row);
 
 	private:
-		TMap<TypedElementDataStorage::IndexHash, TypedElementDataStorage::RowHandle> IndexLookupMap;
-		TMultiMap<TypedElementDataStorage::RowHandle, TypedElementDataStorage::IndexHash> ReverseIndexLookupMap;
+		TMap<TypedElementDataStorage::IndexHash, RowHandle> IndexLookupMap;
+		TMultiMap<RowHandle, TypedElementDataStorage::IndexHash> ReverseIndexLookupMap;
 	
-		void IndexRowUnguarded(TypedElementDataStorage::IndexHash Index, TypedElementDataStorage::RowHandle Row);
+		void IndexRowUnguarded(TypedElementDataStorage::IndexHash Index, RowHandle Row);
 		void RemoveIndexUnguarded(TypedElementDataStorage::IndexHash Index);
 	};
 } // namespace UE::Editor::DataStorage
diff --git a/Engine/Plugins/Experimental/EditorDataStorage/Source/TedsCore/Private/TypedElementDatabaseUI.cpp b/Engine/Plugins/Experimental/EditorDataStorage/Source/TedsCore/Private/TypedElementDatabaseUI.cpp
--- a/Engine/Plugins/Experimental/EditorDataStorage/Source/TedsCore/Private/TypedElementDatabaseUI.cpp
+++ b/Engine/Plugins/Experimental/EditorDataStorage/Source/TedsCore/Private/TypedElementDatabaseUI.cpp
@@ -264,7 +264,7 @@ bool UEditorDataStorageUi::RegisterWidgetFactory(FName Purpose, TUniquePtr<FType
 }
 
 void UEditorDataStorageUi::CreateWidgetConstructors(FName Purpose,
-	const TypedElementDataStorage::FMetaDataView& Arguments, const WidgetConstructorCallback& Callback)
+	const UE::Editor::DataStorage::FMetaDataView& Arguments, const WidgetConstructorCallback& Callback)
 {
 	if (FPurposeInfo* PurposeInfo = WidgetPurposes.Find(Purpose))
 	{
@@ -279,7 +279,7 @@ void UEditorDataStorageUi::CreateWidgetConstructors(FName Purpose,
 }
 
 void UEditorDataStorageUi::CreateWidgetConstructors(FName Purpose, EMatchApproach MatchApproach, 
-	TArray<TWeakObjectPtr<const UScriptStruct>>& Columns, const TypedElementDataStorage::FMetaDataView& Arguments,
+	TArray<TWeakObjectPtr<const UScriptStruct>>& Columns, const UE::Editor::DataStorage::FMetaDataView& Arguments,
 	const WidgetConstructorCallback& Callback)
 {
 	if (FPurposeInfo* PurposeInfo = WidgetPurposes.Find(Purpose))
@@ -325,7 +325,7 @@ void UEditorDataStorageUi::CreateWidgetConstructors(FName Purpose, EMatchApproac
 	}
 }
 
-void UEditorDataStorageUi::ConstructWidgets(FName Purpose, const TypedElementDataStorage::FMetaDataView& Arguments,
+void UEditorDataStorageUi::ConstructWidgets(FName Purpose, const UE::Editor::DataStorage::FMetaDataView& Arguments,
 	const WidgetCreatedCallback& ConstructionCallback)
 {
 	if (FPurposeInfo* PurposeInfo = WidgetPurposes.Find(Purpose))
@@ -360,7 +360,7 @@ void UEditorDataStorageUi::ConstructWidgets(FName Purpose, const TypedElementDat
 
 bool UEditorDataStorageUi::CreateSingleWidgetConstructor(
 	const FWidgetFactory::ConstructorType& Constructor,
-	const TypedElementDataStorage::FMetaDataView& Arguments,
+	const UE::Editor::DataStorage::FMetaDataView& Arguments,
 	TArray<TWeakObjectPtr<const UScriptStruct>> MatchedColumnTypes,
 	const TypedElementDataStorage::FQueryConditions& QueryConditions,
 	const WidgetConstructorCallback& Callback)
@@ -370,7 +370,7 @@ bool UEditorDataStorageUi::CreateSingleWidgetConstructor(
 		Visitor(
 			TArray<TWeakObjectPtr<const UScriptStruct>>&& InMatchedColumnTypes,
 			const TypedElementDataStorage::FQueryConditions& InQueryConditions,
-			const TypedElementDataStorage::FMetaDataView& InArguments,
+			const UE::Editor::DataStorage::FMetaDataView& InArguments,
 			const WidgetConstructorCallback& InCallback) 
 			: MatchedColumnTypes(MoveTemp(InMatchedColumnTypes))
 			, QueryConditions(InQueryConditions)
@@ -380,7 +380,7 @@ bool UEditorDataStorageUi::CreateSingleWidgetConstructor(
 
 		TArray<TWeakObjectPtr<const UScriptStruct>> MatchedColumnTypes;
 		const TypedElementDataStorage::FQueryConditions& QueryConditions;
-		const TypedElementDataStorage::FMetaDataView& Arguments;
+		const UE::Editor::DataStorage::FMetaDataView& Arguments;
 		const WidgetConstructorCallback& Callback;
 
 		bool operator()(const UScriptStruct* Target)
@@ -418,7 +418,7 @@ bool UEditorDataStorageUi::CreateSingleWidgetConstructor(
 
 void UEditorDataStorageUi::CreateWidgetInstance(
 	FTypedElementWidgetConstructor& Constructor, 
-	const TypedElementDataStorage::FMetaDataView& Arguments,
+	const UE::Editor::DataStorage::FMetaDataView& Arguments,
 	const WidgetCreatedCallback& ConstructionCallback)
 {
 	UE::Editor::DataStorage::RowHandle Row = Storage->AddRow(WidgetTable);
@@ -435,7 +435,7 @@ void UEditorDataStorageUi::CreateWidgetInstance(
 }
 
 TSharedPtr<SWidget> UEditorDataStorageUi::ConstructWidget(UE::Editor::DataStorage::RowHandle Row, FTypedElementWidgetConstructor& Constructor,
-	const TypedElementDataStorage::FMetaDataView& Arguments)
+	const UE::Editor::DataStorage::FMetaDataView& Arguments)
 {
 	return Constructor.ConstructFinalWidget(Row, Storage, this, Arguments);
 }
@@ -467,7 +467,7 @@ void UEditorDataStorageUi::CreateStandardArchetypes()
 }
 
 void UEditorDataStorageUi::CreateWidgetConstructors_LongestMatch(const TArray<FWidgetFactory>& WidgetFactories,
-	TArray<TWeakObjectPtr<const UScriptStruct>>& Columns, const TypedElementDataStorage::FMetaDataView& Arguments,
+	TArray<TWeakObjectPtr<const UScriptStruct>>& Columns, const UE::Editor::DataStorage::FMetaDataView& Arguments,
 	const WidgetConstructorCallback& Callback)
 {
 	TArray<TWeakObjectPtr<const UScriptStruct>> MatchedColumns;
@@ -517,7 +517,7 @@ void UEditorDataStorageUi::CreateWidgetConstructors_LongestMatch(const TArray<FW
 }
 
 void UEditorDataStorageUi::CreateWidgetConstructors_ExactMatch(const TArray<FWidgetFactory>& WidgetFactories,
-	TArray<TWeakObjectPtr<const UScriptStruct>>& Columns, const TypedElementDataStorage::FMetaDataView& Arguments,
+	TArray<TWeakObjectPtr<const UScriptStruct>>& Columns, const UE::Editor::DataStorage::FMetaDataView& Arguments,
 	const WidgetConstructorCallback& Callback)
 {
 	int32 ColumnCount = Columns.Num();
@@ -548,7 +548,7 @@ void UEditorDataStorageUi::CreateWidgetConstructors_ExactMatch(const TArray<FWid
 }
 
 void UEditorDataStorageUi::CreateWidgetConstructors_SingleMatch(const TArray<FWidgetFactory>& WidgetFactories,
-	TArray<TWeakObjectPtr<const UScriptStruct>>& Columns, const TypedElementDataStorage::FMetaDataView& Arguments,
+	TArray<TWeakObjectPtr<const UScriptStruct>>& Columns, const UE::Editor::DataStorage::FMetaDataView& Arguments,
 	const WidgetConstructorCallback& Callback)
 {
 	auto FactoryIt = WidgetFactories.rbegin();
diff --git a/Engine/Plugins/Experimental/EditorDataStorage/Source/TedsCore/Private/TypedElementDatabaseUI.h b/Engine/Plugins/Experimental/EditorDataStorage/Source/TedsCore/Private/TypedElementDatabaseUI.h
--- a/Engine/Plugins/Experimental/EditorDataStorage/Source/TedsCore/Private/TypedElementDatabaseUI.h
+++ b/Engine/Plugins/Experimental/EditorDataStorage/Source/TedsCore/Private/TypedElementDatabaseUI.h
@@ -41,14 +41,14 @@ public:
 		TypedElementDataStorage::FQueryConditions Columns) override;
 
 	void CreateWidgetConstructors(FName Purpose,
-		const TypedElementDataStorage::FMetaDataView& Arguments, const WidgetConstructorCallback& Callback) override;
+		const UE::Editor::DataStorage::FMetaDataView& Arguments, const WidgetConstructorCallback& Callback) override;
 	void CreateWidgetConstructors(FName Purpose, EMatchApproach MatchApproach, TArray<TWeakObjectPtr<const UScriptStruct>>& Columns,
-		const TypedElementDataStorage::FMetaDataView& Arguments, const WidgetConstructorCallback& Callback) override;
+		const UE::Editor::DataStorage::FMetaDataView& Arguments, const WidgetConstructorCallback& Callback) override;
 
-	void ConstructWidgets(FName Purpose, const TypedElementDataStorage::FMetaDataView& Arguments,
+	void ConstructWidgets(FName Purpose, const UE::Editor::DataStorage::FMetaDataView& Arguments,
 		const WidgetCreatedCallback& ConstructionCallback) override;
 	TSharedPtr<SWidget> ConstructWidget(UE::Editor::DataStorage::RowHandle Row, FTypedElementWidgetConstructor& Constructor,
-		const TypedElementDataStorage::FMetaDataView& Arguments) override;
+		const UE::Editor::DataStorage::FMetaDataView& Arguments) override;
 
 	void ListWidgetPurposes(const WidgetPurposeCallback& Callback) const override;
 
@@ -82,30 +82,30 @@ private:
 
 	bool CreateSingleWidgetConstructor(
 		const FWidgetFactory::ConstructorType& Constructor,
-		const TypedElementDataStorage::FMetaDataView& Arguments,
+		const UE::Editor::DataStorage::FMetaDataView& Arguments,
 		TArray<TWeakObjectPtr<const UScriptStruct>> MatchedColumnTypes,
 		const TypedElementDataStorage::FQueryConditions& QueryConditions,
 		const WidgetConstructorCallback& Callback);
 
 	void CreateWidgetInstance(
 		FTypedElementWidgetConstructor& Constructor,
-		const TypedElementDataStorage::FMetaDataView& Arguments,
+		const UE::Editor::DataStorage::FMetaDataView& Arguments,
 		const WidgetCreatedCallback& ConstructionCallback);
 
 	void CreateWidgetConstructors_LongestMatch(
 		const TArray<FWidgetFactory>& WidgetFactories, 
 		TArray<TWeakObjectPtr<const UScriptStruct>>& Columns, 
-		const TypedElementDataStorage::FMetaDataView& Arguments,
+		const UE::Editor::DataStorage::FMetaDataView& Arguments,
 		const WidgetConstructorCallback& Callback);
 	void CreateWidgetConstructors_ExactMatch(
 		const TArray<FWidgetFactory>& WidgetFactories,
 		TArray<TWeakObjectPtr<const UScriptStruct>>& Columns,
-		const TypedElementDataStorage::FMetaDataView& Arguments,
+		const UE::Editor::DataStorage::FMetaDataView& Arguments,
 		const WidgetConstructorCallback& Callback);
 	void CreateWidgetConstructors_SingleMatch(
 		const TArray<FWidgetFactory>& WidgetFactories,
 		TArray<TWeakObjectPtr<const UScriptStruct>>& Columns,
-		const TypedElementDataStorage::FMetaDataView& Arguments,
+		const UE::Editor::DataStorage::FMetaDataView& Arguments,
 		const WidgetConstructorCallback& Callback);
 
 	UE::Editor::DataStorage::TableHandle WidgetTable{ UE::Editor::DataStorage::InvalidTableHandle };
