Capture the address to the environment for the defered command instead of a copy of this. The context might no longer exist when the command is executed thus causing some undefined behavior.

diff --git a/Engine/Plugins/Experimental/EditorDataStorage/Source/TedsCore/Private/Processors/TypedElementProcessorAdaptors.cpp b/Engine/Plugins/Experimental/EditorDataStorage/Source/TedsCore/Private/Processors/TypedElementProcessorAdaptors.cpp
--- a/Engine/Plugins/Experimental/EditorDataStorage/Source/TedsCore/Private/Processors/TypedElementProcessorAdaptors.cpp
+++ b/Engine/Plugins/Experimental/EditorDataStorage/Source/TedsCore/Private/Processors/TypedElementProcessorAdaptors.cpp
@@ -383,12 +383,12 @@ namespace UE::Editor::DataStorage
 					ScratchBuffer.Emplace<FAddDynamicColumn>(Relocator, Description, FMassEntityHandle::FromNumber(Row), ColumnData);
 		
 				this->Context.Defer().template PushCommand<FMassDeferredAddCommand>(
-					[AddedColumn, this](FMassEntityManager& System)
+					[AddedColumn, PtrToEnvironment = &Environment](FMassEntityManager& System)
 					{
 						// Check entity before proceeding. It's possible it may have been invalidated before this deferred call fired.
 						if (System.IsEntityActive(AddedColumn->Entity))
 						{
-							const UScriptStruct* DynamicStructType = Environment.GenerateDynamicColumn(*AddedColumn->Description.TemplateType, AddedColumn->Description.Identifier);
+							const UScriptStruct* DynamicStructType = PtrToEnvironment->GenerateDynamicColumn(*AddedColumn->Description.TemplateType, AddedColumn->Description.Identifier);
 
 							FStructView Fragment = System.GetFragmentDataStruct(AddedColumn->Entity, DynamicStructType);
 							// Check before adding.  Mass's AddFragmentToEntity is not idempotent and will assert if adding
