[Teds] Rename interface types and base factory class

ITypedElementDataStorageInterface -> IEditorDataStorageProvider
ITypedElementDataStorageUiInterface -> IEditorDataStorageUiProvider
ITypedElementDataStorageCompatibilityInterface -> IEditorDataStorageCompatibilityProvider
UTypedElementDataStorageFactory -> UEditorDataStorageFactory

diff --git a/Engine/Plugins/Developer/PerforceSourceControl/Source/PerforceSourceControl/Private/PerforceSourceControlOperations.cpp b/Engine/Plugins/Developer/PerforceSourceControl/Source/PerforceSourceControl/Private/PerforceSourceControlOperations.cpp
--- a/Engine/Plugins/Developer/PerforceSourceControl/Source/PerforceSourceControl/Private/PerforceSourceControlOperations.cpp
+++ b/Engine/Plugins/Developer/PerforceSourceControl/Source/PerforceSourceControl/Private/PerforceSourceControlOperations.cpp
@@ -2037,7 +2037,7 @@ bool FPerforceUpdateStatusWorker::UpdateStates() const
 			return;
 		}
 
-		ITypedElementDataStorageInterface* DataStorage = Registry->GetMutableDataStorage();
+		IEditorDataStorageProvider* DataStorage = Registry->GetMutableDataStorage();
 		if (!DataStorage)
 		{
 			return;
diff --git a/Engine/Plugins/Experimental/Animation/HierarchyTable/Source/Editor/Private/HierarchyTable/Factories/OverrideFactory.cpp b/Engine/Plugins/Experimental/Animation/HierarchyTable/Source/Editor/Private/HierarchyTable/Factories/OverrideFactory.cpp
--- a/Engine/Plugins/Experimental/Animation/HierarchyTable/Source/Editor/Private/HierarchyTable/Factories/OverrideFactory.cpp
+++ b/Engine/Plugins/Experimental/Animation/HierarchyTable/Source/Editor/Private/HierarchyTable/Factories/OverrideFactory.cpp
@@ -5,7 +5,7 @@
 #include "HierarchyTable/Widgets/OverrideWidgetConstructor.h"
 #include "Elements/Interfaces/TypedElementDataStorageUiInterface.h"
 
-void UHierarchyTableOverrideFactory::RegisterWidgetConstructors(ITypedElementDataStorageInterface& DataStorage, ITypedElementDataStorageUiInterface& DataStorageUi) const
+void UHierarchyTableOverrideFactory::RegisterWidgetConstructors(IEditorDataStorageProvider& DataStorage, IEditorDataStorageUiProvider& DataStorageUi) const
 {
 	using namespace UE::Editor::DataStorage::Queries;
 
diff --git a/Engine/Plugins/Experimental/Animation/HierarchyTable/Source/Editor/Private/HierarchyTable/Factories/OverrideFactory.h b/Engine/Plugins/Experimental/Animation/HierarchyTable/Source/Editor/Private/HierarchyTable/Factories/OverrideFactory.h
--- a/Engine/Plugins/Experimental/Animation/HierarchyTable/Source/Editor/Private/HierarchyTable/Factories/OverrideFactory.h
+++ b/Engine/Plugins/Experimental/Animation/HierarchyTable/Source/Editor/Private/HierarchyTable/Factories/OverrideFactory.h
@@ -6,10 +6,10 @@
 #include "OverrideFactory.generated.h"
 
 UCLASS()
-class UHierarchyTableOverrideFactory : public UTypedElementDataStorageFactory
+class UHierarchyTableOverrideFactory : public UEditorDataStorageFactory
 {
 	GENERATED_BODY()
 
 public:
-	void RegisterWidgetConstructors(ITypedElementDataStorageInterface& DataStorage, ITypedElementDataStorageUiInterface& DataStorageUi) const override;
+	void RegisterWidgetConstructors(IEditorDataStorageProvider& DataStorage, IEditorDataStorageUiProvider& DataStorageUi) const override;
 };
\ No newline at end of file
diff --git a/Engine/Plugins/Experimental/Animation/HierarchyTable/Source/Editor/Private/HierarchyTable/Widgets/OverrideWidgetConstructor.cpp b/Engine/Plugins/Experimental/Animation/HierarchyTable/Source/Editor/Private/HierarchyTable/Widgets/OverrideWidgetConstructor.cpp
--- a/Engine/Plugins/Experimental/Animation/HierarchyTable/Source/Editor/Private/HierarchyTable/Widgets/OverrideWidgetConstructor.cpp
+++ b/Engine/Plugins/Experimental/Animation/HierarchyTable/Source/Editor/Private/HierarchyTable/Widgets/OverrideWidgetConstructor.cpp
@@ -23,8 +23,8 @@ TSharedPtr<SWidget> FTypedElementWidgetConstructor_Override::CreateWidget(
 }
 
 bool FTypedElementWidgetConstructor_Override::FinalizeWidget(
-	ITypedElementDataStorageInterface* DataStorage,
-	ITypedElementDataStorageUiInterface* DataStorageUi,
+	IEditorDataStorageProvider* DataStorage,
+	IEditorDataStorageUiProvider* DataStorageUi,
 	UE::Editor::DataStorage::RowHandle Row,
 	const TSharedPtr<SWidget>& Widget)
 {
diff --git a/Engine/Plugins/Experimental/Animation/HierarchyTable/Source/Editor/Private/HierarchyTable/Widgets/OverrideWidgetConstructor.h b/Engine/Plugins/Experimental/Animation/HierarchyTable/Source/Editor/Private/HierarchyTable/Widgets/OverrideWidgetConstructor.h
--- a/Engine/Plugins/Experimental/Animation/HierarchyTable/Source/Editor/Private/HierarchyTable/Widgets/OverrideWidgetConstructor.h
+++ b/Engine/Plugins/Experimental/Animation/HierarchyTable/Source/Editor/Private/HierarchyTable/Widgets/OverrideWidgetConstructor.h
@@ -17,6 +17,6 @@ public:
 
 protected:
 	TSharedPtr<SWidget> CreateWidget(const UE::Editor::DataStorage::FMetaDataView& Arguments);
-	bool FinalizeWidget(ITypedElementDataStorageInterface* DataStorage, ITypedElementDataStorageUiInterface* DataStorageUi,
+	bool FinalizeWidget(IEditorDataStorageProvider* DataStorage, IEditorDataStorageUiProvider* DataStorageUi,
 		UE::Editor::DataStorage::RowHandle Row, const TSharedPtr<SWidget>& Widget) override;
 };
\ No newline at end of file
diff --git a/Engine/Plugins/Experimental/Animation/HierarchyTable/Source/Editor/Private/HierarchyTableEditorToolkit.cpp b/Engine/Plugins/Experimental/Animation/HierarchyTable/Source/Editor/Private/HierarchyTableEditorToolkit.cpp
--- a/Engine/Plugins/Experimental/Animation/HierarchyTable/Source/Editor/Private/HierarchyTableEditorToolkit.cpp
+++ b/Engine/Plugins/Experimental/Animation/HierarchyTable/Source/Editor/Private/HierarchyTableEditorToolkit.cpp
@@ -86,7 +86,7 @@ void FHierarchyTableEditorToolkit::RegisterTabSpawners(const TSharedRef<class FT
 void FHierarchyTableEditorToolkit::OnClose()
 {
 	UTypedElementRegistry* Registry = UTypedElementRegistry::GetInstance();
-	ITypedElementDataStorageInterface* DSI = Registry->GetMutableDataStorage();
+	IEditorDataStorageProvider* DSI = Registry->GetMutableDataStorage();
 
 	for (const TTuple<int32, UE::Editor::DataStorage::RowHandle>& Row : EntryIndexToHandleMap)
 	{
@@ -157,7 +157,7 @@ TSharedRef<SWidget> FHierarchyTableEditorToolkit::CreateTedsOutliner()
 
 	FTedsOutlinerModule& TedsOutlinerModule = FModuleManager::GetModuleChecked<FTedsOutlinerModule>("TedsOutliner");
 
-	ITypedElementDataStorageInterface* DSI = Registry->GetMutableDataStorage();
+	IEditorDataStorageProvider* DSI = Registry->GetMutableDataStorage();
 	static TableHandle Table = DSI->FindTable(FName("Editor_HierarchyTableTable"));
 	
 	TArray<UScriptStruct*> BaseHierarchyTableTypeColumns = Handler->GetColumns();
@@ -233,7 +233,7 @@ void FHierarchyTableEditorToolkit::AddCurveEntry(const FName CurveName)
 	/*
 	using namespace UE::Editor::DataStorage;
 	UTypedElementRegistry* Registry = UTypedElementRegistry::GetInstance();
-	ITypedElementDataStorageInterface* DSI = Registry->GetMutableDataStorage();
+	IEditorDataStorageProvider* DSI = Registry->GetMutableDataStorage();
 	static TableHandle Table = DSI->FindTable(FName("Editor_HierarchyTableTable"));
 
 	RowHandle* RootBoneHandle = EntryIndexToHandleMap.Find(0);
diff --git a/Engine/Plugins/Experimental/Animation/HierarchyTable/Source/Editor/Private/HierarchyTableTedsFactory.cpp b/Engine/Plugins/Experimental/Animation/HierarchyTable/Source/Editor/Private/HierarchyTableTedsFactory.cpp
--- a/Engine/Plugins/Experimental/Animation/HierarchyTable/Source/Editor/Private/HierarchyTableTedsFactory.cpp
+++ b/Engine/Plugins/Experimental/Animation/HierarchyTable/Source/Editor/Private/HierarchyTableTedsFactory.cpp
@@ -7,7 +7,7 @@
 #include "Elements/Framework/TypedElementQueryBuilder.h"
 #include "Elements/Framework/TypedElementRegistry.h"
 
-void UTypedElementHierarchyTableFactory::RegisterTables(ITypedElementDataStorageInterface& DataStorage)
+void UTypedElementHierarchyTableFactory::RegisterTables(IEditorDataStorageProvider& DataStorage)
 {
 	DataStorage.RegisterTable(
 		TTypedElementColumnTypeList<
@@ -15,6 +15,6 @@ void UTypedElementHierarchyTableFactory::RegisterTables(ITypedElementDataStorage
 		TEXT("Editor_HierarchyTableTable"));
 }
 
-void UTypedElementHierarchyTableFactory::RegisterQueries(ITypedElementDataStorageInterface& DataStorage)
+void UTypedElementHierarchyTableFactory::RegisterQueries(IEditorDataStorageProvider& DataStorage)
 {
 }
\ No newline at end of file
diff --git a/Engine/Plugins/Experimental/Animation/HierarchyTable/Source/Editor/Private/HierarchyTableWidgetConstructor.cpp b/Engine/Plugins/Experimental/Animation/HierarchyTable/Source/Editor/Private/HierarchyTableWidgetConstructor.cpp
--- a/Engine/Plugins/Experimental/Animation/HierarchyTable/Source/Editor/Private/HierarchyTableWidgetConstructor.cpp
+++ b/Engine/Plugins/Experimental/Animation/HierarchyTable/Source/Editor/Private/HierarchyTableWidgetConstructor.cpp
@@ -23,7 +23,7 @@ TSharedPtr<SWidget> FHierarchyTableWidgetConstructor::CreateWidget(const UE::Edi
 		.VAlign(VAlign_Center);
 }
 
-bool FHierarchyTableWidgetConstructor::FinalizeWidget(ITypedElementDataStorageInterface* DataStorage, ITypedElementDataStorageUiInterface* DataStorageUi,
+bool FHierarchyTableWidgetConstructor::FinalizeWidget(IEditorDataStorageProvider* DataStorage, IEditorDataStorageUiProvider* DataStorageUi,
 	UE::Editor::DataStorage::RowHandle Row, const TSharedPtr<SWidget>& Widget)
 {
 	TSharedPtr<SBox> WidgetInstance = StaticCastSharedPtr<SBox>(Widget);
diff --git a/Engine/Plugins/Experimental/Animation/HierarchyTable/Source/Editor/Public/HierarchyTableTedsFactory.h b/Engine/Plugins/Experimental/Animation/HierarchyTable/Source/Editor/Public/HierarchyTableTedsFactory.h
--- a/Engine/Plugins/Experimental/Animation/HierarchyTable/Source/Editor/Public/HierarchyTableTedsFactory.h
+++ b/Engine/Plugins/Experimental/Animation/HierarchyTable/Source/Editor/Public/HierarchyTableTedsFactory.h
@@ -9,13 +9,13 @@
 #include "HierarchyTableTedsFactory.generated.h"
 
 UCLASS()
-class UTypedElementHierarchyTableFactory : public UTypedElementDataStorageFactory
+class UTypedElementHierarchyTableFactory : public UEditorDataStorageFactory
 {
 	GENERATED_BODY()
 
 public:
 	~UTypedElementHierarchyTableFactory() override = default;
 
-	void RegisterTables(ITypedElementDataStorageInterface& DataStorage) override;
-	void RegisterQueries(ITypedElementDataStorageInterface& DataStorage) override;
+	void RegisterTables(IEditorDataStorageProvider& DataStorage) override;
+	void RegisterQueries(IEditorDataStorageProvider& DataStorage) override;
 };
diff --git a/Engine/Plugins/Experimental/Animation/HierarchyTable/Source/Editor/Public/HierarchyTableWidgetConstructor.h b/Engine/Plugins/Experimental/Animation/HierarchyTable/Source/Editor/Public/HierarchyTableWidgetConstructor.h
--- a/Engine/Plugins/Experimental/Animation/HierarchyTable/Source/Editor/Public/HierarchyTableWidgetConstructor.h
+++ b/Engine/Plugins/Experimental/Animation/HierarchyTable/Source/Editor/Public/HierarchyTableWidgetConstructor.h
@@ -25,7 +25,7 @@ protected:
 	// Begin FHierarchyTableWidgetConstructor
 	TSharedPtr<SWidget> CreateWidget(const UE::Editor::DataStorage::FMetaDataView& Arguments);
 
-	bool FinalizeWidget(ITypedElementDataStorageInterface* DataStorage, ITypedElementDataStorageUiInterface* DataStorageUi,
+	bool FinalizeWidget(IEditorDataStorageProvider* DataStorage, IEditorDataStorageUiProvider* DataStorageUi,
 		UE::Editor::DataStorage::RowHandle Row, const TSharedPtr<SWidget>& Widget) override;
 	// End FHierarchyTableWidgetConstructor
 };
\ No newline at end of file
diff --git a/Engine/Plugins/Experimental/Animation/HierarchyTableBuiltin/Source/Editor/Private/MaskProfile/MaskProfileFactory.cpp b/Engine/Plugins/Experimental/Animation/HierarchyTableBuiltin/Source/Editor/Private/MaskProfile/MaskProfileFactory.cpp
--- a/Engine/Plugins/Experimental/Animation/HierarchyTableBuiltin/Source/Editor/Private/MaskProfile/MaskProfileFactory.cpp
+++ b/Engine/Plugins/Experimental/Animation/HierarchyTableBuiltin/Source/Editor/Private/MaskProfile/MaskProfileFactory.cpp
@@ -4,7 +4,7 @@
 #include "MaskProfile/MaskProfileProxyColumn.h"
 #include "MaskProfile/MaskProfileWidgetConstructor.h"
 
-void UHierarchyTableMaskFactory::RegisterWidgetConstructors(ITypedElementDataStorageInterface& DataStorage, ITypedElementDataStorageUiInterface& DataStorageUi) const
+void UHierarchyTableMaskFactory::RegisterWidgetConstructors(IEditorDataStorageProvider& DataStorage, IEditorDataStorageUiProvider& DataStorageUi) const
 {
 	using namespace UE::Editor::DataStorage::Queries;
 
diff --git a/Engine/Plugins/Experimental/Animation/HierarchyTableBuiltin/Source/Editor/Private/MaskProfile/MaskProfileFactory.h b/Engine/Plugins/Experimental/Animation/HierarchyTableBuiltin/Source/Editor/Private/MaskProfile/MaskProfileFactory.h
--- a/Engine/Plugins/Experimental/Animation/HierarchyTableBuiltin/Source/Editor/Private/MaskProfile/MaskProfileFactory.h
+++ b/Engine/Plugins/Experimental/Animation/HierarchyTableBuiltin/Source/Editor/Private/MaskProfile/MaskProfileFactory.h
@@ -7,10 +7,10 @@
 #include "MaskProfileFactory.generated.h"
 
 UCLASS()
-class UHierarchyTableMaskFactory : public UTypedElementDataStorageFactory
+class UHierarchyTableMaskFactory : public UEditorDataStorageFactory
 {
 	GENERATED_BODY()
 
 public:
-	void RegisterWidgetConstructors(ITypedElementDataStorageInterface& DataStorage, ITypedElementDataStorageUiInterface& DataStorageUi) const override;
+	void RegisterWidgetConstructors(IEditorDataStorageProvider& DataStorage, IEditorDataStorageUiProvider& DataStorageUi) const override;
 };
\ No newline at end of file
diff --git a/Engine/Plugins/Experimental/Animation/HierarchyTableBuiltin/Source/Editor/Private/TimeProfile/TimeProfileFactory.cpp b/Engine/Plugins/Experimental/Animation/HierarchyTableBuiltin/Source/Editor/Private/TimeProfile/TimeProfileFactory.cpp
--- a/Engine/Plugins/Experimental/Animation/HierarchyTableBuiltin/Source/Editor/Private/TimeProfile/TimeProfileFactory.cpp
+++ b/Engine/Plugins/Experimental/Animation/HierarchyTableBuiltin/Source/Editor/Private/TimeProfile/TimeProfileFactory.cpp
@@ -4,7 +4,7 @@
 #include "TimeProfile/TimeProfileProxyColumn.h"
 #include "TimeProfile/TimeProfileWidgetConstructor.h"
 
-void UHierarchyTableTimeFactory::RegisterWidgetConstructors(ITypedElementDataStorageInterface& DataStorage, ITypedElementDataStorageUiInterface& DataStorageUi) const
+void UHierarchyTableTimeFactory::RegisterWidgetConstructors(IEditorDataStorageProvider& DataStorage, IEditorDataStorageUiProvider& DataStorageUi) const
 {
 	using namespace UE::Editor::DataStorage::Queries;
 
diff --git a/Engine/Plugins/Experimental/Animation/HierarchyTableBuiltin/Source/Editor/Private/TimeProfile/TimeProfileFactory.h b/Engine/Plugins/Experimental/Animation/HierarchyTableBuiltin/Source/Editor/Private/TimeProfile/TimeProfileFactory.h
--- a/Engine/Plugins/Experimental/Animation/HierarchyTableBuiltin/Source/Editor/Private/TimeProfile/TimeProfileFactory.h
+++ b/Engine/Plugins/Experimental/Animation/HierarchyTableBuiltin/Source/Editor/Private/TimeProfile/TimeProfileFactory.h
@@ -6,10 +6,10 @@
 #include "TimeProfileFactory.generated.h"
 
 UCLASS()
-class UHierarchyTableTimeFactory : public UTypedElementDataStorageFactory
+class UHierarchyTableTimeFactory : public UEditorDataStorageFactory
 {
 	GENERATED_BODY()
 
 public:
-	void RegisterWidgetConstructors(ITypedElementDataStorageInterface& DataStorage, ITypedElementDataStorageUiInterface& DataStorageUi) const override;
+	void RegisterWidgetConstructors(IEditorDataStorageProvider& DataStorage, IEditorDataStorageUiProvider& DataStorageUi) const override;
 };
\ No newline at end of file
diff --git a/Engine/Plugins/Experimental/EditorDataStorage/Source/TedsCore/Private/Commands/EditorDataStorageCompatibilityCommands.cpp b/Engine/Plugins/Experimental/EditorDataStorage/Source/TedsCore/Private/Commands/EditorDataStorageCompatibilityCommands.cpp
--- a/Engine/Plugins/Experimental/EditorDataStorage/Source/TedsCore/Private/Commands/EditorDataStorageCompatibilityCommands.cpp
+++ b/Engine/Plugins/Experimental/EditorDataStorage/Source/TedsCore/Private/Commands/EditorDataStorageCompatibilityCommands.cpp
@@ -225,7 +225,7 @@ namespace UE::Editor::DataStorage
 	// FPrepareCommands
 	//
 
-	FPrepareCommands::FPrepareCommands(ITypedElementDataStorageInterface& InStorage, UEditorDataStorageCompatibility& InStorageCompat,
+	FPrepareCommands::FPrepareCommands(IEditorDataStorageProvider& InStorage, UEditorDataStorageCompatibility& InStorageCompat,
 		CompatibilityCommandBuffer::FCollection& InCommands)
 		: Storage(InStorage)
 		, StorageCompat(InStorageCompat)
@@ -304,7 +304,7 @@ namespace UE::Editor::DataStorage
 		}
 	}
 
-	void FPrepareCommands::RunPreparation(ITypedElementDataStorageInterface& Storage, UEditorDataStorageCompatibility& StorageCompat,
+	void FPrepareCommands::RunPreparation(IEditorDataStorageProvider& Storage, UEditorDataStorageCompatibility& StorageCompat,
 		CompatibilityCommandBuffer::FCollection& Commands)
 	{
 		FPrepareCommands PrepareVisitor(Storage, StorageCompat, Commands);
@@ -490,7 +490,7 @@ namespace UE::Editor::DataStorage
 	// FCommandProcessor
 	//
 
-	FCommandProcessor::FCommandProcessor(ITypedElementDataStorageInterface& InStorage, UEditorDataStorageCompatibility& InStorageCompatibility)
+	FCommandProcessor::FCommandProcessor(IEditorDataStorageProvider& InStorage, UEditorDataStorageCompatibility& InStorageCompatibility)
 		: Storage(InStorage)
 		, StorageCompatibility(InStorageCompatibility)
 		, MementoSystem(InStorageCompatibility.Environment->GetMementoSystem())
diff --git a/Engine/Plugins/Experimental/EditorDataStorage/Source/TedsCore/Private/Commands/EditorDataStorageCompatibilityCommands.h b/Engine/Plugins/Experimental/EditorDataStorage/Source/TedsCore/Private/Commands/EditorDataStorageCompatibilityCommands.h
--- a/Engine/Plugins/Experimental/EditorDataStorage/Source/TedsCore/Private/Commands/EditorDataStorageCompatibilityCommands.h
+++ b/Engine/Plugins/Experimental/EditorDataStorage/Source/TedsCore/Private/Commands/EditorDataStorageCompatibilityCommands.h
@@ -9,7 +9,7 @@
 #include "UObject/WeakObjectPtr.h"
 #include "UObject/WeakObjectPtrTemplates.h"
 
-class ITypedElementDataStorageInterface;
+class IEditorDataStorageProvider;
 class UScriptStruct;
 class UEditorDataStorageCompatibility;
 class UObject;
@@ -27,7 +27,7 @@ namespace UE::Editor::DataStorage
 
 	/**
 	 * Objects with type info defined in either UScriptStruct or UClass can be stored into TEDS via the
-	 * ITypedElementDataStorageCompatibilityInterface
+	 * IEditorDataStorageCompatibilityProvider
 	 * This is a discriminated union which aids with callbacks made when objects are added
 	 */
 	struct FObjectTypeInfo
@@ -228,7 +228,7 @@ namespace UE::Editor::DataStorage
 	/** Prepares each command for further processing, e.g. resolving the target table. */
 	struct FPrepareCommands final
 	{
-		FPrepareCommands(ITypedElementDataStorageInterface& InStorage, UEditorDataStorageCompatibility& InStorageCompat,
+		FPrepareCommands(IEditorDataStorageProvider& InStorage, UEditorDataStorageCompatibility& InStorageCompat,
 			CompatibilityCommandBuffer::FCollection& InCommands);
 
 		template<typename T>
@@ -241,12 +241,12 @@ namespace UE::Editor::DataStorage
 		void operator()(FRemoveInteractiveSyncFromWorldTag& Command);
 		void operator()(FAddSyncFromWorldTag& Command);
 
-		ITypedElementDataStorageInterface& Storage;
+		IEditorDataStorageProvider& Storage;
 		UEditorDataStorageCompatibility& StorageCompat;
 		CompatibilityCommandBuffer::FCollection& Commands;
 		int32 CurrentIndex = 0;
 
-		static void RunPreparation(ITypedElementDataStorageInterface& Storage, UEditorDataStorageCompatibility& StorageCompat, 
+		static void RunPreparation(IEditorDataStorageProvider& Storage, UEditorDataStorageCompatibility& StorageCompat, 
 			CompatibilityCommandBuffer::FCollection& Commands);
 	};
 
@@ -299,7 +299,7 @@ namespace UE::Editor::DataStorage
 	/** Executes the commands in the command buffer for TEDS Compatibility. */
 	struct FCommandProcessor final
 	{
-		FCommandProcessor(ITypedElementDataStorageInterface& InStorage, UEditorDataStorageCompatibility& InStorageCompatibility);
+		FCommandProcessor(IEditorDataStorageProvider& InStorage, UEditorDataStorageCompatibility& InStorageCompatibility);
 
 		void SetupRow(RowHandle Row, UObject* Object);
 		void SetupRow(RowHandle Row, void* Object, TWeakObjectPtr<UScriptStruct> TypeInfo);
@@ -323,7 +323,7 @@ namespace UE::Editor::DataStorage
 		void operator()(FRemoveInteractiveSyncFromWorldTag& Command);
 		void operator()(FAddSyncFromWorldTag& Command);
 
-		ITypedElementDataStorageInterface& Storage;
+		IEditorDataStorageProvider& Storage;
 		UEditorDataStorageCompatibility& StorageCompatibility;
 		FMementoSystem& MementoSystem;
 	};
diff --git a/Engine/Plugins/Experimental/EditorDataStorage/Source/TedsCore/Private/Compatibility/TedsTypedElementActorHandleFactory.cpp b/Engine/Plugins/Experimental/EditorDataStorage/Source/TedsCore/Private/Compatibility/TedsTypedElementActorHandleFactory.cpp
--- a/Engine/Plugins/Experimental/EditorDataStorage/Source/TedsCore/Private/Compatibility/TedsTypedElementActorHandleFactory.cpp
+++ b/Engine/Plugins/Experimental/EditorDataStorage/Source/TedsCore/Private/Compatibility/TedsTypedElementActorHandleFactory.cpp
@@ -11,20 +11,20 @@
 #include "Elements/Framework/TypedElementRegistry.h"
 #include "GameFramework/Actor.h"
 
-void UTypedElementActorHandleDataStorageFactory::PreRegister(ITypedElementDataStorageInterface& DataStorage)
+void UTypedElementActorHandleDataStorageFactory::PreRegister(IEditorDataStorageProvider& DataStorage)
 {
 	Super::PreRegister(DataStorage);
 	
 	BridgeEnableDelegateHandle = UE::Editor::DataStorage::Compatibility::OnTypedElementBridgeEnabled().AddUObject(this, &UTypedElementActorHandleDataStorageFactory::HandleBridgeEnabled);
 }
 
-void UTypedElementActorHandleDataStorageFactory::PreShutdown(ITypedElementDataStorageInterface& DataStorage)
+void UTypedElementActorHandleDataStorageFactory::PreShutdown(IEditorDataStorageProvider& DataStorage)
 {
 	UE::Editor::DataStorage::Compatibility::OnTypedElementBridgeEnabled().Remove(BridgeEnableDelegateHandle);
 	BridgeEnableDelegateHandle.Reset();
 }
 
-void UTypedElementActorHandleDataStorageFactory::RegisterQueries(ITypedElementDataStorageInterface& DataStorage)
+void UTypedElementActorHandleDataStorageFactory::RegisterQueries(IEditorDataStorageProvider& DataStorage)
 {
 	Super::RegisterQueries(DataStorage);
 
@@ -42,7 +42,7 @@ void UTypedElementActorHandleDataStorageFactory::RegisterQueries(ITypedElementDa
 	.Compile());
 }
 
-void UTypedElementActorHandleDataStorageFactory::RegisterQuery_ActorHandlePopulate(ITypedElementDataStorageInterface& DataStorage)
+void UTypedElementActorHandleDataStorageFactory::RegisterQuery_ActorHandlePopulate(IEditorDataStorageProvider& DataStorage)
 {
 	using namespace UE::Editor::DataStorage::Queries;
 
@@ -75,7 +75,7 @@ void UTypedElementActorHandleDataStorageFactory::HandleBridgeEnabled(bool bEnabl
 {
 	using namespace UE::Editor::DataStorage;
 
-	ITypedElementDataStorageInterface* DataStorage = UTypedElementRegistry::GetInstance()->GetMutableDataStorage();
+	IEditorDataStorageProvider* DataStorage = UTypedElementRegistry::GetInstance()->GetMutableDataStorage();
 
 	if (bEnabled)
 	{
diff --git a/Engine/Plugins/Experimental/EditorDataStorage/Source/TedsCore/Private/Compatibility/TedsTypedElementActorHandleFactory.h b/Engine/Plugins/Experimental/EditorDataStorage/Source/TedsCore/Private/Compatibility/TedsTypedElementActorHandleFactory.h
--- a/Engine/Plugins/Experimental/EditorDataStorage/Source/TedsCore/Private/Compatibility/TedsTypedElementActorHandleFactory.h
+++ b/Engine/Plugins/Experimental/EditorDataStorage/Source/TedsCore/Private/Compatibility/TedsTypedElementActorHandleFactory.h
@@ -16,16 +16,16 @@ class UTypedElementRegistry;
  * with TEDS/
  */
 UCLASS(Transient)
-class UTypedElementActorHandleDataStorageFactory : public UTypedElementDataStorageFactory
+class UTypedElementActorHandleDataStorageFactory : public UEditorDataStorageFactory
 {
 	GENERATED_BODY()
 public:
 
-	virtual void PreRegister(ITypedElementDataStorageInterface& DataStorage) override;
-	virtual void RegisterQueries(ITypedElementDataStorageInterface& DataStorage) override;
-	virtual void PreShutdown(ITypedElementDataStorageInterface& DataStorage) override;
+	virtual void PreRegister(IEditorDataStorageProvider& DataStorage) override;
+	virtual void RegisterQueries(IEditorDataStorageProvider& DataStorage) override;
+	virtual void PreShutdown(IEditorDataStorageProvider& DataStorage) override;
 private:
-	virtual void RegisterQuery_ActorHandlePopulate(ITypedElementDataStorageInterface& DataStorage);
+	virtual void RegisterQuery_ActorHandlePopulate(IEditorDataStorageProvider& DataStorage);
 	void HandleBridgeEnabled(bool bEnabled);
 
 	FDelegateHandle BridgeEnableDelegateHandle;
diff --git a/Engine/Plugins/Experimental/EditorDataStorage/Source/TedsCore/Private/Compatibility/TedsTypedElementBridgeQueries.cpp b/Engine/Plugins/Experimental/EditorDataStorage/Source/TedsCore/Private/Compatibility/TedsTypedElementBridgeQueries.cpp
--- a/Engine/Plugins/Experimental/EditorDataStorage/Source/TedsCore/Private/Compatibility/TedsTypedElementBridgeQueries.cpp
+++ b/Engine/Plugins/Experimental/EditorDataStorage/Source/TedsCore/Private/Compatibility/TedsTypedElementBridgeQueries.cpp
@@ -25,13 +25,13 @@ uint8 UTypedElementBridgeDataStorageFactory::GetOrder() const
 	return 110;
 }
 
-void UTypedElementBridgeDataStorageFactory::PreRegister(ITypedElementDataStorageInterface& DataStorage)
+void UTypedElementBridgeDataStorageFactory::PreRegister(IEditorDataStorageProvider& DataStorage)
 {
 	Super::PreRegister(DataStorage);
 	DebugEnabledDelegateHandle = UE::Editor::DataStorage::Compatibility::Private::CVarBridgeEnabled->OnChangedDelegate().AddUObject(this, &UTypedElementBridgeDataStorageFactory::HandleOnEnabled);
 }
 
-void UTypedElementBridgeDataStorageFactory::PreShutdown(ITypedElementDataStorageInterface& DataStorage)
+void UTypedElementBridgeDataStorageFactory::PreShutdown(IEditorDataStorageProvider& DataStorage)
 {
 	UE::Editor::DataStorage::Compatibility::Private::CVarBridgeEnabled->OnChangedDelegate().Remove(DebugEnabledDelegateHandle);
 	DebugEnabledDelegateHandle.Reset();
@@ -40,7 +40,7 @@ void UTypedElementBridgeDataStorageFactory::PreShutdown(ITypedElementDataStorage
 	Super::PreShutdown(DataStorage);
 }
 
-void UTypedElementBridgeDataStorageFactory::RegisterQueries(ITypedElementDataStorageInterface& DataStorage)
+void UTypedElementBridgeDataStorageFactory::RegisterQueries(IEditorDataStorageProvider& DataStorage)
 {
 	Super::RegisterQueries(DataStorage);
 
@@ -55,7 +55,7 @@ bool UTypedElementBridgeDataStorageFactory::IsEnabled()
 	return UE::Editor::DataStorage::Compatibility::Private::CVarBridgeEnabled->GetBool();
 }
 
-void UTypedElementBridgeDataStorageFactory::RegisterQuery_NewUObject(ITypedElementDataStorageInterface& DataStorage)
+void UTypedElementBridgeDataStorageFactory::RegisterQuery_NewUObject(IEditorDataStorageProvider& DataStorage)
 {
 	using namespace UE::Editor::DataStorage::Queries;
 	
@@ -65,11 +65,11 @@ void UTypedElementBridgeDataStorageFactory::RegisterQuery_NewUObject(ITypedEleme
 		.Compile());
 }
 
-void UTypedElementBridgeDataStorageFactory::UnregisterQuery_NewUObject(ITypedElementDataStorageInterface& DataStorage)
+void UTypedElementBridgeDataStorageFactory::UnregisterQuery_NewUObject(IEditorDataStorageProvider& DataStorage)
 {
 }
 
-void UTypedElementBridgeDataStorageFactory::CleanupTypedElementColumns(ITypedElementDataStorageInterface& DataStorage)
+void UTypedElementBridgeDataStorageFactory::CleanupTypedElementColumns(IEditorDataStorageProvider& DataStorage)
 {
 	using namespace UE::Editor::DataStorage::Queries;
 
@@ -90,7 +90,7 @@ void UTypedElementBridgeDataStorageFactory::CleanupTypedElementColumns(ITypedEle
 
 void UTypedElementBridgeDataStorageFactory::HandleOnEnabled(IConsoleVariable* CVar)
 {
-	ITypedElementDataStorageInterface* DataStorage = UTypedElementRegistry::GetInstance()->GetMutableDataStorage();
+	IEditorDataStorageProvider* DataStorage = UTypedElementRegistry::GetInstance()->GetMutableDataStorage();
 	bool bIsEnabled = CVar->GetBool();
 
 	if (bIsEnabled)
diff --git a/Engine/Plugins/Experimental/EditorDataStorage/Source/TedsCore/Private/Compatibility/TedsTypedElementBridgeQueries.h b/Engine/Plugins/Experimental/EditorDataStorage/Source/TedsCore/Private/Compatibility/TedsTypedElementBridgeQueries.h
--- a/Engine/Plugins/Experimental/EditorDataStorage/Source/TedsCore/Private/Compatibility/TedsTypedElementBridgeQueries.h
+++ b/Engine/Plugins/Experimental/EditorDataStorage/Source/TedsCore/Private/Compatibility/TedsTypedElementBridgeQueries.h
@@ -16,22 +16,22 @@ class UTypedElementRegistry;
  * are cleaned up when TEDS is shut down.
  */
 UCLASS(Transient)
-class UTypedElementBridgeDataStorageFactory : public UTypedElementDataStorageFactory
+class UTypedElementBridgeDataStorageFactory : public UEditorDataStorageFactory
 {
 	GENERATED_BODY()
 public:
 	~UTypedElementBridgeDataStorageFactory() override = default;
 	virtual uint8 GetOrder() const override;
-	virtual void PreRegister(ITypedElementDataStorageInterface& DataStorage) override;
-	virtual void PreShutdown(ITypedElementDataStorageInterface& DataStorage) override;
-	virtual void RegisterQueries(ITypedElementDataStorageInterface& DataStorage) override;
+	virtual void PreRegister(IEditorDataStorageProvider& DataStorage) override;
+	virtual void PreShutdown(IEditorDataStorageProvider& DataStorage) override;
+	virtual void RegisterQueries(IEditorDataStorageProvider& DataStorage) override;
 
 	static bool IsEnabled();
 
 private:
-	void RegisterQuery_NewUObject(ITypedElementDataStorageInterface& DataStorage);
-	void UnregisterQuery_NewUObject(ITypedElementDataStorageInterface& DataStorage);
-	void CleanupTypedElementColumns(ITypedElementDataStorageInterface& DataStorage);
+	void RegisterQuery_NewUObject(IEditorDataStorageProvider& DataStorage);
+	void UnregisterQuery_NewUObject(IEditorDataStorageProvider& DataStorage);
+	void CleanupTypedElementColumns(IEditorDataStorageProvider& DataStorage);
 	void HandleOnEnabled(IConsoleVariable* CVar);
 	
 	UE::Editor::DataStorage::QueryHandle RemoveTypedElementRowHandleQuery = UE::Editor::DataStorage::InvalidQueryHandle;
diff --git a/Engine/Plugins/Experimental/EditorDataStorage/Source/TedsCore/Private/Compatibility/TypedElementActorIconOverrideQueries.cpp b/Engine/Plugins/Experimental/EditorDataStorage/Source/TedsCore/Private/Compatibility/TypedElementActorIconOverrideQueries.cpp
--- a/Engine/Plugins/Experimental/EditorDataStorage/Source/TedsCore/Private/Compatibility/TypedElementActorIconOverrideQueries.cpp
+++ b/Engine/Plugins/Experimental/EditorDataStorage/Source/TedsCore/Private/Compatibility/TypedElementActorIconOverrideQueries.cpp
@@ -8,7 +8,7 @@
 #include "Elements/Framework/TypedElementQueryBuilder.h"
 #include "GameFramework/Actor.h"
 
-void UActorIconOverrideDataStorageFactory::RegisterQueries(ITypedElementDataStorageInterface& DataStorage)
+void UActorIconOverrideDataStorageFactory::RegisterQueries(IEditorDataStorageProvider& DataStorage)
 {
 	using namespace UE::Editor::DataStorage::Queries;
 
diff --git a/Engine/Plugins/Experimental/EditorDataStorage/Source/TedsCore/Private/Compatibility/TypedElementActorIconOverrideQueries.h b/Engine/Plugins/Experimental/EditorDataStorage/Source/TedsCore/Private/Compatibility/TypedElementActorIconOverrideQueries.h
--- a/Engine/Plugins/Experimental/EditorDataStorage/Source/TedsCore/Private/Compatibility/TypedElementActorIconOverrideQueries.h
+++ b/Engine/Plugins/Experimental/EditorDataStorage/Source/TedsCore/Private/Compatibility/TypedElementActorIconOverrideQueries.h
@@ -8,12 +8,12 @@
 #include "TypedElementActorIconOverrideQueries.generated.h"
 
 UCLASS()
-class UActorIconOverrideDataStorageFactory : public UTypedElementDataStorageFactory
+class UActorIconOverrideDataStorageFactory : public UEditorDataStorageFactory
 {
 	GENERATED_BODY()
 
 public:
 	virtual ~UActorIconOverrideDataStorageFactory() override = default;
 
-	virtual void RegisterQueries(ITypedElementDataStorageInterface& DataStorage) override;
+	virtual void RegisterQueries(IEditorDataStorageProvider& DataStorage) override;
 };
diff --git a/Engine/Plugins/Experimental/EditorDataStorage/Source/TedsCore/Private/Compatibility/TypedElementActorLabelQueries.cpp b/Engine/Plugins/Experimental/EditorDataStorage/Source/TedsCore/Private/Compatibility/TypedElementActorLabelQueries.cpp
--- a/Engine/Plugins/Experimental/EditorDataStorage/Source/TedsCore/Private/Compatibility/TypedElementActorLabelQueries.cpp
+++ b/Engine/Plugins/Experimental/EditorDataStorage/Source/TedsCore/Private/Compatibility/TypedElementActorLabelQueries.cpp
@@ -13,13 +13,13 @@
 
 #define LOCTEXT_NAMESPACE "TedsCore"
 
-void UActorLabelDataStorageFactory::RegisterQueries(ITypedElementDataStorageInterface& DataStorage)
+void UActorLabelDataStorageFactory::RegisterQueries(IEditorDataStorageProvider& DataStorage)
 {
 	RegisterActorLabelToColumnQuery(DataStorage);
 	RegisterLabelColumnToActorQuery(DataStorage);
 }
 
-void UActorLabelDataStorageFactory::RegisterActorLabelToColumnQuery(ITypedElementDataStorageInterface& DataStorage) const
+void UActorLabelDataStorageFactory::RegisterActorLabelToColumnQuery(IEditorDataStorageProvider& DataStorage) const
 {
 	using namespace UE::Editor::DataStorage::Queries;
 	
@@ -47,7 +47,7 @@ void UActorLabelDataStorageFactory::RegisterActorLabelToColumnQuery(ITypedElemen
 		.Compile());
 }
 
-void UActorLabelDataStorageFactory::RegisterLabelColumnToActorQuery(ITypedElementDataStorageInterface& DataStorage) const
+void UActorLabelDataStorageFactory::RegisterLabelColumnToActorQuery(IEditorDataStorageProvider& DataStorage) const
 {
 	using namespace UE::Editor::DataStorage::Queries;
 	
diff --git a/Engine/Plugins/Experimental/EditorDataStorage/Source/TedsCore/Private/Compatibility/TypedElementActorLabelQueries.h b/Engine/Plugins/Experimental/EditorDataStorage/Source/TedsCore/Private/Compatibility/TypedElementActorLabelQueries.h
--- a/Engine/Plugins/Experimental/EditorDataStorage/Source/TedsCore/Private/Compatibility/TypedElementActorLabelQueries.h
+++ b/Engine/Plugins/Experimental/EditorDataStorage/Source/TedsCore/Private/Compatibility/TypedElementActorLabelQueries.h
@@ -9,23 +9,23 @@
 
 
 UCLASS()
-class UActorLabelDataStorageFactory : public UTypedElementDataStorageFactory
+class UActorLabelDataStorageFactory : public UEditorDataStorageFactory
 {
 	GENERATED_BODY()
 
 public:
 	~UActorLabelDataStorageFactory() override = default;
 
-	void RegisterQueries(ITypedElementDataStorageInterface& DataStorage) override;
+	void RegisterQueries(IEditorDataStorageProvider& DataStorage) override;
 
 private:
 	/**
 	 * Takes the label set on an actor and copies it to the Data Storage if they differ.
 	 */
-	void RegisterActorLabelToColumnQuery(ITypedElementDataStorageInterface& DataStorage) const;
+	void RegisterActorLabelToColumnQuery(IEditorDataStorageProvider& DataStorage) const;
 	/**
 	 * Takes the label stored in the Data Storage and copies it to the actor's label if the FTypedElementSyncBackToWorldTag
 	 * has been set and the labels differ.
 	 */
-	void RegisterLabelColumnToActorQuery(ITypedElementDataStorageInterface& DataStorage) const;
+	void RegisterLabelColumnToActorQuery(IEditorDataStorageProvider& DataStorage) const;
 };
\ No newline at end of file
diff --git a/Engine/Plugins/Experimental/EditorDataStorage/Source/TedsCore/Private/Compatibility/TypedElementActorParentQueries.cpp b/Engine/Plugins/Experimental/EditorDataStorage/Source/TedsCore/Private/Compatibility/TypedElementActorParentQueries.cpp
--- a/Engine/Plugins/Experimental/EditorDataStorage/Source/TedsCore/Private/Compatibility/TypedElementActorParentQueries.cpp
+++ b/Engine/Plugins/Experimental/EditorDataStorage/Source/TedsCore/Private/Compatibility/TypedElementActorParentQueries.cpp
@@ -20,7 +20,7 @@ namespace UE::Editor::DataStorage::Queries::Private
 		TEXT("Mirror parent information for actors to TEDS (only works when set on startup)"));
 };
 
-void UActorParentDataStorageFactory::RegisterQueries(ITypedElementDataStorageInterface& DataStorage)
+void UActorParentDataStorageFactory::RegisterQueries(IEditorDataStorageProvider& DataStorage)
 {
 	if(UE::Editor::DataStorage::Queries::Private::bAddParentColumnToActors)
 	{
@@ -29,7 +29,7 @@ void UActorParentDataStorageFactory::RegisterQueries(ITypedElementDataStorageInt
 	}
 }
 
-void UActorParentDataStorageFactory::RegisterAddParentColumn(ITypedElementDataStorageInterface& DataStorage) const
+void UActorParentDataStorageFactory::RegisterAddParentColumn(IEditorDataStorageProvider& DataStorage) const
 {
 	using namespace UE::Editor::DataStorage::Queries;
 	
@@ -63,7 +63,7 @@ void UActorParentDataStorageFactory::RegisterAddParentColumn(ITypedElementDataSt
 		.Compile());
 }
 
-void UActorParentDataStorageFactory::RegisterUpdateOrRemoveParentColumn(ITypedElementDataStorageInterface& DataStorage) const
+void UActorParentDataStorageFactory::RegisterUpdateOrRemoveParentColumn(IEditorDataStorageProvider& DataStorage) const
 {
 	using namespace UE::Editor::DataStorage::Queries;
 	
diff --git a/Engine/Plugins/Experimental/EditorDataStorage/Source/TedsCore/Private/Compatibility/TypedElementActorParentQueries.h b/Engine/Plugins/Experimental/EditorDataStorage/Source/TedsCore/Private/Compatibility/TypedElementActorParentQueries.h
--- a/Engine/Plugins/Experimental/EditorDataStorage/Source/TedsCore/Private/Compatibility/TypedElementActorParentQueries.h
+++ b/Engine/Plugins/Experimental/EditorDataStorage/Source/TedsCore/Private/Compatibility/TypedElementActorParentQueries.h
@@ -8,24 +8,24 @@
 #include "TypedElementActorParentQueries.generated.h"
 
 UCLASS()
-class UActorParentDataStorageFactory : public UTypedElementDataStorageFactory
+class UActorParentDataStorageFactory : public UEditorDataStorageFactory
 {
 	GENERATED_BODY()
 
 public:
 	~UActorParentDataStorageFactory() override = default;
 
-	void RegisterQueries(ITypedElementDataStorageInterface& DataStorage) override;
+	void RegisterQueries(IEditorDataStorageProvider& DataStorage) override;
 
 private:
 	/**
 	 * Checks rows with actors that don't have a parent column yet if one needs to be added whenever
 	 * the row is marked for updates.
 	 */
-	void RegisterAddParentColumn(ITypedElementDataStorageInterface& DataStorage) const;
+	void RegisterAddParentColumn(IEditorDataStorageProvider& DataStorage) const;
 	/**
 	 * Updates the parent column with the parent from the actor or removes it if there's no parent associated
 	 * with the actor anymore.
 	 */
-	void RegisterUpdateOrRemoveParentColumn(ITypedElementDataStorageInterface& DataStorage) const;
+	void RegisterUpdateOrRemoveParentColumn(IEditorDataStorageProvider& DataStorage) const;
 };
diff --git a/Engine/Plugins/Experimental/EditorDataStorage/Source/TedsCore/Private/Compatibility/TypedElementActorTransformQueries.cpp b/Engine/Plugins/Experimental/EditorDataStorage/Source/TedsCore/Private/Compatibility/TypedElementActorTransformQueries.cpp
--- a/Engine/Plugins/Experimental/EditorDataStorage/Source/TedsCore/Private/Compatibility/TypedElementActorTransformQueries.cpp
+++ b/Engine/Plugins/Experimental/EditorDataStorage/Source/TedsCore/Private/Compatibility/TypedElementActorTransformQueries.cpp
@@ -8,14 +8,14 @@
 #include "Elements/Framework/TypedElementQueryBuilder.h"
 #include "GameFramework/Actor.h"
 
-void UActorTransformDataStorageFactory::RegisterQueries(ITypedElementDataStorageInterface& DataStorage)
+void UActorTransformDataStorageFactory::RegisterQueries(IEditorDataStorageProvider& DataStorage)
 {
 	RegisterActorAddTransformColumn(DataStorage);
 	RegisterActorLocalTransformToColumn(DataStorage);
 	RegisterLocalTransformColumnToActor(DataStorage);
 }
 
-void UActorTransformDataStorageFactory::RegisterActorAddTransformColumn(ITypedElementDataStorageInterface& DataStorage) const
+void UActorTransformDataStorageFactory::RegisterActorAddTransformColumn(IEditorDataStorageProvider& DataStorage) const
 {
 	using namespace UE::Editor::DataStorage::Queries;
 
@@ -38,7 +38,7 @@ void UActorTransformDataStorageFactory::RegisterActorAddTransformColumn(ITypedEl
 		.Compile());
 }
 
-void UActorTransformDataStorageFactory::RegisterActorLocalTransformToColumn(ITypedElementDataStorageInterface& DataStorage) const
+void UActorTransformDataStorageFactory::RegisterActorLocalTransformToColumn(IEditorDataStorageProvider& DataStorage) const
 {
 	using namespace UE::Editor::DataStorage::Queries;
 
@@ -65,7 +65,7 @@ void UActorTransformDataStorageFactory::RegisterActorLocalTransformToColumn(ITyp
 		.Compile());
 }
 
-void UActorTransformDataStorageFactory::RegisterLocalTransformColumnToActor(ITypedElementDataStorageInterface& DataStorage) const
+void UActorTransformDataStorageFactory::RegisterLocalTransformColumnToActor(IEditorDataStorageProvider& DataStorage) const
 {
 	using namespace UE::Editor::DataStorage::Queries;
 	
diff --git a/Engine/Plugins/Experimental/EditorDataStorage/Source/TedsCore/Private/Compatibility/TypedElementActorTransformQueries.h b/Engine/Plugins/Experimental/EditorDataStorage/Source/TedsCore/Private/Compatibility/TypedElementActorTransformQueries.h
--- a/Engine/Plugins/Experimental/EditorDataStorage/Source/TedsCore/Private/Compatibility/TypedElementActorTransformQueries.h
+++ b/Engine/Plugins/Experimental/EditorDataStorage/Source/TedsCore/Private/Compatibility/TypedElementActorTransformQueries.h
@@ -8,29 +8,29 @@
 #include "TypedElementActorTransformQueries.generated.h"
 
 UCLASS()
-class UActorTransformDataStorageFactory : public UTypedElementDataStorageFactory
+class UActorTransformDataStorageFactory : public UEditorDataStorageFactory
 {
 	GENERATED_BODY()
 
 public:
 	~UActorTransformDataStorageFactory() override = default;
 
-	void RegisterQueries(ITypedElementDataStorageInterface& DataStorage) override;
+	void RegisterQueries(IEditorDataStorageProvider& DataStorage) override;
 
 private:
 	/**
 	 * Checks actors that don't have a transform column and adds one if an actor has been
 	 * assigned a transform.
 	 */
-	void RegisterActorAddTransformColumn(ITypedElementDataStorageInterface& DataStorage) const;
+	void RegisterActorAddTransformColumn(IEditorDataStorageProvider& DataStorage) const;
 	/**
 	 * Takes the transform set on an actor and copies it to the Data Storage or removes the
 	 * transform column if there's not transform available anymore.
 	 */
-	void RegisterActorLocalTransformToColumn(ITypedElementDataStorageInterface& DataStorage) const;
+	void RegisterActorLocalTransformToColumn(IEditorDataStorageProvider& DataStorage) const;
 	/**
 	 * Takes the transform stored in the Data Storage and copies it to the actor's tranform if
 	 * the FTypedElementSyncBackToWorldTag has been set.
 	 */
-	void RegisterLocalTransformColumnToActor(ITypedElementDataStorageInterface& DataStorage) const;
+	void RegisterLocalTransformColumnToActor(IEditorDataStorageProvider& DataStorage) const;
 };
diff --git a/Engine/Plugins/Experimental/EditorDataStorage/Source/TedsCore/Private/Compatibility/TypedElementActorViewportProcessors.cpp b/Engine/Plugins/Experimental/EditorDataStorage/Source/TedsCore/Private/Compatibility/TypedElementActorViewportProcessors.cpp
--- a/Engine/Plugins/Experimental/EditorDataStorage/Source/TedsCore/Private/Compatibility/TypedElementActorViewportProcessors.cpp
+++ b/Engine/Plugins/Experimental/EditorDataStorage/Source/TedsCore/Private/Compatibility/TypedElementActorViewportProcessors.cpp
@@ -23,7 +23,7 @@ namespace UE::Editor::DataStorage::Private
 
 				TRACE_CPUPROFILER_EVENT_SCOPE(TEDS.Debug.AddOverlayColorToSelectionCommand);
 
-				if (ITypedElementDataStorageInterface* DataStorage = UTypedElementRegistry::GetInstance()->GetMutableDataStorage())
+				if (IEditorDataStorageProvider* DataStorage = UTypedElementRegistry::GetInstance()->GetMutableDataStorage())
 				{
 					static QueryHandle OverlayQuery = InvalidQueryHandle;
 					if (OverlayQuery == InvalidQueryHandle)
@@ -81,7 +81,7 @@ namespace UE::Editor::DataStorage::Private
 
 				TRACE_CPUPROFILER_EVENT_SCOPE(TEDS.Debug.AddOverlayColorToSelectionCommand);
 
-				if (ITypedElementDataStorageInterface* DataStorage = UTypedElementRegistry::GetInstance()->GetMutableDataStorage())
+				if (IEditorDataStorageProvider* DataStorage = UTypedElementRegistry::GetInstance()->GetMutableDataStorage())
 				{
 					static QueryHandle OverlayQuery = InvalidQueryHandle;
 					if (OverlayQuery == InvalidQueryHandle)
@@ -132,7 +132,7 @@ namespace UE::Editor::DataStorage::Private
 
 				TRACE_CPUPROFILER_EVENT_SCOPE(TEDS.Debug.AddOverlayColorToSelectionCommand);
 
-				if (ITypedElementDataStorageInterface* DataStorage = UTypedElementRegistry::GetInstance()->GetMutableDataStorage())
+				if (IEditorDataStorageProvider* DataStorage = UTypedElementRegistry::GetInstance()->GetMutableDataStorage())
 				{
 					static QueryHandle OverlayQuery = InvalidQueryHandle;
 					if (OverlayQuery == InvalidQueryHandle)
@@ -164,13 +164,13 @@ namespace UE::Editor::DataStorage::Private
 			}));
 } // namespace UE::Editor::DataStorage::Private
 
-void UActorViewportDataStorageFactory::RegisterQueries(ITypedElementDataStorageInterface& DataStorage)
+void UActorViewportDataStorageFactory::RegisterQueries(IEditorDataStorageProvider& DataStorage)
 {
 	RegisterOutlineColorColumnToActor(DataStorage);
 	RegisterOverlayColorColumnToActor(DataStorage);
 }
 
-void UActorViewportDataStorageFactory::RegisterOutlineColorColumnToActor(ITypedElementDataStorageInterface& DataStorage)
+void UActorViewportDataStorageFactory::RegisterOutlineColorColumnToActor(IEditorDataStorageProvider& DataStorage)
 {
 	using namespace UE::Editor::DataStorage::Queries;
 
@@ -195,7 +195,7 @@ void UActorViewportDataStorageFactory::RegisterOutlineColorColumnToActor(ITypedE
 		.Compile());
 }
 
-void UActorViewportDataStorageFactory::RegisterOverlayColorColumnToActor(ITypedElementDataStorageInterface& DataStorage)
+void UActorViewportDataStorageFactory::RegisterOverlayColorColumnToActor(IEditorDataStorageProvider& DataStorage)
 {
 	using namespace UE::Editor::DataStorage::Queries;
 	
diff --git a/Engine/Plugins/Experimental/EditorDataStorage/Source/TedsCore/Private/Compatibility/TypedElementActorViewportProcessors.h b/Engine/Plugins/Experimental/EditorDataStorage/Source/TedsCore/Private/Compatibility/TypedElementActorViewportProcessors.h
--- a/Engine/Plugins/Experimental/EditorDataStorage/Source/TedsCore/Private/Compatibility/TypedElementActorViewportProcessors.h
+++ b/Engine/Plugins/Experimental/EditorDataStorage/Source/TedsCore/Private/Compatibility/TypedElementActorViewportProcessors.h
@@ -8,16 +8,16 @@
 #include "TypedElementActorViewportProcessors.generated.h"
 
 UCLASS()
-class UActorViewportDataStorageFactory : public UTypedElementDataStorageFactory
+class UActorViewportDataStorageFactory : public UEditorDataStorageFactory
 {
 	GENERATED_BODY()
 
 public:
 	~UActorViewportDataStorageFactory() override = default;
 
-	void RegisterQueries(ITypedElementDataStorageInterface& DataStorage) override;
+	void RegisterQueries(IEditorDataStorageProvider& DataStorage) override;
 
 private:
-	void RegisterOutlineColorColumnToActor(ITypedElementDataStorageInterface& DataStorage);
-	void RegisterOverlayColorColumnToActor(ITypedElementDataStorageInterface& DataStorage);
+	void RegisterOutlineColorColumnToActor(IEditorDataStorageProvider& DataStorage);
+	void RegisterOverlayColorColumnToActor(IEditorDataStorageProvider& DataStorage);
 };
diff --git a/Engine/Plugins/Experimental/EditorDataStorage/Source/TedsCore/Private/Compatibility/TypedElementObjectReinstancingManager.cpp b/Engine/Plugins/Experimental/EditorDataStorage/Source/TedsCore/Private/Compatibility/TypedElementObjectReinstancingManager.cpp
--- a/Engine/Plugins/Experimental/EditorDataStorage/Source/TedsCore/Private/Compatibility/TypedElementObjectReinstancingManager.cpp
+++ b/Engine/Plugins/Experimental/EditorDataStorage/Source/TedsCore/Private/Compatibility/TypedElementObjectReinstancingManager.cpp
@@ -71,7 +71,7 @@ void UTedsObjectReinstancingManager::HandleOnObjectsReinstanced(
 {
 	using namespace UE::Editor::DataStorage;
 
-	ITypedElementDataStorageInterface* Interface = DataStorage;
+	IEditorDataStorageProvider* Interface = DataStorage;
 	for (FCoreUObjectDelegates::FReplacementObjectMap::TConstIterator Iter = ObjectReplacementMap.CreateConstIterator(); Iter; ++Iter)
 	{
 		const void* PreDeleteObject = Iter->Key;
diff --git a/Engine/Plugins/Experimental/EditorDataStorage/Source/TedsCore/Private/Compatibility/TypedElementObjectReinstancingManager.h b/Engine/Plugins/Experimental/EditorDataStorage/Source/TedsCore/Private/Compatibility/TypedElementObjectReinstancingManager.h
--- a/Engine/Plugins/Experimental/EditorDataStorage/Source/TedsCore/Private/Compatibility/TypedElementObjectReinstancingManager.h
+++ b/Engine/Plugins/Experimental/EditorDataStorage/Source/TedsCore/Private/Compatibility/TypedElementObjectReinstancingManager.h
@@ -9,8 +9,7 @@
 
 #include "TypedElementObjectReinstancingManager.generated.h"
 
-class ITypedElementDataStorageCompatibilityInterface;
-class ITypedElementDataStorageInterface;
+class IEditorDataStorageCompatibilityProvider;
 class UEditorDataStorageCompatibility;
 class UTypedElementMementoSystem;
 class UEditorDataStorage;
diff --git a/Engine/Plugins/Experimental/EditorDataStorage/Source/TedsCore/Private/Compatibility/TypedElementUObjectWorldQueries.cpp b/Engine/Plugins/Experimental/EditorDataStorage/Source/TedsCore/Private/Compatibility/TypedElementUObjectWorldQueries.cpp
--- a/Engine/Plugins/Experimental/EditorDataStorage/Source/TedsCore/Private/Compatibility/TypedElementUObjectWorldQueries.cpp
+++ b/Engine/Plugins/Experimental/EditorDataStorage/Source/TedsCore/Private/Compatibility/TypedElementUObjectWorldQueries.cpp
@@ -7,13 +7,13 @@
 #include "Elements/Framework/TypedElementQueryBuilder.h"
 #include "Engine/World.h"
 
-void UObjectWorldDataStorageFactory::RegisterQueries(ITypedElementDataStorageInterface& DataStorage)
+void UObjectWorldDataStorageFactory::RegisterQueries(IEditorDataStorageProvider& DataStorage)
 {
 	RegisterAddWorldColumn(DataStorage);
 	RegisterUpdateOrRemoveWorldColumn(DataStorage);
 }
 
-void UObjectWorldDataStorageFactory::RegisterAddWorldColumn(ITypedElementDataStorageInterface& DataStorage) const
+void UObjectWorldDataStorageFactory::RegisterAddWorldColumn(IEditorDataStorageProvider& DataStorage) const
 {
 	using namespace UE::Editor::DataStorage::Queries;
 	
@@ -40,7 +40,7 @@ void UObjectWorldDataStorageFactory::RegisterAddWorldColumn(ITypedElementDataSto
 		.Compile());
 }
 
-void UObjectWorldDataStorageFactory::RegisterUpdateOrRemoveWorldColumn(ITypedElementDataStorageInterface& DataStorage) const
+void UObjectWorldDataStorageFactory::RegisterUpdateOrRemoveWorldColumn(IEditorDataStorageProvider& DataStorage) const
 {
 	using namespace UE::Editor::DataStorage::Queries;
 	
diff --git a/Engine/Plugins/Experimental/EditorDataStorage/Source/TedsCore/Private/Compatibility/TypedElementUObjectWorldQueries.h b/Engine/Plugins/Experimental/EditorDataStorage/Source/TedsCore/Private/Compatibility/TypedElementUObjectWorldQueries.h
--- a/Engine/Plugins/Experimental/EditorDataStorage/Source/TedsCore/Private/Compatibility/TypedElementUObjectWorldQueries.h
+++ b/Engine/Plugins/Experimental/EditorDataStorage/Source/TedsCore/Private/Compatibility/TypedElementUObjectWorldQueries.h
@@ -8,24 +8,24 @@
 #include "TypedElementUObjectWorldQueries.generated.h"
 
 UCLASS()
-class UObjectWorldDataStorageFactory : public UTypedElementDataStorageFactory
+class UObjectWorldDataStorageFactory : public UEditorDataStorageFactory
 {
 	GENERATED_BODY()
 
 public:
 	~UObjectWorldDataStorageFactory() override = default;
 
-	void RegisterQueries(ITypedElementDataStorageInterface& DataStorage) override;
+	void RegisterQueries(IEditorDataStorageProvider& DataStorage) override;
 
 private:
 	/**
 	 * Checks rows with UObjects that don't have a world column yet if one needs to be added whenever
 	 * the row is marked for updates.
 	 */
-	void RegisterAddWorldColumn(ITypedElementDataStorageInterface& DataStorage) const;
+	void RegisterAddWorldColumn(IEditorDataStorageProvider& DataStorage) const;
 	/**
 	 * Updates the world column with the world in the UObject or removes it if there's no world associated
 	 * with the UObject anymore.
 	 */
-	void RegisterUpdateOrRemoveWorldColumn(ITypedElementDataStorageInterface& DataStorage) const;
+	void RegisterUpdateOrRemoveWorldColumn(IEditorDataStorageProvider& DataStorage) const;
 };
diff --git a/Engine/Plugins/Experimental/EditorDataStorage/Source/TedsCore/Private/Debug/TypedElementDatabaseDebugCommands.cpp b/Engine/Plugins/Experimental/EditorDataStorage/Source/TedsCore/Private/Debug/TypedElementDatabaseDebugCommands.cpp
--- a/Engine/Plugins/Experimental/EditorDataStorage/Source/TedsCore/Private/Debug/TypedElementDatabaseDebugCommands.cpp
+++ b/Engine/Plugins/Experimental/EditorDataStorage/Source/TedsCore/Private/Debug/TypedElementDatabaseDebugCommands.cpp
@@ -23,7 +23,7 @@ namespace UE::Editor::DataStorage::Private
 	static const UStruct* GetTypeInfo(const TObjectPtr<const UScriptStruct>& TypeInfo) { return TypeInfo; }
 
 	template<typename TypeInfoType>
-	void PrintObjectTypeInformation(ITypedElementDataStorageInterface* DataStorage, FString Message, FOutputDevice& Output)
+	void PrintObjectTypeInformation(IEditorDataStorageProvider* DataStorage, FString Message, FOutputDevice& Output)
 	{
 		using namespace UE::Editor::DataStorage::Queries;
 
@@ -65,7 +65,7 @@ namespace UE::Editor::DataStorage::Private
 	{
 		using namespace UE::Editor::DataStorage::Queries;
 
-		if (ITypedElementDataStorageInterface* DataStorage = UTypedElementRegistry::GetInstance()->GetMutableDataStorage())
+		if (IEditorDataStorageProvider* DataStorage = UTypedElementRegistry::GetInstance()->GetMutableDataStorage())
 		{
 			static QueryHandle LabelQuery = [DataStorage]
 			{
@@ -136,7 +136,7 @@ FAutoConsoleCommandWithOutputDevice PrintObjectTypeInformationConsoleCommand(
 
 			TRACE_CPUPROFILER_EVENT_SCOPE(TEDS.Debug.PrintObjectTypeInfo);
 
-			if (ITypedElementDataStorageInterface* DataStorage = UTypedElementRegistry::GetInstance()->GetMutableDataStorage())
+			if (IEditorDataStorageProvider* DataStorage = UTypedElementRegistry::GetInstance()->GetMutableDataStorage())
 			{
 				FString Message;
 				Output.Log(TEXT("The Typed Elements Data Storage has the types:"));
@@ -190,17 +190,17 @@ FAutoConsoleCommandWithOutputDevice ListExtensionsConsoleCommand(
 
 			UTypedElementRegistry* Registry = UTypedElementRegistry::GetInstance();
 
-			if (const ITypedElementDataStorageInterface* DataStorage = Registry->GetDataStorage())
+			if (const IEditorDataStorageProvider* DataStorage = Registry->GetDataStorage())
 			{
 				Message = TEXT("Data Storage Extensions: \n");
 				DataStorage->ListExtensions(RecordExtensions);
 			}
-			if (const ITypedElementDataStorageCompatibilityInterface* DataStorageCompat = Registry->GetDataStorageCompatibility())
+			if (const IEditorDataStorageCompatibilityProvider* DataStorageCompat = Registry->GetDataStorageCompatibility())
 			{
 				Message += TEXT("Data Storage Compatibility Extensions: \n");
 				DataStorageCompat->ListExtensions(RecordExtensions);
 			}
-			if (const ITypedElementDataStorageUiInterface* DataStorageUi = Registry->GetDataStorageUi())
+			if (const IEditorDataStorageUiProvider* DataStorageUi = Registry->GetDataStorageUi())
 			{
 				Message += TEXT("Data Storage UI Extensions: \n");
 				DataStorageUi->ListExtensions(RecordExtensions);
@@ -217,7 +217,7 @@ static FAutoConsoleCommand CVarCreateRow(
 	TEXT("Argument: \n"),
 	FConsoleCommandWithArgsDelegate::CreateLambda([](const TArray<FString>& Args)
 	{
-		ITypedElementDataStorageInterface* DataStorage = UTypedElementRegistry::GetInstance()->GetMutableDataStorage();
+		IEditorDataStorageProvider* DataStorage = UTypedElementRegistry::GetInstance()->GetMutableDataStorage();
 		static UE::Editor::DataStorage::TableHandle Table = DataStorage->RegisterTable<FTestColumnA>(FName(TEXT("Debug.CreateRow Table")));
 
 		const UE::Editor::DataStorage::RowHandle RowHandle = DataStorage->AddRow(Table);
@@ -230,7 +230,7 @@ static FAutoConsoleCommand CVarAddDynamicColumnTag(
 	TEXT("Argument: Row, Identifier\n"),
 	FConsoleCommandWithArgsDelegate::CreateLambda([](const TArray<FString>& Args)
 	{
-		ITypedElementDataStorageInterface* DataStorage = UTypedElementRegistry::GetInstance()->GetMutableDataStorage();
+		IEditorDataStorageProvider* DataStorage = UTypedElementRegistry::GetInstance()->GetMutableDataStorage();
 		
 		if (Args.Num() != 2)
 		{
@@ -252,7 +252,7 @@ static FAutoConsoleCommand CVarAddDynamicColumn(
 	TEXT("Argument: Row, Identifier\n"),
 	FConsoleCommandWithArgsDelegate::CreateLambda([](const TArray<FString>& Args)
 	{
-		ITypedElementDataStorageInterface* DataStorage = UTypedElementRegistry::GetInstance()->GetMutableDataStorage();
+		IEditorDataStorageProvider* DataStorage = UTypedElementRegistry::GetInstance()->GetMutableDataStorage();
 		
 		if (Args.Num() != 2)
 		{
@@ -283,7 +283,7 @@ static FAutoConsoleCommand CVarRemoveDynamicColumn(
 	TEXT("Argument: Row, Identifier\n"),
 	FConsoleCommandWithArgsDelegate::CreateLambda([](const TArray<FString>& Args)
 	{
-		ITypedElementDataStorageInterface* DataStorage = UTypedElementRegistry::GetInstance()->GetMutableDataStorage();
+		IEditorDataStorageProvider* DataStorage = UTypedElementRegistry::GetInstance()->GetMutableDataStorage();
 		
 		if (Args.Num() != 2)
 		{
@@ -305,7 +305,7 @@ static FAutoConsoleCommand CVarAddToDynamicColumn(
 	TEXT("Argument: Row, TagId, Value, [optional] MethodId\n"),
 	FConsoleCommandWithArgsDelegate::CreateLambda([](const TArray<FString>& Args)
 	{
-		ITypedElementDataStorageInterface* DataStorage = UTypedElementRegistry::GetInstance()->GetMutableDataStorage();
+		IEditorDataStorageProvider* DataStorage = UTypedElementRegistry::GetInstance()->GetMutableDataStorage();
 		
 		if (Args.Num() < 3 || Args.Num() > 4)
 		{
@@ -357,7 +357,7 @@ static FAutoConsoleCommand CVarPrintDynamicColumn(
 	TEXT("Argument: Row, TagId\n"),
 	FConsoleCommandWithArgsDelegate::CreateLambda([](const TArray<FString>& Args)
 	{
-		ITypedElementDataStorageInterface* DataStorage = UTypedElementRegistry::GetInstance()->GetMutableDataStorage();
+		IEditorDataStorageProvider* DataStorage = UTypedElementRegistry::GetInstance()->GetMutableDataStorage();
 		
 		if (Args.Num() != 2)
 		{
@@ -413,7 +413,7 @@ static FAutoConsoleCommand CVarPrintDynamicColumnWithQuery(
 			return;
 		}
 		
-		ITypedElementDataStorageInterface* DataStorage = UTypedElementRegistry::GetInstance()->GetMutableDataStorage();
+		IEditorDataStorageProvider* DataStorage = UTypedElementRegistry::GetInstance()->GetMutableDataStorage();
 		
 		const FName Identifier(*Args[0]);
 
@@ -473,7 +473,7 @@ static FAutoConsoleCommand CVarCountDynamicTagWithQuery(
 			return;
 		}
 				
-		ITypedElementDataStorageInterface* DataStorage = UTypedElementRegistry::GetInstance()->GetMutableDataStorage();
+		IEditorDataStorageProvider* DataStorage = UTypedElementRegistry::GetInstance()->GetMutableDataStorage();
 				
 		const FName Identifier(*Args[0]);
 
@@ -514,7 +514,7 @@ static FAutoConsoleCommand CVarRegisterListDynamicColumnQuery(
 			return;
 		}
 				
-		ITypedElementDataStorageInterface* DataStorage = UTypedElementRegistry::GetInstance()->GetMutableDataStorage();
+		IEditorDataStorageProvider* DataStorage = UTypedElementRegistry::GetInstance()->GetMutableDataStorage();
 				
 		const FName Identifier(*Args[0]);
 		const FName ActivationGroup(*Args[1]);
@@ -553,7 +553,7 @@ static FAutoConsoleCommand CVarActivateListDynamicColumnQuery(
 			return;
 		}
 						
-		ITypedElementDataStorageInterface* DataStorage = UTypedElementRegistry::GetInstance()->GetMutableDataStorage();
+		IEditorDataStorageProvider* DataStorage = UTypedElementRegistry::GetInstance()->GetMutableDataStorage();
 		
 		const FName ActivationGroup(*Args[0]);
 
@@ -567,7 +567,7 @@ static FAutoConsoleCommand CVarAddValueTag(
 	FConsoleCommandWithArgsDelegate::CreateLambda([](const TArray<FString>& Args)
 	{
 		using namespace UE::Editor::DataStorage;
-		ITypedElementDataStorageInterface* DataStorage = UTypedElementRegistry::GetInstance()->GetMutableDataStorage();
+		IEditorDataStorageProvider* DataStorage = UTypedElementRegistry::GetInstance()->GetMutableDataStorage();
 
 		if (Args.Num() != 3)
 		{
@@ -601,7 +601,7 @@ static FAutoConsoleCommand CVarRemoveValueTag(
 	TEXT("Argument: Row, Group\n"),
 	FConsoleCommandWithArgsDelegate::CreateLambda([](const TArray<FString>& Args)
 	{
-		ITypedElementDataStorageInterface* DataStorage = UTypedElementRegistry::GetInstance()->GetMutableDataStorage();
+		IEditorDataStorageProvider* DataStorage = UTypedElementRegistry::GetInstance()->GetMutableDataStorage();
 
 		if (Args.Num() != 2)
 		{
@@ -633,7 +633,7 @@ static FAutoConsoleCommand CVarMatchValueTag(
 	{
 		using namespace UE::Editor::DataStorage::Queries;
 		
-		ITypedElementDataStorageInterface* DataStorage = UTypedElementRegistry::GetInstance()->GetMutableDataStorage();
+		IEditorDataStorageProvider* DataStorage = UTypedElementRegistry::GetInstance()->GetMutableDataStorage();
 
 		if (Args.Num() < 1 || Args.Num() > 2)
 		{
@@ -687,7 +687,7 @@ static FAutoConsoleCommand CVarAddValueTagFromEnum(
 	{
 		using namespace UE::Editor::DataStorage;
 
-		ITypedElementDataStorageInterface* DataStorage = UTypedElementRegistry::GetInstance()->GetMutableDataStorage();
+		IEditorDataStorageProvider* DataStorage = UTypedElementRegistry::GetInstance()->GetMutableDataStorage();
 
 		if (Args.Num() < 1 || Args.Num() > 2)
 		{
@@ -725,7 +725,7 @@ static FAutoConsoleCommand CVarRemoveValueTagFromEnum(
 	FConsoleCommandWithArgsDelegate::CreateLambda([](const TArray<FString>& Args)
 	{
 		using namespace UE::Editor::DataStorage;
-		ITypedElementDataStorageInterface* DataStorage = UTypedElementRegistry::GetInstance()->GetMutableDataStorage();
+		IEditorDataStorageProvider* DataStorage = UTypedElementRegistry::GetInstance()->GetMutableDataStorage();
 
 		if (Args.Num() != 1)
 		{
@@ -747,7 +747,7 @@ static FAutoConsoleCommand CVarMatchValueTagFromEnum(
 	{
 		using namespace UE::Editor::DataStorage::Queries;
 		
-		ITypedElementDataStorageInterface* DataStorage = UTypedElementRegistry::GetInstance()->GetMutableDataStorage();
+		IEditorDataStorageProvider* DataStorage = UTypedElementRegistry::GetInstance()->GetMutableDataStorage();
 
 		if (Args.Num() > 1)
 		{
diff --git a/Engine/Plugins/Experimental/EditorDataStorage/Source/TedsCore/Private/Memento/TypedElementMementoSystem.cpp b/Engine/Plugins/Experimental/EditorDataStorage/Source/TedsCore/Private/Memento/TypedElementMementoSystem.cpp
--- a/Engine/Plugins/Experimental/EditorDataStorage/Source/TedsCore/Private/Memento/TypedElementMementoSystem.cpp
+++ b/Engine/Plugins/Experimental/EditorDataStorage/Source/TedsCore/Private/Memento/TypedElementMementoSystem.cpp
@@ -11,7 +11,7 @@
 
 namespace UE::Editor::DataStorage
 {
-	FMementoSystem::FMementoSystem(ITypedElementDataStorageInterface& InDataStorage)
+	FMementoSystem::FMementoSystem(IEditorDataStorageProvider& InDataStorage)
 		: DataStorage(InDataStorage)
 	{
 		FScopedExclusiveLock Lock(EGlobalLockScope::Public);
diff --git a/Engine/Plugins/Experimental/EditorDataStorage/Source/TedsCore/Private/Memento/TypedElementMementoSystem.h b/Engine/Plugins/Experimental/EditorDataStorage/Source/TedsCore/Private/Memento/TypedElementMementoSystem.h
--- a/Engine/Plugins/Experimental/EditorDataStorage/Source/TedsCore/Private/Memento/TypedElementMementoSystem.h
+++ b/Engine/Plugins/Experimental/EditorDataStorage/Source/TedsCore/Private/Memento/TypedElementMementoSystem.h
@@ -5,7 +5,7 @@
 #include "Containers/Array.h"
 #include "Elements/Common/TypedElementHandles.h"
 
-class ITypedElementDataStorageInterface;
+class IEditorDataStorageProvider;
 class UTedsMementoTranslatorBase;
 
 namespace UE::Editor::DataStorage
@@ -13,7 +13,7 @@ namespace UE::Editor::DataStorage
 	class FMementoSystem
 	{
 	public:
-		explicit FMementoSystem(ITypedElementDataStorageInterface& InDataStorage);
+		explicit FMementoSystem(IEditorDataStorageProvider& InDataStorage);
 
 		RowHandle CreateMemento(RowHandle SourceRow);
 		void CreateMemento(RowHandle ReservedMementoRow, RowHandle SourceRow);
@@ -25,6 +25,6 @@ namespace UE::Editor::DataStorage
 
 		TArray<const UTedsMementoTranslatorBase*> MementoTranslators;
 		TableHandle MementoRowBaseTable;
-		ITypedElementDataStorageInterface& DataStorage;
+		IEditorDataStorageProvider& DataStorage;
 	};
 } // namespace UE::Editor::DataStorage
diff --git a/Engine/Plugins/Experimental/EditorDataStorage/Source/TedsCore/Private/Processors/TypedElementProcessorAdaptors.cpp b/Engine/Plugins/Experimental/EditorDataStorage/Source/TedsCore/Private/Processors/TypedElementProcessorAdaptors.cpp
--- a/Engine/Plugins/Experimental/EditorDataStorage/Source/TedsCore/Private/Processors/TypedElementProcessorAdaptors.cpp
+++ b/Engine/Plugins/Experimental/EditorDataStorage/Source/TedsCore/Private/Processors/TypedElementProcessorAdaptors.cpp
@@ -70,7 +70,7 @@ namespace UE::Editor::DataStorage
 				for (int32 Index = 0; Index < TypeCount; ++Index)
 				{
 					checkf(ColumnTypes->IsValid(), TEXT("Attempting to retrieve a column that is not available."));
-					*RetrievedAddresses = *AccessTypes == ITypedElementDataStorageInterface::EQueryAccessType::ReadWrite
+					*RetrievedAddresses = *AccessTypes == IEditorDataStorageProvider::EQueryAccessType::ReadWrite
 						? reinterpret_cast<char*>(Context.GetMutableFragmentView(ColumnTypes->Get()).GetData())
 						: const_cast<char*>(reinterpret_cast<const char*>(Context.GetFragmentView(ColumnTypes->Get()).GetData()));
 
@@ -614,7 +614,7 @@ namespace UE::Editor::DataStorage
 			}
 		};
 
-		struct FMassDirectContextForwarder final : public ITypedElementDataStorageInterface::IDirectQueryContext
+		struct FMassDirectContextForwarder final : public IEditorDataStorageProvider::IDirectQueryContext
 		{
 			FMassDirectContextForwarder(FMassExecutionContext& InContext, FEnvironment& InEnvironment)
 				: Implementation(InContext, InEnvironment)
@@ -633,7 +633,7 @@ namespace UE::Editor::DataStorage
 			FMassContextCommon Implementation;
 		};
 
-		struct FMassSubqueryContextForwarder  : public ITypedElementDataStorageInterface::ISubqueryContext
+		struct FMassSubqueryContextForwarder  : public IEditorDataStorageProvider::ISubqueryContext
 		{
 			FMassSubqueryContextForwarder(FMassExecutionContext& InContext, FEnvironment& InEnvironment)
 				: Implementation(InContext, InEnvironment)
@@ -673,7 +673,7 @@ namespace UE::Editor::DataStorage
 		struct FMassQueryContextImplementation final : FMassWithEnvironmentContextCommon
 		{
 			FMassQueryContextImplementation(
-				ITypedElementDataStorageInterface::FQueryDescription& InQueryDescription,
+				IEditorDataStorageProvider::FQueryDescription& InQueryDescription,
 				FMassExecutionContext& InContext, 
 				FExtendedQueryStore& InQueryStore,
 				FEnvironment& InEnvironment)
@@ -697,7 +697,7 @@ namespace UE::Editor::DataStorage
 			}
 	
 			void GetDependencies(TArrayView<UObject*> RetrievedAddresses, TConstArrayView<TWeakObjectPtr<const UClass>> SubsystemTypes,
-				TConstArrayView<ITypedElementDataStorageInterface::EQueryAccessType> AccessTypes)
+				TConstArrayView<IEditorDataStorageProvider::EQueryAccessType> AccessTypes)
 			{
 				checkf(RetrievedAddresses.Num() == SubsystemTypes.Num(), TEXT("Unable to retrieve a batch of subsystem as the number of addresses "
 					"doesn't match the number of requested subsystem types."));
@@ -706,12 +706,12 @@ namespace UE::Editor::DataStorage
 			}
 
 			void GetDependenciesUnguarded(int32 SubsystemCount, UObject** RetrievedAddresses, const TWeakObjectPtr<const UClass>* DependencyTypes,
-				const ITypedElementDataStorageInterface::EQueryAccessType* AccessTypes)
+				const IEditorDataStorageProvider::EQueryAccessType* AccessTypes)
 			{
 				for (int32 Index = 0; Index < SubsystemCount; ++Index)
 				{
 					checkf(DependencyTypes->IsValid(), TEXT("Attempting to retrieve a subsystem that's no longer valid."));
-					*RetrievedAddresses = *AccessTypes == ITypedElementDataStorageInterface::EQueryAccessType::ReadWrite
+					*RetrievedAddresses = *AccessTypes == IEditorDataStorageProvider::EQueryAccessType::ReadWrite
 						? Context.GetMutableSubsystem<USubsystem>(const_cast<UClass*>(DependencyTypes->Get()))
 						: const_cast<USubsystem*>(Context.GetSubsystem<USubsystem>(const_cast<UClass*>(DependencyTypes->Get())));
 
@@ -772,14 +772,14 @@ namespace UE::Editor::DataStorage
 				}
 			}
 	
-			ITypedElementDataStorageInterface::FQueryDescription& QueryDescription;
+			IEditorDataStorageProvider::FQueryDescription& QueryDescription;
 			FExtendedQueryStore& QueryStore;
 		};
 
-		struct FMassContextForwarder final : public ITypedElementDataStorageInterface::IQueryContext
+		struct FMassContextForwarder final : public IEditorDataStorageProvider::IQueryContext
 		{
 			FMassContextForwarder(
-				ITypedElementDataStorageInterface::FQueryDescription& InQueryDescription,
+				IEditorDataStorageProvider::FQueryDescription& InQueryDescription,
 				FMassExecutionContext& InContext, 
 				FExtendedQueryStore& InQueryStore,
 				FEnvironment& InEnvironment)
@@ -840,11 +840,11 @@ namespace UE::Editor::DataStorage
 	}
 
 	void FPhasePreOrPostAmbleExecutor::ExecuteQuery(
-		ITypedElementDataStorageInterface::FQueryDescription& Description,
+		IEditorDataStorageProvider::FQueryDescription& Description,
 		FExtendedQueryStore& QueryStore,
 		FEnvironment& Environment,
 		FMassEntityQuery& NativeQuery,
-		ITypedElementDataStorageInterface::QueryCallbackRef Callback)
+		IEditorDataStorageProvider::QueryCallbackRef Callback)
 	{
 		if (Description.Callback.ActivationCount > 0)
 		{
@@ -917,17 +917,17 @@ bool FTypedElementQueryProcessorData::CommonQueryConfiguration(
 	return false;
 }
 
-EMassProcessingPhase FTypedElementQueryProcessorData::MapToMassProcessingPhase(ITypedElementDataStorageInterface::EQueryTickPhase Phase)
+EMassProcessingPhase FTypedElementQueryProcessorData::MapToMassProcessingPhase(IEditorDataStorageProvider::EQueryTickPhase Phase)
 {
 	switch(Phase)
 	{
-	case ITypedElementDataStorageInterface::EQueryTickPhase::PrePhysics:
+	case IEditorDataStorageProvider::EQueryTickPhase::PrePhysics:
 		return EMassProcessingPhase::PrePhysics;
-	case ITypedElementDataStorageInterface::EQueryTickPhase::DuringPhysics:
+	case IEditorDataStorageProvider::EQueryTickPhase::DuringPhysics:
 		return EMassProcessingPhase::DuringPhysics;
-	case ITypedElementDataStorageInterface::EQueryTickPhase::PostPhysics:
+	case IEditorDataStorageProvider::EQueryTickPhase::PostPhysics:
 		return EMassProcessingPhase::PostPhysics;
-	case ITypedElementDataStorageInterface::EQueryTickPhase::FrameEnd:
+	case IEditorDataStorageProvider::EQueryTickPhase::FrameEnd:
 		return EMassProcessingPhase::FrameEnd;
 	default:
 		checkf(false, TEXT("Query tick phase '%i' is unsupported."), static_cast<int>(Phase));
@@ -954,8 +954,8 @@ void FTypedElementQueryProcessorData::DebugOutputDescription(FOutputDevice& Ar,
 
 	if (const FExtendedQuery* StoredQuery = QueryStore ? QueryStore->Get(ParentQuery) : nullptr)
 	{
-		const ITypedElementDataStorageInterface::FQueryDescription& Description = StoredQuery->Description;
-		const ITypedElementDataStorageInterface::FQueryDescription::FCallbackData& Callback = Description.Callback;
+		const IEditorDataStorageProvider::FQueryDescription& Description = StoredQuery->Description;
+		const IEditorDataStorageProvider::FQueryDescription::FCallbackData& Callback = Description.Callback;
 		
 		if (!Callback.Group.IsNone())
 		{
@@ -1013,20 +1013,20 @@ void FTypedElementQueryProcessorData::DebugOutputDescription(FOutputDevice& Ar,
 }
 
 bool FTypedElementQueryProcessorData::PrepareCachedDependenciesOnQuery(
-	ITypedElementDataStorageInterface::FQueryDescription& Description, FMassExecutionContext& Context)
+	IEditorDataStorageProvider::FQueryDescription& Description, FMassExecutionContext& Context)
 {
 	const int32 DependencyCount = Description.DependencyTypes.Num();
 	const TWeakObjectPtr<const UClass>* Types = Description.DependencyTypes.GetData();
-	const ITypedElementDataStorageInterface::EQueryDependencyFlags* Flags = Description.DependencyFlags.GetData();
+	const IEditorDataStorageProvider::EQueryDependencyFlags* Flags = Description.DependencyFlags.GetData();
 	TWeakObjectPtr<UObject>* Caches = Description.CachedDependencies.GetData();
 
 	for (int32 Index = 0; Index < DependencyCount; ++Index)
 	{
 		checkf(Types->IsValid(), TEXT("Attempting to retrieve a dependency type that's no longer available."));
 		
-		if (EnumHasAnyFlags(*Flags, ITypedElementDataStorageInterface::EQueryDependencyFlags::AlwaysRefresh) || !Caches->IsValid())
+		if (EnumHasAnyFlags(*Flags, IEditorDataStorageProvider::EQueryDependencyFlags::AlwaysRefresh) || !Caches->IsValid())
 		{
-			*Caches = EnumHasAnyFlags(*Flags, ITypedElementDataStorageInterface::EQueryDependencyFlags::ReadOnly)
+			*Caches = EnumHasAnyFlags(*Flags, IEditorDataStorageProvider::EQueryDependencyFlags::ReadOnly)
 				? const_cast<USubsystem*>(Context.GetSubsystem<USubsystem>(const_cast<UClass*>(Types->Get())))
 				: Context.GetMutableSubsystem<USubsystem>(const_cast<UClass*>(Types->Get()));
 			if (*Caches != nullptr)
@@ -1097,8 +1097,8 @@ UE::Editor::DataStorage::FQueryResult FTypedElementQueryProcessorData::Execute(
 {
 	using namespace UE::Editor::DataStorage;
 
-	ITypedElementDataStorageInterface::FQueryResult Result;
-	Result.Completed = ITypedElementDataStorageInterface::FQueryResult::ECompletion::Fully;
+	IEditorDataStorageProvider::FQueryResult Result;
+	Result.Completed = IEditorDataStorageProvider::FQueryResult::ECompletion::Fully;
 
 	if (Description.Callback.ActivationCount > 0)
 	{
@@ -1133,8 +1133,8 @@ UE::Editor::DataStorage::FQueryResult FTypedElementQueryProcessorData::Execute(
 {
 	using namespace UE::Editor::DataStorage;
 
-	ITypedElementDataStorageInterface::FQueryResult Result;
-	Result.Completed = ITypedElementDataStorageInterface::FQueryResult::ECompletion::Fully;
+	IEditorDataStorageProvider::FQueryResult Result;
+	Result.Completed = IEditorDataStorageProvider::FQueryResult::ECompletion::Fully;
 
 	FMassEntityHandle NativeEntity = FMassEntityHandle::FromNumber(RowHandle);
 	if (Description.Callback.ActivationCount > 0 && EntityManager.IsEntityActive(NativeEntity))
@@ -1170,7 +1170,7 @@ void FTypedElementQueryProcessorData::Execute(FMassEntityManager& EntityManager,
 
 	checkf(StoredQuery, TEXT("A query callback was registered for execution without an associated query."));
 	
-	ITypedElementDataStorageInterface::FQueryDescription& Description = StoredQuery->Description;
+	IEditorDataStorageProvider::FQueryDescription& Description = StoredQuery->Description;
 	if (Description.Callback.ActivationCount > 0)
 	{
 		auto ExceteFunction = [this, &Description](FMassExecutionContext& Context)
@@ -1378,10 +1378,10 @@ bool UTypedElementQueryObserverCallbackAdapterProcessorBase::ConfigureQueryCallb
 	
 	switch (Query.Description.Callback.Type)
 	{
-	case ITypedElementDataStorageInterface::EQueryCallbackType::ObserveAdd:
+	case IEditorDataStorageProvider::EQueryCallbackType::ObserveAdd:
 		Operation = EMassObservedOperation::Add;
 		break;
-	case ITypedElementDataStorageInterface::EQueryCallbackType::ObserveRemove:
+	case IEditorDataStorageProvider::EQueryCallbackType::ObserveRemove:
 		Operation = EMassObservedOperation::Remove;
 		break;
 	default:
diff --git a/Engine/Plugins/Experimental/EditorDataStorage/Source/TedsCore/Private/Processors/TypedElementProcessorAdaptors.h b/Engine/Plugins/Experimental/EditorDataStorage/Source/TedsCore/Private/Processors/TypedElementProcessorAdaptors.h
--- a/Engine/Plugins/Experimental/EditorDataStorage/Source/TedsCore/Private/Processors/TypedElementProcessorAdaptors.h
+++ b/Engine/Plugins/Experimental/EditorDataStorage/Source/TedsCore/Private/Processors/TypedElementProcessorAdaptors.h
@@ -24,11 +24,11 @@ namespace UE::Editor::DataStorage
 		~FPhasePreOrPostAmbleExecutor();
 
 		void ExecuteQuery(
-			ITypedElementDataStorageInterface::FQueryDescription& Description,
+			IEditorDataStorageProvider::FQueryDescription& Description,
 			FExtendedQueryStore& QueryStore,
 			FEnvironment& Environment,
 			FMassEntityQuery& NativeQuery,
-			ITypedElementDataStorageInterface::QueryCallbackRef Callback);
+			IEditorDataStorageProvider::QueryCallbackRef Callback);
 
 		FMassExecutionContext Context;
 	};
@@ -53,7 +53,7 @@ struct FTypedElementQueryProcessorData
 		FExtendedQueryStore& InQueryStore,
 		FEnvironment& InEnvironment,
 		TArrayView<FMassEntityQuery> Subqueries);
-	static EMassProcessingPhase MapToMassProcessingPhase(ITypedElementDataStorageInterface::EQueryTickPhase Phase);
+	static EMassProcessingPhase MapToMassProcessingPhase(IEditorDataStorageProvider::EQueryTickPhase Phase);
 	FString GetProcessorName() const;
 	void DebugOutputDescription(FOutputDevice& Ar, int32 Indent) const;
 
@@ -82,7 +82,7 @@ struct FTypedElementQueryProcessorData
 	void Execute(FMassEntityManager& EntityManager, FMassExecutionContext& Context);
 
 	static bool PrepareCachedDependenciesOnQuery(
-		ITypedElementDataStorageInterface::FQueryDescription& Description, FMassExecutionContext& Context);
+		IEditorDataStorageProvider::FQueryDescription& Description, FMassExecutionContext& Context);
 
 	FExtendedQueryStore::Handle ParentQuery;
 	FExtendedQueryStore* QueryStore{ nullptr };
diff --git a/Engine/Plugins/Experimental/EditorDataStorage/Source/TedsCore/Private/Queries/TypedElementAlertQueries.cpp b/Engine/Plugins/Experimental/EditorDataStorage/Source/TedsCore/Private/Queries/TypedElementAlertQueries.cpp
--- a/Engine/Plugins/Experimental/EditorDataStorage/Source/TedsCore/Private/Queries/TypedElementAlertQueries.cpp
+++ b/Engine/Plugins/Experimental/EditorDataStorage/Source/TedsCore/Private/Queries/TypedElementAlertQueries.cpp
@@ -21,7 +21,7 @@ FAutoConsoleCommand AddRandomAlertToRowConsoleCommand(
 
 			TRACE_CPUPROFILER_EVENT_SCOPE(TEDS.Debug.AddRandomAlertToSelectedRows);
 
-			if (ITypedElementDataStorageInterface* DataStorage = UTypedElementRegistry::GetInstance()->GetMutableDataStorage())
+			if (IEditorDataStorageProvider* DataStorage = UTypedElementRegistry::GetInstance()->GetMutableDataStorage())
 			{
 				static QueryHandle Query = [DataStorage]
 				{
@@ -82,7 +82,7 @@ FAutoConsoleCommand ClearAllAlertsConsoleCommand(
 			static TTypedElementColumnTypeList<FTypedElementSyncBackToWorldTag> BatchAddColumns;
 			static TTypedElementColumnTypeList<FTypedElementAlertColumn, FTypedElementChildAlertColumn> BatchRemoveColumns;
 
-			if (ITypedElementDataStorageInterface* DataStorage = UTypedElementRegistry::GetInstance()->GetMutableDataStorage())
+			if (IEditorDataStorageProvider* DataStorage = UTypedElementRegistry::GetInstance()->GetMutableDataStorage())
 			{
 				static QueryHandle AlertInfoQuery = [DataStorage]
 				{
@@ -118,7 +118,7 @@ FAutoConsoleCommand ClearSelectedAlertsConsoleCommand(
 
 			TRACE_CPUPROFILER_EVENT_SCOPE(TEDS.Debug.ClearSelectedAlerts);
 
-			if (ITypedElementDataStorageInterface* DataStorage = UTypedElementRegistry::GetInstance()->GetMutableDataStorage())
+			if (IEditorDataStorageProvider* DataStorage = UTypedElementRegistry::GetInstance()->GetMutableDataStorage())
 			{
 				static QueryHandle AlertQuery = [DataStorage]
 				{
@@ -145,7 +145,7 @@ FAutoConsoleCommand ClearSelectedAlertsConsoleCommand(
 
 const FName UTypedElementAlertQueriesFactory::AlertConditionName = FName(TEXT("Alerts"));
 
-void UTypedElementAlertQueriesFactory::RegisterQueries(ITypedElementDataStorageInterface& DataStorage)
+void UTypedElementAlertQueriesFactory::RegisterQueries(IEditorDataStorageProvider& DataStorage)
 {
 	RegisterSubQueries(DataStorage);
 	RegisterParentUpdatesQueries(DataStorage);
@@ -154,7 +154,7 @@ void UTypedElementAlertQueriesFactory::RegisterQueries(ITypedElementDataStorageI
 	RegisterOnRemoveQueries(DataStorage);
 }
 
-void UTypedElementAlertQueriesFactory::RegisterSubQueries(ITypedElementDataStorageInterface& DataStorage)
+void UTypedElementAlertQueriesFactory::RegisterSubQueries(IEditorDataStorageProvider& DataStorage)
 {
 	using namespace UE::Editor::DataStorage::Queries;
 	
@@ -169,7 +169,7 @@ void UTypedElementAlertQueriesFactory::RegisterSubQueries(ITypedElementDataStora
 		.Compile());
 }
 
-void UTypedElementAlertQueriesFactory::RegisterParentUpdatesQueries(ITypedElementDataStorageInterface& DataStorage)
+void UTypedElementAlertQueriesFactory::RegisterParentUpdatesQueries(IEditorDataStorageProvider& DataStorage)
 {
 	using namespace UE::Editor::DataStorage::Queries;
 
@@ -206,7 +206,7 @@ void UTypedElementAlertQueriesFactory::RegisterParentUpdatesQueries(ITypedElemen
 		.Compile());
 }
 
-void UTypedElementAlertQueriesFactory::RegisterChildAlertUpdatesQueries(ITypedElementDataStorageInterface& DataStorage)
+void UTypedElementAlertQueriesFactory::RegisterChildAlertUpdatesQueries(IEditorDataStorageProvider& DataStorage)
 {
 	using namespace UE::Editor::DataStorage::Queries;
 
@@ -266,7 +266,7 @@ void UTypedElementAlertQueriesFactory::RegisterChildAlertUpdatesQueries(ITypedEl
 		.Compile());
 }
 
-void UTypedElementAlertQueriesFactory::RegisterOnAddQueries(ITypedElementDataStorageInterface& DataStorage)
+void UTypedElementAlertQueriesFactory::RegisterOnAddQueries(IEditorDataStorageProvider& DataStorage)
 {
 	using namespace UE::Editor::DataStorage::Queries;
 	
@@ -295,7 +295,7 @@ void UTypedElementAlertQueriesFactory::RegisterOnAddQueries(ITypedElementDataSto
 		.Compile());
 }
 
-void UTypedElementAlertQueriesFactory::RegisterOnRemoveQueries(ITypedElementDataStorageInterface& DataStorage)
+void UTypedElementAlertQueriesFactory::RegisterOnRemoveQueries(IEditorDataStorageProvider& DataStorage)
 {
 	using namespace UE::Editor::DataStorage::Queries;
 	using namespace UE::Editor::DataStorage;
diff --git a/Engine/Plugins/Experimental/EditorDataStorage/Source/TedsCore/Private/Queries/TypedElementAlertQueries.h b/Engine/Plugins/Experimental/EditorDataStorage/Source/TedsCore/Private/Queries/TypedElementAlertQueries.h
--- a/Engine/Plugins/Experimental/EditorDataStorage/Source/TedsCore/Private/Queries/TypedElementAlertQueries.h
+++ b/Engine/Plugins/Experimental/EditorDataStorage/Source/TedsCore/Private/Queries/TypedElementAlertQueries.h
@@ -20,21 +20,21 @@ namespace UE::Editor::DataStorage
  * Calls to manage alerts, in particular child alerts.
  */
 UCLASS()
-class UTypedElementAlertQueriesFactory final : public UTypedElementDataStorageFactory
+class UTypedElementAlertQueriesFactory final : public UEditorDataStorageFactory
 {
 	GENERATED_BODY()
 
 public:
 	~UTypedElementAlertQueriesFactory() override = default;
 
-	void RegisterQueries(ITypedElementDataStorageInterface& DataStorage) override;
+	void RegisterQueries(IEditorDataStorageProvider& DataStorage) override;
 
 private:
-	void RegisterSubQueries(ITypedElementDataStorageInterface& DataStorage);
-	void RegisterParentUpdatesQueries(ITypedElementDataStorageInterface& DataStorage);
-	void RegisterChildAlertUpdatesQueries(ITypedElementDataStorageInterface& DataStorage);
-	void RegisterOnAddQueries(ITypedElementDataStorageInterface& DataStorage);
-	void RegisterOnRemoveQueries(ITypedElementDataStorageInterface& DataStorage);
+	void RegisterSubQueries(IEditorDataStorageProvider& DataStorage);
+	void RegisterParentUpdatesQueries(IEditorDataStorageProvider& DataStorage);
+	void RegisterChildAlertUpdatesQueries(IEditorDataStorageProvider& DataStorage);
+	void RegisterOnAddQueries(IEditorDataStorageProvider& DataStorage);
+	void RegisterOnRemoveQueries(IEditorDataStorageProvider& DataStorage);
 
 	static void AddChildAlertsToHierarchy(
 		UE::Editor::DataStorage::IQueryContext& Context, UE::Editor::DataStorage::RowHandle Parent, int32 ParentQueryIndex);
diff --git a/Engine/Plugins/Experimental/EditorDataStorage/Source/TedsCore/Private/Queries/TypedElementExtendedQueryStore.h b/Engine/Plugins/Experimental/EditorDataStorage/Source/TedsCore/Private/Queries/TypedElementExtendedQueryStore.h
--- a/Engine/Plugins/Experimental/EditorDataStorage/Source/TedsCore/Private/Queries/TypedElementExtendedQueryStore.h
+++ b/Engine/Plugins/Experimental/EditorDataStorage/Source/TedsCore/Private/Queries/TypedElementExtendedQueryStore.h
@@ -20,7 +20,7 @@ namespace UE::Editor::DataStorage
 	struct FExtendedQuery
 	{
 		FMassEntityQuery NativeQuery; // Used if there's no processor bound.
-		ITypedElementDataStorageInterface::FQueryDescription Description;
+		IEditorDataStorageProvider::FQueryDescription Description;
 		TStrongObjectPtr<UMassProcessor> Processor;
 	};
 
@@ -42,7 +42,7 @@ namespace UE::Editor::DataStorage
 
 		 /** Adds a new query to the store and initializes the query with the provided arguments. */
 		Handle RegisterQuery(
-			ITypedElementDataStorageInterface::FQueryDescription Query,
+			IEditorDataStorageProvider::FQueryDescription Query,
 			FEnvironment& Environment,
 			FMassEntityManager& EntityManager,
 			FMassProcessingPhaseManager& PhaseManager);
@@ -53,10 +53,10 @@ namespace UE::Editor::DataStorage
 		void Clear(FMassEntityManager& EntityManager, FMassProcessingPhaseManager& PhaseManager);
 
 		/** Register the defaults for a tick group. These will be applied on top of any settings provided with a query registration. */
-		void RegisterTickGroup(FName GroupName, ITypedElementDataStorageInterface::EQueryTickPhase Phase,
+		void RegisterTickGroup(FName GroupName, IEditorDataStorageProvider::EQueryTickPhase Phase,
 			FName BeforeGroup, FName AfterGroup, EExecutionMode ExecutionMode);
 		/** Removes a previously registered set of tick group defaults. */
-		void UnregisterTickGroup(FName GroupName, ITypedElementDataStorageInterface::EQueryTickPhase Phase);
+		void UnregisterTickGroup(FName GroupName, IEditorDataStorageProvider::EQueryTickPhase Phase);
 
 		/**
 		 * @section Retrieval
@@ -78,7 +78,7 @@ namespace UE::Editor::DataStorage
 		const FExtendedQuery& GetChecked(Handle Entry) const;
 
 		/** Gets the original description used to create an extended query or an empty default if the provided query isn't alive. */
-		const ITypedElementDataStorageInterface::FQueryDescription& GetQueryDescription(Handle Query) const;
+		const IEditorDataStorageProvider::FQueryDescription& GetQueryDescription(Handle Query) const;
 
 		/** Checks to see if a query is still available or has been removed. */
 		bool IsAlive(Handle Entry) const;
@@ -124,24 +124,24 @@ namespace UE::Editor::DataStorage
 		void RunPhasePreambleQueries(
 			FMassEntityManager& EntityManager,
 			FEnvironment& Environment,
-			ITypedElementDataStorageInterface::EQueryTickPhase Phase,
+			IEditorDataStorageProvider::EQueryTickPhase Phase,
 			float DeltaTime);
 		void RunPhasePostambleQueries(
 			FMassEntityManager& EntityManager,
 			FEnvironment& Environment,
-			ITypedElementDataStorageInterface::EQueryTickPhase Phase,
+			IEditorDataStorageProvider::EQueryTickPhase Phase,
 			float DeltaTime);
 
 		void DebugPrintQueryCallbacks(FOutputDevice& Output) const;
 
 	private:
-		using QueryTickPhaseType = std::underlying_type_t<ITypedElementDataStorageInterface::EQueryTickPhase>;
-		static constexpr QueryTickPhaseType MaxTickPhase = static_cast<QueryTickPhaseType>(ITypedElementDataStorageInterface::EQueryTickPhase::Max);
+		using QueryTickPhaseType = std::underlying_type_t<IEditorDataStorageProvider::EQueryTickPhase>;
+		static constexpr QueryTickPhaseType MaxTickPhase = static_cast<QueryTickPhaseType>(IEditorDataStorageProvider::EQueryTickPhase::Max);
 
 		struct FTickGroupId
 		{
 			FName Name;
-			ITypedElementDataStorageInterface::EQueryTickPhase Phase;
+			IEditorDataStorageProvider::EQueryTickPhase Phase;
 
 			friend inline uint32 GetTypeHash(const FTickGroupId& Id) { return HashCombine(GetTypeHash(Id.Name), GetTypeHash(Id.Phase)); }
 			friend inline bool operator==(const FTickGroupId& Lhs, const FTickGroupId& Rhs) { return Lhs.Phase == Rhs.Phase && Lhs.Name == Rhs.Name; }
@@ -164,30 +164,30 @@ namespace UE::Editor::DataStorage
 		EDirectQueryExecutionFlags ExecutionFlags,
 			CallbackReference Callback);
 
-		FMassEntityQuery& SetupNativeQuery(ITypedElementDataStorageInterface::FQueryDescription& Query, FExtendedQuery& StoredQuery);
-		bool SetupDynamicColumns(ITypedElementDataStorageInterface::FQueryDescription& Query, FEnvironment& Environment);
-		bool SetupSelectedColumns(ITypedElementDataStorageInterface::FQueryDescription& Query, FMassEntityQuery& NativeQuery);
-		bool SetupConditions(ITypedElementDataStorageInterface::FQueryDescription& Query, FMassEntityQuery& NativeQuery);
-		bool SetupChunkFilters(Handle QueryHandle, ITypedElementDataStorageInterface::FQueryDescription& Query, FEnvironment& Environment, FMassEntityQuery& NativeQuery);
-		bool SetupDependencies(ITypedElementDataStorageInterface::FQueryDescription& Query, FMassEntityQuery& NativeQuery);
-		bool SetupTickGroupDefaults(ITypedElementDataStorageInterface::FQueryDescription& Query);
+		FMassEntityQuery& SetupNativeQuery(IEditorDataStorageProvider::FQueryDescription& Query, FExtendedQuery& StoredQuery);
+		bool SetupDynamicColumns(IEditorDataStorageProvider::FQueryDescription& Query, FEnvironment& Environment);
+		bool SetupSelectedColumns(IEditorDataStorageProvider::FQueryDescription& Query, FMassEntityQuery& NativeQuery);
+		bool SetupConditions(IEditorDataStorageProvider::FQueryDescription& Query, FMassEntityQuery& NativeQuery);
+		bool SetupChunkFilters(Handle QueryHandle, IEditorDataStorageProvider::FQueryDescription& Query, FEnvironment& Environment, FMassEntityQuery& NativeQuery);
+		bool SetupDependencies(IEditorDataStorageProvider::FQueryDescription& Query, FMassEntityQuery& NativeQuery);
+		bool SetupTickGroupDefaults(IEditorDataStorageProvider::FQueryDescription& Query);
 		bool SetupProcessors(Handle QueryHandle, FExtendedQuery& StoredQuery, FEnvironment& Environment,
 			FMassEntityManager& EntityManager, FMassProcessingPhaseManager& PhaseManager);
-		bool SetupActivatable(Handle QueryHandle, ITypedElementDataStorageInterface::FQueryDescription& Query);
+		bool SetupActivatable(Handle QueryHandle, IEditorDataStorageProvider::FQueryDescription& Query);
 
-		EMassFragmentAccess ConvertToNativeAccessType(ITypedElementDataStorageInterface::EQueryAccessType AccessType);
-		EMassFragmentPresence ConvertToNativePresenceType(ITypedElementDataStorageInterface::EQueryAccessType AccessType);
+		EMassFragmentAccess ConvertToNativeAccessType(IEditorDataStorageProvider::EQueryAccessType AccessType);
+		EMassFragmentPresence ConvertToNativePresenceType(IEditorDataStorageProvider::EQueryAccessType AccessType);
 
-		void RegisterPreambleQuery(ITypedElementDataStorageInterface::EQueryTickPhase Phase, Handle Query);
-		void RegisterPostambleQuery(ITypedElementDataStorageInterface::EQueryTickPhase Phase, Handle Query);
-		void UnregisterPreambleQuery(ITypedElementDataStorageInterface::EQueryTickPhase Phase, Handle Query);
-		void UnregisterPostambleQuery(ITypedElementDataStorageInterface::EQueryTickPhase Phase, Handle Query);
+		void RegisterPreambleQuery(IEditorDataStorageProvider::EQueryTickPhase Phase, Handle Query);
+		void RegisterPostambleQuery(IEditorDataStorageProvider::EQueryTickPhase Phase, Handle Query);
+		void UnregisterPreambleQuery(IEditorDataStorageProvider::EQueryTickPhase Phase, Handle Query);
+		void UnregisterPostambleQuery(IEditorDataStorageProvider::EQueryTickPhase Phase, Handle Query);
 		void RunPhasePreOrPostAmbleQueries(FMassEntityManager& EntityManager, FEnvironment& Environment,
-			ITypedElementDataStorageInterface::EQueryTickPhase Phase, float DeltaTime, TArray<Handle>& QueryHandles);
+			IEditorDataStorageProvider::EQueryTickPhase Phase, float DeltaTime, TArray<Handle>& QueryHandles);
 
 		void UnregisterQueryData(Handle Query, FExtendedQuery& QueryData, FMassEntityManager& EntityManager, FMassProcessingPhaseManager& PhaseManager);
 
-		static const ITypedElementDataStorageInterface::FQueryDescription EmptyDescription;
+		static const IEditorDataStorageProvider::FQueryDescription EmptyDescription;
 
 		QueryStore Queries;
 		TMultiMap<FName, Handle> ActivatableMapping;
diff --git a/Engine/Plugins/Experimental/EditorDataStorage/Source/TedsCore/Private/Queries/TypedElementHierarchyQueries.cpp b/Engine/Plugins/Experimental/EditorDataStorage/Source/TedsCore/Private/Queries/TypedElementHierarchyQueries.cpp
--- a/Engine/Plugins/Experimental/EditorDataStorage/Source/TedsCore/Private/Queries/TypedElementHierarchyQueries.cpp
+++ b/Engine/Plugins/Experimental/EditorDataStorage/Source/TedsCore/Private/Queries/TypedElementHierarchyQueries.cpp
@@ -5,7 +5,7 @@
 #include "Elements/Columns/TypedElementHiearchyColumns.h"
 #include "Elements/Framework/TypedElementQueryBuilder.h"
 
-void UTypedElementHiearchyQueriesFactory::RegisterQueries(ITypedElementDataStorageInterface& DataStorage)
+void UTypedElementHiearchyQueriesFactory::RegisterQueries(IEditorDataStorageProvider& DataStorage)
 {
 	using namespace UE::Editor::DataStorage::Queries;
 
diff --git a/Engine/Plugins/Experimental/EditorDataStorage/Source/TedsCore/Private/Queries/TypedElementHierarchyQueries.h b/Engine/Plugins/Experimental/EditorDataStorage/Source/TedsCore/Private/Queries/TypedElementHierarchyQueries.h
--- a/Engine/Plugins/Experimental/EditorDataStorage/Source/TedsCore/Private/Queries/TypedElementHierarchyQueries.h
+++ b/Engine/Plugins/Experimental/EditorDataStorage/Source/TedsCore/Private/Queries/TypedElementHierarchyQueries.h
@@ -11,12 +11,12 @@
  * Calls to queries for general hierarchy management.
  */
 UCLASS()
-class UTypedElementHiearchyQueriesFactory : public UTypedElementDataStorageFactory
+class UTypedElementHiearchyQueriesFactory : public UEditorDataStorageFactory
 {
 	GENERATED_BODY()
 
 public:
 	~UTypedElementHiearchyQueriesFactory() override = default;
 
-	void RegisterQueries(ITypedElementDataStorageInterface& DataStorage) override;
+	void RegisterQueries(IEditorDataStorageProvider& DataStorage) override;
 };
diff --git a/Engine/Plugins/Experimental/EditorDataStorage/Source/TedsCore/Private/Queries/TypedElementMiscQueries.cpp b/Engine/Plugins/Experimental/EditorDataStorage/Source/TedsCore/Private/Queries/TypedElementMiscQueries.cpp
--- a/Engine/Plugins/Experimental/EditorDataStorage/Source/TedsCore/Private/Queries/TypedElementMiscQueries.cpp
+++ b/Engine/Plugins/Experimental/EditorDataStorage/Source/TedsCore/Private/Queries/TypedElementMiscQueries.cpp
@@ -5,7 +5,7 @@
 #include "Elements/Columns/TypedElementMiscColumns.h"
 #include "Elements/Framework/TypedElementQueryBuilder.h"
 
-void UTypedElementRemoveSyncToWorldTagFactory::RegisterQueries(ITypedElementDataStorageInterface& DataStorage)
+void UTypedElementRemoveSyncToWorldTagFactory::RegisterQueries(IEditorDataStorageProvider& DataStorage)
 {
 	using namespace UE::Editor::DataStorage::Queries;
 
diff --git a/Engine/Plugins/Experimental/EditorDataStorage/Source/TedsCore/Private/Queries/TypedElementMiscQueries.h b/Engine/Plugins/Experimental/EditorDataStorage/Source/TedsCore/Private/Queries/TypedElementMiscQueries.h
--- a/Engine/Plugins/Experimental/EditorDataStorage/Source/TedsCore/Private/Queries/TypedElementMiscQueries.h
+++ b/Engine/Plugins/Experimental/EditorDataStorage/Source/TedsCore/Private/Queries/TypedElementMiscQueries.h
@@ -11,12 +11,12 @@
  * Removes all FTypedElementSyncBackToWorldTags at the end of an update cycle.
  */
 UCLASS()
-class UTypedElementRemoveSyncToWorldTagFactory : public UTypedElementDataStorageFactory
+class UTypedElementRemoveSyncToWorldTagFactory : public UEditorDataStorageFactory
 {
 	GENERATED_BODY()
 
 public:
 	~UTypedElementRemoveSyncToWorldTagFactory() override = default;
 
-	void RegisterQueries(ITypedElementDataStorageInterface& DataStorage) override;
+	void RegisterQueries(IEditorDataStorageProvider& DataStorage) override;
 };
diff --git a/Engine/Plugins/Experimental/EditorDataStorage/Source/TedsCore/Private/TypedElementDatabase.cpp b/Engine/Plugins/Experimental/EditorDataStorage/Source/TedsCore/Private/TypedElementDatabase.cpp
--- a/Engine/Plugins/Experimental/EditorDataStorage/Source/TedsCore/Private/TypedElementDatabase.cpp
+++ b/Engine/Plugins/Experimental/EditorDataStorage/Source/TedsCore/Private/TypedElementDatabase.cpp
@@ -167,7 +167,7 @@ void UEditorDataStorage::SetFactories(TConstArrayView<UClass*> FactoryClasses)
 {
 	Factories.Reserve(FactoryClasses.Num());
 
-	UClass* BaseFactoryType = UTypedElementDataStorageFactory::StaticClass();
+	UClass* BaseFactoryType = UEditorDataStorageFactory::StaticClass();
 
 	for (UClass* FactoryClass : FactoryClasses)
 	{
@@ -179,7 +179,7 @@ void UEditorDataStorage::SetFactories(TConstArrayView<UClass*> FactoryClasses)
 		{
 			continue;
 		}
-		UTypedElementDataStorageFactory* Factory = NewObject<UTypedElementDataStorageFactory>(this, FactoryClass, NAME_None, EObjectFlags::RF_Transient);
+		UEditorDataStorageFactory* Factory = NewObject<UEditorDataStorageFactory>(this, FactoryClass, NAME_None, EObjectFlags::RF_Transient);
 		Factories.Add(FFactoryTypePair
 			{
 				.Type = FactoryClass,
@@ -219,7 +219,7 @@ UEditorDataStorage::FactoryConstIterator UEditorDataStorage::CreateFactoryIterat
 	return UEditorDataStorage::FactoryConstIterator(this);
 }
 
-const UTypedElementDataStorageFactory* UEditorDataStorage::FindFactory(const UClass* FactoryType) const
+const UEditorDataStorageFactory* UEditorDataStorage::FindFactory(const UClass* FactoryType) const
 {
 	for (const FFactoryTypePair& Factory : Factories)
 	{
@@ -990,7 +990,7 @@ void UEditorDataStorage::UnregisterQuery(QueryHandle Query)
 	}
 }
 
-const ITypedElementDataStorageInterface::FQueryDescription& UEditorDataStorage::GetQueryDescription(QueryHandle Query) const
+const IEditorDataStorageProvider::FQueryDescription& UEditorDataStorage::GetQueryDescription(QueryHandle Query) const
 {
 	const UE::Editor::DataStorage::FExtendedQueryStore::Handle StorageHandle(Query);
 	return Environment->GetQueryStore().GetQueryDescription(StorageHandle);
@@ -1020,7 +1020,7 @@ FName UEditorDataStorage::GetQueryTickGroupName(EQueryTickGroups Group) const
 	}
 }
 
-ITypedElementDataStorageInterface::FQueryResult UEditorDataStorage::RunQuery(QueryHandle Query)
+IEditorDataStorageProvider::FQueryResult UEditorDataStorage::RunQuery(QueryHandle Query)
 {
 	TRACE_CPUPROFILER_EVENT_SCOPE(TEDS.RunQuery);
 
@@ -1035,8 +1035,8 @@ ITypedElementDataStorageInterface::FQueryResult UEditorDataStorage::RunQuery(Que
 	}
 }
 
-ITypedElementDataStorageInterface::FQueryResult UEditorDataStorage::RunQuery(
-	QueryHandle Query, ITypedElementDataStorageInterface::DirectQueryCallbackRef Callback)
+IEditorDataStorageProvider::FQueryResult UEditorDataStorage::RunQuery(
+	QueryHandle Query, IEditorDataStorageProvider::DirectQueryCallbackRef Callback)
 {
 	TRACE_CPUPROFILER_EVENT_SCOPE(TEDS.RunQuery);
 
@@ -1052,9 +1052,9 @@ ITypedElementDataStorageInterface::FQueryResult UEditorDataStorage::RunQuery(
 	}
 }
 
-ITypedElementDataStorageInterface::FQueryResult UEditorDataStorage::RunQuery(
+IEditorDataStorageProvider::FQueryResult UEditorDataStorage::RunQuery(
 	QueryHandle Query, UE::Editor::DataStorage::EDirectQueryExecutionFlags Flags,
-	ITypedElementDataStorageInterface::DirectQueryCallbackRef Callback)
+	IEditorDataStorageProvider::DirectQueryCallbackRef Callback)
 {
 	TRACE_CPUPROFILER_EVENT_SCOPE(TEDS.RunQuery);
 
