[Teds] Rename interface types and base factory class

diff --git a/Engine/Plugins/Experimental/EditorDataStorage/Source/TedsCore/Private/TypedElementDatabase.h b/Engine/Plugins/Experimental/EditorDataStorage/Source/TedsCore/Private/TypedElementDatabase.h
--- a/Engine/Plugins/Experimental/EditorDataStorage/Source/TedsCore/Private/TypedElementDatabase.h
+++ b/Engine/Plugins/Experimental/EditorDataStorage/Source/TedsCore/Private/TypedElementDatabase.h
@@ -15,7 +15,7 @@
 
 struct FMassEntityManager;
 struct FMassProcessingPhaseManager;
-class UTypedElementDataStorageFactory;
+class UEditorDataStorageFactory;
 class FOutputDevice;
 class UWorld;
 
@@ -27,7 +27,7 @@ namespace UE::Editor::DataStorage
 UCLASS()
 class TEDSCORE_API UEditorDataStorage
 	: public UObject
-	, public ITypedElementDataStorageInterface
+	, public IEditorDataStorageProvider
 {
 	GENERATED_BODY()
 
@@ -61,8 +61,8 @@ public:
 		int32 Index = 0;
 	};
 
-	using FactoryIterator = TFactoryIterator<UTypedElementDataStorageFactory, UEditorDataStorage>;
-	using FactoryConstIterator = TFactoryIterator<const UTypedElementDataStorageFactory, const UEditorDataStorage>;
+	using FactoryIterator = TFactoryIterator<UEditorDataStorageFactory, UEditorDataStorage>;
+	using FactoryConstIterator = TFactoryIterator<const UEditorDataStorageFactory, const UEditorDataStorage>;
 
 public:
 	~UEditorDataStorage() override = default;
@@ -78,7 +78,7 @@ public:
 	FactoryConstIterator CreateFactoryIterator() const;
 
 	/** Returns factory instance given the type of factory */
-	virtual const UTypedElementDataStorageFactory* FindFactory(const UClass* FactoryType) const override;
+	virtual const UEditorDataStorageFactory* FindFactory(const UClass* FactoryType) const override;
 	/** Helper for FindFactory(const UClass*) */
 	template<typename FactoryTypeT>
 	const FactoryTypeT* FindFactory() const;
@@ -191,7 +191,7 @@ private:
 		// Used to find the factory by type without needing to dereference each one
 		TObjectPtr<UClass> Type;
 		
-		TObjectPtr<UTypedElementDataStorageFactory> Instance;
+		TObjectPtr<UEditorDataStorageFactory> Instance;
 	};
 	
 	static const FName TickGroupName_Default;
diff --git a/Engine/Plugins/Experimental/EditorDataStorage/Source/TedsCore/Private/TypedElementDatabaseCompatibility.cpp b/Engine/Plugins/Experimental/EditorDataStorage/Source/TedsCore/Private/TypedElementDatabaseCompatibility.cpp
--- a/Engine/Plugins/Experimental/EditorDataStorage/Source/TedsCore/Private/TypedElementDatabaseCompatibility.cpp
+++ b/Engine/Plugins/Experimental/EditorDataStorage/Source/TedsCore/Private/TypedElementDatabaseCompatibility.cpp
@@ -728,7 +728,7 @@ void UEditorDataStorageCompatibility::PendingRegistration<AddressType>::ForEachA
 }
 
 template<typename AddressType>
-void UEditorDataStorageCompatibility::PendingRegistration<AddressType>::ProcessEntries(ITypedElementDataStorageInterface& StorageInterface,
+void UEditorDataStorageCompatibility::PendingRegistration<AddressType>::ProcessEntries(IEditorDataStorageProvider& StorageInterface,
 	UEditorDataStorageCompatibility& Compatibility, const TFunctionRef<void(UE::Editor::DataStorage::RowHandle, const AddressType&)>& SetupRowCallback)
 {
 	// Thread-safe as it's only called from functions that already lock using.
@@ -897,7 +897,7 @@ void UEditorDataStorageCompatibility::TickPendingUObjectRegistration()
 		UObjectsPendingRegistration.ProcessEntries(*Storage, *this,
 			[this](UE::Editor::DataStorage::RowHandle Row, const TWeakObjectPtr<UObject>& Object)
 			{
-				ITypedElementDataStorageInterface* Interface = Storage;
+				IEditorDataStorageProvider* Interface = Storage;
 				Interface->AddColumn(Row, FTypedElementUObjectColumn{ .Object = Object });
 				Interface->AddColumn(Row, FTypedElementUObjectIdColumn
 					{ 
@@ -927,7 +927,7 @@ void UEditorDataStorageCompatibility::TickPendingExternalObjectRegistration()
 		ExternalObjectsPendingRegistration.ProcessEntries(*Storage, *this,
 			[this](UE::Editor::DataStorage::RowHandle Row, const ExternalObjectRegistration& Object)
 			{
-				ITypedElementDataStorageInterface* Interface = Storage;
+				IEditorDataStorageProvider* Interface = Storage;
 				Interface->AddColumn(Row, FTypedElementExternalObjectColumn{ .Object = Object.Object });
 				Interface->AddColumn(Row, FTypedElementScriptStructTypeInfoColumn{ .TypeInfo = Object.TypeInfo });
 				// Make sure the new row is tagged for update.
@@ -1265,7 +1265,7 @@ void UEditorDataStorageCompatibility::FRegistrationCommandChange::Revert(UObject
 	if (UObject* TargetRetrieved = TargetObject.Get(/*bEvenIfPendingKill=*/ true))
 	{
 		UEditorDataStorageCompatibility* DataStorageCompat = Owner.Get(); 
-		ITypedElementDataStorageInterface* DataStorage = DataStorageCompat->Storage;
+		IEditorDataStorageProvider* DataStorage = DataStorageCompat->Storage;
 			
 		RowHandle ObjectRow = DataStorageCompat->FindRowWithCompatibleObjectExplicit(TargetRetrieved);
 		if (DataStorage->IsRowAvailable(ObjectRow))
@@ -1301,7 +1301,7 @@ UEditorDataStorageCompatibility::FDeregistrationCommandChange::FDeregistrationCo
 {
 	using namespace UE::Editor::DataStorage;
 
-	ITypedElementDataStorageInterface* DataStorage = InOwner->Storage;
+	IEditorDataStorageProvider* DataStorage = InOwner->Storage;
 
 	RowHandle ObjectRow = InOwner->FindRowWithCompatibleObjectExplicit(InTargetObject);
 	if (DataStorage->IsRowAvailable(ObjectRow))
diff --git a/Engine/Plugins/Experimental/EditorDataStorage/Source/TedsCore/Private/TypedElementDatabaseCompatibility.h b/Engine/Plugins/Experimental/EditorDataStorage/Source/TedsCore/Private/TypedElementDatabaseCompatibility.h
--- a/Engine/Plugins/Experimental/EditorDataStorage/Source/TedsCore/Private/TypedElementDatabaseCompatibility.h
+++ b/Engine/Plugins/Experimental/EditorDataStorage/Source/TedsCore/Private/TypedElementDatabaseCompatibility.h
@@ -19,7 +19,7 @@
 
 #include "TypedElementDatabaseCompatibility.generated.h"
 
-class ITypedElementDataStorageInterface;
+class IEditorDataStorageProvider;
 struct FMassActorManager;
 class UTypedElementMementoSystem;
 
@@ -33,7 +33,7 @@ namespace UE::Editor::DataStorage
 UCLASS()
 class TEDSCORE_API UEditorDataStorageCompatibility
 	: public UObject
-	, public ITypedElementDataStorageCompatibilityInterface
+	, public IEditorDataStorageCompatibilityProvider
 {
 	GENERATED_BODY()
 
@@ -186,7 +186,7 @@ private:
 		int32 Num() const;
 		
 		void ForEachAddress(const TFunctionRef<void(AddressType&)>& Callback);
-		void ProcessEntries(ITypedElementDataStorageInterface& Storage, UEditorDataStorageCompatibility& Compatibility,
+		void ProcessEntries(IEditorDataStorageProvider& Storage, UEditorDataStorageCompatibility& Compatibility,
 			const TFunctionRef<void(UE::Editor::DataStorage::RowHandle, const AddressType&)>& SetupRowCallback);
 		void Reset();
 	};
@@ -208,7 +208,7 @@ private:
 	UE::Editor::DataStorage::TableHandle StandardActorWithTransformTable{ UE::Editor::DataStorage::InvalidTableHandle };
 	UE::Editor::DataStorage::TableHandle StandardUObjectTable{ UE::Editor::DataStorage::InvalidTableHandle };
 	UE::Editor::DataStorage::TableHandle StandardExternalObjectTable{ UE::Editor::DataStorage::InvalidTableHandle };
-	ITypedElementDataStorageInterface* Storage{ nullptr };
+	IEditorDataStorageProvider* Storage{ nullptr };
 
 	/**
 	 * Reference of objects (UObject and AActor) that need to be fully synced from the world to the database.
diff --git a/Engine/Plugins/Experimental/EditorDataStorage/Source/TedsCore/Private/TypedElementDatabaseUI.cpp b/Engine/Plugins/Experimental/EditorDataStorage/Source/TedsCore/Private/TypedElementDatabaseUI.cpp
--- a/Engine/Plugins/Experimental/EditorDataStorage/Source/TedsCore/Private/TypedElementDatabaseUI.cpp
+++ b/Engine/Plugins/Experimental/EditorDataStorage/Source/TedsCore/Private/TypedElementDatabaseUI.cpp
@@ -28,8 +28,8 @@ namespace Internal
 }
 
 void UEditorDataStorageUi::Initialize(
-	ITypedElementDataStorageInterface* StorageInterface,
-	ITypedElementDataStorageCompatibilityInterface* StorageCompatibilityInterface)
+	IEditorDataStorageProvider* StorageInterface,
+	IEditorDataStorageCompatibilityProvider* StorageCompatibilityInterface)
 {
 	checkf(StorageInterface, TEXT("TEDS' compatibility manager is being initialized with an invalid storage target."));
 
@@ -63,11 +63,11 @@ bool UEditorDataStorageUi::RegisterWidgetFactory(FName Purpose, const UScriptStr
 	{
 		switch (PurposeInfo->Type)
 		{
-		case ITypedElementDataStorageUiInterface::EPurposeType::Generic:
+		case IEditorDataStorageUiProvider::EPurposeType::Generic:
 			PurposeInfo->Factories.Emplace(Constructor);
 			PurposeInfo->bIsSorted = false;
 			return true;
-		case ITypedElementDataStorageUiInterface::EPurposeType::UniqueByName:
+		case IEditorDataStorageUiProvider::EPurposeType::UniqueByName:
 			if (PurposeInfo->Factories.IsEmpty())
 			{
 				PurposeInfo->Factories.Emplace(Constructor);
@@ -78,13 +78,13 @@ bool UEditorDataStorageUi::RegisterWidgetFactory(FName Purpose, const UScriptStr
 				PurposeInfo->Factories.EmplaceAt(0, Constructor);
 			}
 			return true;
-		case ITypedElementDataStorageUiInterface::EPurposeType::UniqueByNameAndColumn:
+		case IEditorDataStorageUiProvider::EPurposeType::UniqueByNameAndColumn:
 			UE_LOG(LogEditorDataStorageUI, Warning,
 				TEXT("Unable to register widget factory '%s' as purpose '%s' requires at least one column for matching."), 
 				*Constructor->GetName(), *Purpose.ToString());
 			return false;
 		default:
-			checkf(false, TEXT("Unexpected ITypedElementDataStorageUiInterface::EPurposeType found provided when registering widget factory."));
+			checkf(false, TEXT("Unexpected IEditorDataStorageUiProvider::EPurposeType found provided when registering widget factory."));
 			return false;
 		}
 	}
@@ -109,11 +109,11 @@ bool UEditorDataStorageUi::RegisterWidgetFactory(
 		{
 			switch (PurposeInfo->Type)
 			{
-			case ITypedElementDataStorageUiInterface::EPurposeType::Generic:
+			case IEditorDataStorageUiProvider::EPurposeType::Generic:
 				PurposeInfo->Factories.Emplace(Constructor);
 				PurposeInfo->bIsSorted = false;
 				return true;
-			case ITypedElementDataStorageUiInterface::EPurposeType::UniqueByName:
+			case IEditorDataStorageUiProvider::EPurposeType::UniqueByName:
 				if (!Columns.IsEmpty())
 				{
 					if (PurposeInfo->Factories.IsEmpty())
@@ -131,7 +131,7 @@ bool UEditorDataStorageUi::RegisterWidgetFactory(
 				{
 					return false;
 				}
-			case ITypedElementDataStorageUiInterface::EPurposeType::UniqueByNameAndColumn:
+			case IEditorDataStorageUiProvider::EPurposeType::UniqueByNameAndColumn:
 				if (!Columns.IsEmpty())
 				{
 					PurposeInfo->Factories.Emplace(Constructor, MoveTemp(Columns));
@@ -143,7 +143,7 @@ bool UEditorDataStorageUi::RegisterWidgetFactory(
 					return false;
 				}
 			default:
-				checkf(false, TEXT("Unexpected ITypedElementDataStorageUiInterface::EPurposeType found provided when registering widget factory."));
+				checkf(false, TEXT("Unexpected IEditorDataStorageUiProvider::EPurposeType found provided when registering widget factory."));
 				return false;
 			}
 		}
@@ -168,11 +168,11 @@ bool UEditorDataStorageUi::RegisterWidgetFactory(FName Purpose, TUniquePtr<FType
 	{
 		switch (PurposeInfo->Type)
 		{
-		case ITypedElementDataStorageUiInterface::EPurposeType::Generic:
+		case IEditorDataStorageUiProvider::EPurposeType::Generic:
 			PurposeInfo->Factories.Emplace(MoveTemp(Constructor));
 			PurposeInfo->bIsSorted = false;
 			return true;
-		case ITypedElementDataStorageUiInterface::EPurposeType::UniqueByName:
+		case IEditorDataStorageUiProvider::EPurposeType::UniqueByName:
 			if (PurposeInfo->Factories.IsEmpty())
 			{
 				PurposeInfo->Factories.Emplace(MoveTemp(Constructor));
@@ -183,13 +183,13 @@ bool UEditorDataStorageUi::RegisterWidgetFactory(FName Purpose, TUniquePtr<FType
 				PurposeInfo->Factories.EmplaceAt(0, MoveTemp(Constructor));
 			}
 			return true;
-		case ITypedElementDataStorageUiInterface::EPurposeType::UniqueByNameAndColumn:
+		case IEditorDataStorageUiProvider::EPurposeType::UniqueByNameAndColumn:
 			UE_LOG(LogEditorDataStorageUI, Warning,
 				TEXT("Unable to register widget factory '%s' as purpose '%s' requires at least one column for matching."),
 				*Constructor->GetTypeInfo()->GetName(), *Purpose.ToString());
 			return false;
 		default:
-			checkf(false, TEXT("Unexpected ITypedElementDataStorageUiInterface::EPurposeType found provided when registering widget factory."));
+			checkf(false, TEXT("Unexpected IEditorDataStorageUiProvider::EPurposeType found provided when registering widget factory."));
 			return false;
 		}
 	}
@@ -212,11 +212,11 @@ bool UEditorDataStorageUi::RegisterWidgetFactory(FName Purpose, TUniquePtr<FType
 		{
 			switch (PurposeInfo->Type)
 			{
-			case ITypedElementDataStorageUiInterface::EPurposeType::Generic:
+			case IEditorDataStorageUiProvider::EPurposeType::Generic:
 				PurposeInfo->Factories.Emplace(MoveTemp(Constructor));
 				PurposeInfo->bIsSorted = false;
 				return true;
-			case ITypedElementDataStorageUiInterface::EPurposeType::UniqueByName:
+			case IEditorDataStorageUiProvider::EPurposeType::UniqueByName:
 				if (!Columns.IsEmpty())
 				{
 					if (PurposeInfo->Factories.IsEmpty())
@@ -234,7 +234,7 @@ bool UEditorDataStorageUi::RegisterWidgetFactory(FName Purpose, TUniquePtr<FType
 				{
 					return false;
 				}
-			case ITypedElementDataStorageUiInterface::EPurposeType::UniqueByNameAndColumn:
+			case IEditorDataStorageUiProvider::EPurposeType::UniqueByNameAndColumn:
 				if (!Columns.IsEmpty())
 				{
 					PurposeInfo->Factories.Emplace(MoveTemp(Constructor), MoveTemp(Columns));
@@ -246,7 +246,7 @@ bool UEditorDataStorageUi::RegisterWidgetFactory(FName Purpose, TUniquePtr<FType
 					return false;
 				}
 			default:
-				checkf(false, TEXT("Unexpected ITypedElementDataStorageUiInterface::EPurposeType found provided when registering widget factory."));
+				checkf(false, TEXT("Unexpected IEditorDataStorageUiProvider::EPurposeType found provided when registering widget factory."));
 				return false;
 			}
 		}
diff --git a/Engine/Plugins/Experimental/EditorDataStorage/Source/TedsCore/Private/TypedElementDatabaseUI.h b/Engine/Plugins/Experimental/EditorDataStorage/Source/TedsCore/Private/TypedElementDatabaseUI.h
--- a/Engine/Plugins/Experimental/EditorDataStorage/Source/TedsCore/Private/TypedElementDatabaseUI.h
+++ b/Engine/Plugins/Experimental/EditorDataStorage/Source/TedsCore/Private/TypedElementDatabaseUI.h
@@ -11,15 +11,15 @@
 
 #include "TypedElementDatabaseUI.generated.h"
 
-class ITypedElementDataStorageInterface;
-class ITypedElementDataStorageCompatibilityInterface;
+class IEditorDataStorageProvider;
+class IEditorDataStorageCompatibilityProvider;
 
 TEDSCORE_API DECLARE_LOG_CATEGORY_EXTERN(LogEditorDataStorageUI, Log, All);
 
 UCLASS()
 class TEDSCORE_API UEditorDataStorageUi final
 	: public UObject
-	, public ITypedElementDataStorageUiInterface
+	, public IEditorDataStorageUiProvider
 {
 	GENERATED_BODY()
 
@@ -27,8 +27,8 @@ public:
 	~UEditorDataStorageUi() override = default;
 
 	void Initialize(
-		ITypedElementDataStorageInterface* StorageInterface, 
-		ITypedElementDataStorageCompatibilityInterface* StorageCompatibilityInterface);
+		IEditorDataStorageProvider* StorageInterface, 
+		IEditorDataStorageCompatibilityProvider* StorageCompatibilityInterface);
 	void Deinitialize();
 
 	void RegisterWidgetPurpose(FName Purpose, EPurposeType Type, FText Description) override;
@@ -112,6 +112,6 @@ private:
 	
 	TMap<FName, FPurposeInfo> WidgetPurposes;
 	
-	ITypedElementDataStorageInterface* Storage{ nullptr };
-	ITypedElementDataStorageCompatibilityInterface* StorageCompatibility{ nullptr };
+	IEditorDataStorageProvider* Storage{ nullptr };
+	IEditorDataStorageCompatibilityProvider* StorageCompatibility{ nullptr };
 };
\ No newline at end of file
diff --git a/Engine/Plugins/Experimental/EditorDataStorage/Source/TedsCore/Private/TypedElementSubsystems.cpp b/Engine/Plugins/Experimental/EditorDataStorage/Source/TedsCore/Private/TypedElementSubsystems.cpp
--- a/Engine/Plugins/Experimental/EditorDataStorage/Source/TedsCore/Private/TypedElementSubsystems.cpp
+++ b/Engine/Plugins/Experimental/EditorDataStorage/Source/TedsCore/Private/TypedElementSubsystems.cpp
@@ -13,7 +13,7 @@ UEditorDataStorageSubsystem::~UEditorDataStorageSubsystem()
 	DataStorage = nullptr;
 }
 
-ITypedElementDataStorageInterface* UEditorDataStorageSubsystem::Get()
+IEditorDataStorageProvider* UEditorDataStorageSubsystem::Get()
 {
 	if (!DataStorage)
 	{
@@ -24,7 +24,7 @@ ITypedElementDataStorageInterface* UEditorDataStorageSubsystem::Get()
 	return DataStorage;
 }
 
-const ITypedElementDataStorageInterface* UEditorDataStorageSubsystem::Get() const
+const IEditorDataStorageProvider* UEditorDataStorageSubsystem::Get() const
 {
 	return const_cast<UEditorDataStorageSubsystem*>(this)->Get();
 }
@@ -39,7 +39,7 @@ UEditorDataStorageUiSubsystem::~UEditorDataStorageUiSubsystem()
 	DataStorageUi = nullptr;
 }
 
-ITypedElementDataStorageUiInterface* UEditorDataStorageUiSubsystem::Get()
+IEditorDataStorageUiProvider* UEditorDataStorageUiSubsystem::Get()
 {
 	if (!DataStorageUi)
 	{
@@ -50,7 +50,7 @@ ITypedElementDataStorageUiInterface* UEditorDataStorageUiSubsystem::Get()
 	return DataStorageUi;
 }
 
-const ITypedElementDataStorageUiInterface* UEditorDataStorageUiSubsystem::Get() const
+const IEditorDataStorageUiProvider* UEditorDataStorageUiSubsystem::Get() const
 {
 	return const_cast<UEditorDataStorageUiSubsystem*>(this)->Get();
 }
@@ -65,7 +65,7 @@ UEditorDataStorageCompatibilitySubsystem::~UEditorDataStorageCompatibilitySubsys
 	DataStorageCompatibility = nullptr;
 }
 
-ITypedElementDataStorageCompatibilityInterface* UEditorDataStorageCompatibilitySubsystem::Get()
+IEditorDataStorageCompatibilityProvider* UEditorDataStorageCompatibilitySubsystem::Get()
 {
 	if (!DataStorageCompatibility)
 	{
@@ -77,7 +77,7 @@ ITypedElementDataStorageCompatibilityInterface* UEditorDataStorageCompatibilityS
 	return DataStorageCompatibility;
 }
 
-const ITypedElementDataStorageCompatibilityInterface* UEditorDataStorageCompatibilitySubsystem::Get() const
+const IEditorDataStorageCompatibilityProvider* UEditorDataStorageCompatibilitySubsystem::Get() const
 {
 	return const_cast<UEditorDataStorageCompatibilitySubsystem*>(this)->Get();
 }
diff --git a/Engine/Plugins/Experimental/EditorDataStorage/Source/TedsCore/Private/TypedElementsDataStorage.cpp b/Engine/Plugins/Experimental/EditorDataStorage/Source/TedsCore/Private/TypedElementsDataStorage.cpp
--- a/Engine/Plugins/Experimental/EditorDataStorage/Source/TedsCore/Private/TypedElementsDataStorage.cpp
+++ b/Engine/Plugins/Experimental/EditorDataStorage/Source/TedsCore/Private/TypedElementsDataStorage.cpp
@@ -100,16 +100,16 @@ void FEditorDataStorageModule::StartupModule()
 
 				// Allow any factories to register their content.
 				TArray<UClass*> FactoryClasses;
-				GetDerivedClasses(UTypedElementDataStorageFactory::StaticClass(), FactoryClasses);
+				GetDerivedClasses(UEditorDataStorageFactory::StaticClass(), FactoryClasses);
 
 				DataStorage->SetFactories(FactoryClasses);
-				TArray<UTypedElementDataStorageFactory*> Factories;
+				TArray<UEditorDataStorageFactory*> Factories;
 				Factories.Reserve(FactoryClasses.Num());
 
 				// First pass to call all registration without dependencies.
 				for (UEditorDataStorage::FactoryIterator Iterator = DataStorage->CreateFactoryIterator(); Iterator; ++Iterator)
 				{
-					UTypedElementDataStorageFactory* Factory = *Iterator;
+					UEditorDataStorageFactory* Factory = *Iterator;
 					
 					Factory->RegisterTables(*DataStorage);
 					Factory->RegisterTables(*DataStorage, *DataStorageCompatibility);
@@ -122,7 +122,7 @@ void FEditorDataStorageModule::StartupModule()
 				// Second pass to call all registration that would benefit or need the registration in the previous pass.
 				for (UEditorDataStorage::FactoryIterator Iterator = DataStorage->CreateFactoryIterator(); Iterator; ++Iterator)
 				{
-					UTypedElementDataStorageFactory* Factory = *Iterator;
+					UEditorDataStorageFactory* Factory = *Iterator;
 					
 					Factory->RegisterQueries(*DataStorage);
 					Factory->RegisterWidgetConstructors(*DataStorage, *DataStorageUi);
diff --git a/Engine/Plugins/Experimental/EditorDataStorage/Source/TedsCore/Public/TypedElementSubsystems.h b/Engine/Plugins/Experimental/EditorDataStorage/Source/TedsCore/Public/TypedElementSubsystems.h
--- a/Engine/Plugins/Experimental/EditorDataStorage/Source/TedsCore/Public/TypedElementSubsystems.h
+++ b/Engine/Plugins/Experimental/EditorDataStorage/Source/TedsCore/Public/TypedElementSubsystems.h
@@ -7,9 +7,9 @@
 
 #include "TypedElementSubsystems.generated.h"
 
-class ITypedElementDataStorageInterface;
-class ITypedElementDataStorageUiInterface;
-class ITypedElementDataStorageCompatibilityInterface;
+class IEditorDataStorageProvider;
+class IEditorDataStorageUiProvider;
+class IEditorDataStorageCompatibilityProvider;
 
 /**
  * A subsystem to provide alternative access to the Editor Data Storage. This should be used in most situations instead of 
@@ -26,11 +26,11 @@ public:
 
 	~UEditorDataStorageSubsystem() override;
 
-	ITypedElementDataStorageInterface* Get();
-	const ITypedElementDataStorageInterface* Get() const;
+	IEditorDataStorageProvider* Get();
+	const IEditorDataStorageProvider* Get() const;
 
 protected:
-	mutable ITypedElementDataStorageInterface* DataStorage{ nullptr };
+	mutable IEditorDataStorageProvider* DataStorage{ nullptr };
 };
 
 /**
@@ -48,11 +48,11 @@ public:
 
 	~UEditorDataStorageUiSubsystem() override;
 
-	ITypedElementDataStorageUiInterface* Get();
-	const ITypedElementDataStorageUiInterface* Get() const;
+	IEditorDataStorageUiProvider* Get();
+	const IEditorDataStorageUiProvider* Get() const;
 
 protected:
-	mutable ITypedElementDataStorageUiInterface* DataStorageUi{ nullptr };
+	mutable IEditorDataStorageUiProvider* DataStorageUi{ nullptr };
 };
 
 /**
@@ -70,9 +70,9 @@ public:
 
 	~UEditorDataStorageCompatibilitySubsystem() override;
 
-	ITypedElementDataStorageCompatibilityInterface* Get();
-	const ITypedElementDataStorageCompatibilityInterface* Get() const;
+	IEditorDataStorageCompatibilityProvider* Get();
+	const IEditorDataStorageCompatibilityProvider* Get() const;
 
 protected:
-	mutable ITypedElementDataStorageCompatibilityInterface* DataStorageCompatibility{ nullptr };
+	mutable IEditorDataStorageCompatibilityProvider* DataStorageCompatibility{ nullptr };
 };
\ No newline at end of file
diff --git a/Engine/Plugins/Experimental/EditorDataStorage/Source/TedsCore/Tests/DynamicColumnTests.cpp b/Engine/Plugins/Experimental/EditorDataStorage/Source/TedsCore/Tests/DynamicColumnTests.cpp
--- a/Engine/Plugins/Experimental/EditorDataStorage/Source/TedsCore/Tests/DynamicColumnTests.cpp
+++ b/Engine/Plugins/Experimental/EditorDataStorage/Source/TedsCore/Tests/DynamicColumnTests.cpp
@@ -12,7 +12,7 @@
 namespace UE::Editor::DataStorage::Tests
 {
 	BEGIN_DEFINE_SPEC(DynamicColumnTestFixture, "Editor.DataStorage.DynamicColumns", EAutomationTestFlags::EditorContext | EAutomationTestFlags::EngineFilter)
-		ITypedElementDataStorageInterface* TedsInterface = nullptr;
+		IEditorDataStorageProvider* TedsInterface = nullptr;
 		const FName TestTableName = TEXT("TestTable_DynamicColumnsTest");
 
 		TableHandle TestTable;
diff --git a/Engine/Plugins/Experimental/EditorDataStorage/Source/TedsCore/Tests/TypedElementAttributeBindingTests.cpp b/Engine/Plugins/Experimental/EditorDataStorage/Source/TedsCore/Tests/TypedElementAttributeBindingTests.cpp
--- a/Engine/Plugins/Experimental/EditorDataStorage/Source/TedsCore/Tests/TypedElementAttributeBindingTests.cpp
+++ b/Engine/Plugins/Experimental/EditorDataStorage/Source/TedsCore/Tests/TypedElementAttributeBindingTests.cpp
@@ -7,44 +7,42 @@
 #include "Elements/Framework/TypedElementRegistry.h"
 #include "Misc/AutomationTest.h"
 
-namespace UE::Editor::DataStorage
-{
-namespace Tests
+namespace UE::Editor::DataStorage::Tests
 {
 	BEGIN_DEFINE_SPEC(TedsAttributeBindingTestsFixture, "Editor.DataStorage.AttributeBinding", EAutomationTestFlags::EditorContext | EAutomationTestFlags::EngineFilter)
 
-		ITypedElementDataStorageInterface* TedsInterface = nullptr;
-		const FName TestTableName = TEXT("TestTable_AttributeBinding");
-		TableHandle TestTableHandle = InvalidTableHandle;
-		RowHandle TestRowHandle = InvalidRowHandle;
+	IEditorDataStorageProvider* TedsInterface = nullptr;
+	const FName TestTableName = TEXT("TestTable_AttributeBinding");
+	TableHandle TestTableHandle = InvalidTableHandle;
+	RowHandle TestRowHandle = InvalidRowHandle;
 
-		TableHandle RegisterTestTable() const
+	TableHandle RegisterTestTable() const
+	{
+		const TableHandle Table = TedsInterface->FindTable(TestTableName);
+		
+		if (Table != InvalidTableHandle)
 		{
-			const TableHandle Table = TedsInterface->FindTable(TestTableName);
-			
-			if (Table != InvalidTableHandle)
-			{
-				return Table;
-			}
-			
-			return TedsInterface->RegisterTable(
-			{
-				FTestColumnInt::StaticStruct(),
-				FTestColumnString::StaticStruct()
-			},
-			TestTableName);
+			return Table;
 		}
-	
-		RowHandle CreateTestRow(TableHandle InTableHandle) const
+		
+		return TedsInterface->RegisterTable(
 		{
-			const RowHandle RowHandle = TedsInterface->AddRow(InTableHandle);
-			return RowHandle;
-		}
+			FTestColumnInt::StaticStruct(),
+			FTestColumnString::StaticStruct()
+		},
+		TestTableName);
+	}
 	
-		void CleanupTestRow(RowHandle InRowHandle) const
-		{
-			TedsInterface->RemoveRow(InRowHandle);
-		}
+	RowHandle CreateTestRow(TableHandle InTableHandle) const
+	{
+		const RowHandle RowHandle = TedsInterface->AddRow(InTableHandle);
+		return RowHandle;
+	}
+	
+	void CleanupTestRow(RowHandle InRowHandle) const
+	{
+		TedsInterface->RemoveRow(InRowHandle);
+	}
 	
 	END_DEFINE_SPEC(TedsAttributeBindingTestsFixture)
 
@@ -224,7 +222,6 @@ namespace Tests
 			TedsInterface = nullptr;
 		});
 	}
-} // namespace Tests
-} // namespace UE::Editor::DataStorage
+} // namespace UE::Editor::DataStorage::Tests
 
 #endif // WITH_TESTS
\ No newline at end of file
diff --git a/Engine/Plugins/Experimental/EditorDataStorage/Source/TedsCore/Tests/TypedElementDataStorageInterfaceTests.cpp b/Engine/Plugins/Experimental/EditorDataStorage/Source/TedsCore/Tests/TypedElementDataStorageInterfaceTests.cpp
--- a/Engine/Plugins/Experimental/EditorDataStorage/Source/TedsCore/Tests/TypedElementDataStorageInterfaceTests.cpp
+++ b/Engine/Plugins/Experimental/EditorDataStorage/Source/TedsCore/Tests/TypedElementDataStorageInterfaceTests.cpp
@@ -6,179 +6,283 @@
 #include "Elements/Framework/TypedElementRegistry.h"
 #include "Misc/AutomationTest.h"
 
-namespace UE::Editor::DataStorage
+namespace UE::Editor::DataStorage::Tests
 {
-	namespace Tests
+	TableHandle RegisterTestTableA(IEditorDataStorageProvider* TedsInterface)
 	{
-		TableHandle RegisterTestTableA(ITypedElementDataStorageInterface* TedsInterface)
+		// RegisterTable is not idempotent, so this utility function makes it so
+	
+		const FName TestTableAName = TEXT("TestTableA");
+		TableHandle Table = TedsInterface->FindTable(TestTableAName);
+		if (Table != InvalidTableHandle)
 		{
-			// RegisterTable is not idempotent, so this utility function makes it so
-		
-			const FName TestTableAName = TEXT("TestTableA");
-			TableHandle Table = TedsInterface->FindTable(TestTableAName);
-			if (Table != InvalidTableHandle)
-			{
-				return Table;
-			}
-		
-			return TedsInterface->RegisterTable(
-			{
-				FTestColumnA::StaticStruct(),
-				FTestColumnB::StaticStruct(),
-				FTestTagColumnA::StaticStruct(),
-				FTestTagColumnB::StaticStruct()
-			},
-			TEXT("TestTableA"));
+			return Table;
 		}
-
-		bool IsEmptyQueryDescription(const FQueryDescription& QueryDescription)
+	
+		return TedsInterface->RegisterTable(
 		{
-			return
-				QueryDescription.Callback.MonitoredType == nullptr &&
-				QueryDescription.SelectionTypes.IsEmpty() &&
-				QueryDescription.SelectionAccessTypes.IsEmpty() &&
-				QueryDescription.SelectionMetaData.IsEmpty() &&
-				QueryDescription.ConditionTypes.IsEmpty() &&
-				QueryDescription.ConditionOperators.IsEmpty() &&
-				QueryDescription.DependencyTypes.IsEmpty() &&
-				QueryDescription.DependencyFlags.IsEmpty() &&
-				QueryDescription.CachedDependencies.IsEmpty() &&
-				QueryDescription.Subqueries.IsEmpty();
-		}
+			FTestColumnA::StaticStruct(),
+			FTestColumnB::StaticStruct(),
+			FTestTagColumnA::StaticStruct(),
+			FTestTagColumnB::StaticStruct()
+		},
+		TEXT("TestTableA"));
+	}
+
+	bool IsEmptyQueryDescription(const FQueryDescription& QueryDescription)
+	{
+		return
+			QueryDescription.Callback.MonitoredType == nullptr &&
+			QueryDescription.SelectionTypes.IsEmpty() &&
+			QueryDescription.SelectionAccessTypes.IsEmpty() &&
+			QueryDescription.SelectionMetaData.IsEmpty() &&
+			QueryDescription.ConditionTypes.IsEmpty() &&
+			QueryDescription.ConditionOperators.IsEmpty() &&
+			QueryDescription.DependencyTypes.IsEmpty() &&
+			QueryDescription.DependencyFlags.IsEmpty() &&
+			QueryDescription.CachedDependencies.IsEmpty() &&
+			QueryDescription.Subqueries.IsEmpty();
+	}
+
+	void IncrementThing(int64* QueryCallCountPtr)
+	{
+		++(*QueryCallCountPtr);
+	}
+
+	// Temporary wrapper to only register a single observer
+	// There is no way to unregister an observer at this time
+	template<typename ColumnType>
+	void RegisterTestObserver(
+		FName ObserverName,
+		IEditorDataStorageProvider* TedsInterface, Queries::FObserver::EEvent EventType, int64* QueryCallCountPtr, QueryHandle* QueryHandleInOut)
+	{
+		using namespace UE::Editor::DataStorage::Queries;
+		FQueryDescription QueryDescription = TedsInterface->GetQueryDescription(*QueryHandleInOut);
 
-		void IncrementThing(int64* QueryCallCountPtr)
+		// Check if it is empty... If not then we need to register
+		// This is for an idempotent API to the test harness				
+		if (IsEmptyQueryDescription(QueryDescription))
 		{
-			++(*QueryCallCountPtr);
+			*QueryHandleInOut = TedsInterface->RegisterQuery(
+				Select(
+					ObserverName,
+					FObserver(EventType, ColumnType::StaticStruct()),
+					[QueryCallCountPtr](IQueryContext& Context, RowHandle Row)
+					{
+						++(*QueryCallCountPtr);
+					})
+				.Compile());
 		}
-
-		// Temporary wrapper to only register a single observer
-		// There is no way to unregister an observer at this time
-		template<typename ColumnType>
-		void RegisterTestObserver(
-			FName ObserverName,
-			ITypedElementDataStorageInterface* TedsInterface, Queries::FObserver::EEvent EventType, int64* QueryCallCountPtr, QueryHandle* QueryHandleInOut)
+	}
+
+	BEGIN_DEFINE_SPEC(EditorDataStorageTestsFixture, "Editor.DataStorage.Core", EAutomationTestFlags::EditorContext | EAutomationTestFlags::EngineFilter)
+		IEditorDataStorageProvider* TedsInterface = nullptr;
+		TableHandle TestTableHandleA = InvalidTableHandle;
+
+		TArray<RowHandle> CreatedRows;
+
+		QueryHandle DataColumnA_AddObserverHandle = InvalidQueryHandle;
+		QueryHandle DataColumnC_AddObserverHandle = InvalidQueryHandle;
+		QueryHandle DataColumnB_RemoveObserverHandle = InvalidQueryHandle;
+		QueryHandle TagColumnA_AddObserverHandle = InvalidQueryHandle;
+		QueryHandle TagColumnC_AddObserverHandle = InvalidQueryHandle;
+		QueryHandle TagColumnB_RemoveObserverHandle = InvalidQueryHandle;
+
+		// Non-zero to catch problems in test where the difference is not being used
+		// Avoids need to reset this to 0 for each test
+		int64 DataColumnA_AddObserverCallCount = 10;
+		int64 DataColumnC_AddObserverCallCount = 10;
+		int64 DataColumnB_RemoveObserverCallCount = 10;
+		int64 TagColumnA_AddObserverCallCount = 10;
+		int64 TagColumnC_AddObserverCallCount = 10;
+		int64 TagColumnB_RemoveObserverCallCount = 10;
+
+		void CreateTestRows(int32 Count)
 		{
-			using namespace UE::Editor::DataStorage::Queries;
-			FQueryDescription QueryDescription = TedsInterface->GetQueryDescription(*QueryHandleInOut);
-
-			// Check if it is empty... If not then we need to register
-			// This is for an idempotent API to the test harness				
-			if (IsEmptyQueryDescription(QueryDescription))
+			CreatedRows.Reserve(Count);
+			for (int32 Index = 0; Index < Count; ++Index)
 			{
-				*QueryHandleInOut = TedsInterface->RegisterQuery(
-					Select(
-						ObserverName,
-						FObserver(EventType, ColumnType::StaticStruct()),
-						[QueryCallCountPtr](IQueryContext& Context, RowHandle Row)
-						{
-							++(*QueryCallCountPtr);
-						})
-					.Compile());
+				RowHandle RowHandle = TedsInterface->AddRow(TestTableHandleA);
+				TestFalse("Expect valid row handle", RowHandle == InvalidRowHandle);
+				CreatedRows.Add(RowHandle);
 			}
 		}
-	
-		BEGIN_DEFINE_SPEC(EditorDataStorageTestsFixture, "Editor.DataStorage.Core", EAutomationTestFlags::EditorContext | EAutomationTestFlags::EngineFilter)
-			ITypedElementDataStorageInterface* TedsInterface = nullptr;
-			TableHandle TestTableHandleA = InvalidTableHandle;
-
-			TArray<RowHandle> CreatedRows;
-
-			QueryHandle DataColumnA_AddObserverHandle = InvalidQueryHandle;
-			QueryHandle DataColumnC_AddObserverHandle = InvalidQueryHandle;
-			QueryHandle DataColumnB_RemoveObserverHandle = InvalidQueryHandle;
-			QueryHandle TagColumnA_AddObserverHandle = InvalidQueryHandle;
-			QueryHandle TagColumnC_AddObserverHandle = InvalidQueryHandle;
-			QueryHandle TagColumnB_RemoveObserverHandle = InvalidQueryHandle;
-	
-			// Non-zero to catch problems in test where the difference is not being used
-			// Avoids need to reset this to 0 for each test
-			int64 DataColumnA_AddObserverCallCount = 10;
-			int64 DataColumnC_AddObserverCallCount = 10;
-			int64 DataColumnB_RemoveObserverCallCount = 10;
-			int64 TagColumnA_AddObserverCallCount = 10;
-			int64 TagColumnC_AddObserverCallCount = 10;
-			int64 TagColumnB_RemoveObserverCallCount = 10;
-	
-			void CreateTestRows(int32 Count)
-			{
-				CreatedRows.Reserve(Count);
-				for (int32 Index = 0; Index < Count; ++Index)
-				{
-					RowHandle RowHandle = TedsInterface->AddRow(TestTableHandleA);
-					TestFalse("Expect valid row handle", RowHandle == InvalidRowHandle);
-					CreatedRows.Add(RowHandle);
-				}
-			}
 
-			void CleanupTestRows()
+		void CleanupTestRows()
+		{
+			for (RowHandle RowHandle : CreatedRows)
 			{
-				for (RowHandle RowHandle : CreatedRows)
-				{
-					TestTrue("Expected row to have been available", TedsInterface->IsRowAvailable(RowHandle));
-					TestTrue("Expected row to have been assigned to a table", TedsInterface->IsRowAssigned(RowHandle));
-					TedsInterface->RemoveRow(RowHandle);
-					TestFalse("Expected row to have been unassigned a table", TedsInterface->IsRowAssigned(RowHandle));
-					TestFalse("Expected row to be not available", TedsInterface->IsRowAvailable(RowHandle));
-				}
-				CreatedRows.Empty();
+				TestTrue("Expected row to have been available", TedsInterface->IsRowAvailable(RowHandle));
+				TestTrue("Expected row to have been assigned to a table", TedsInterface->IsRowAssigned(RowHandle));
+				TedsInterface->RemoveRow(RowHandle);
+				TestFalse("Expected row to have been unassigned a table", TedsInterface->IsRowAssigned(RowHandle));
+				TestFalse("Expected row to be not available", TedsInterface->IsRowAvailable(RowHandle));
 			}
-		END_DEFINE_SPEC(EditorDataStorageTestsFixture)
+			CreatedRows.Empty();
+		}
+	END_DEFINE_SPEC(EditorDataStorageTestsFixture)
+
+	void EditorDataStorageTestsFixture::Define()
+	{
+		BeforeEach([this]()
+		{
+			UTypedElementRegistry* TypedElementRegistry = UTypedElementRegistry::GetInstance();
+			TedsInterface = TypedElementRegistry->GetMutableDataStorage();
+			TestTrue("", TedsInterface != nullptr);
+		});
 
-		void EditorDataStorageTestsFixture::Define()
+		Describe("", [this]()
 		{
 			BeforeEach([this]()
 			{
-				UTypedElementRegistry* TypedElementRegistry = UTypedElementRegistry::GetInstance();
-				TedsInterface = TypedElementRegistry->GetMutableDataStorage();
-				TestTrue("", TedsInterface != nullptr);
+				TestTrue("Test requires valid TEDS Interface", TedsInterface != nullptr);
+			});
+		
+			Describe("", [this]()
+			{
+				Describe("RegisterTable(TConstArrayView<const UScriptStruct*>, const FName)", [this]()
+				{
+					It("should register a table and provide a valid handle", [this]()
+					{
+						// Note, this test doesn't really do anything a second time.
+						// The test fixture depends on global state which currently cannot be cleaned
+						// up.
+						TableHandle Handle = RegisterTestTableA(TedsInterface);
+						TestNotEqual("Expecting valid table handle", Handle, InvalidTableHandle);
+					});
+				});
 			});
 
 			Describe("", [this]()
 			{
 				BeforeEach([this]()
 				{
-					TestTrue("Test requires valid TEDS Interface", TedsInterface != nullptr);
+					TestTableHandleA = RegisterTestTableA(TedsInterface);
 				});
 			
-				Describe("", [this]()
+				AfterEach([this]()
 				{
-					Describe("RegisterTable(TConstArrayView<const UScriptStruct*>, const FName)", [this]()
+					CleanupTestRows();
+				});
+			
+				Describe("AddRow(TableHandle)", [this]()
+				{
+					for (int32 RowCount = 1; RowCount <= 2; ++RowCount)
 					{
-						It("should register a table and provide a valid handle", [this]()
+						It(FString::Printf(TEXT("should should be able to create %d rows when called %d times"), RowCount, RowCount), [this, RowCount]()
 						{
-							// Note, this test doesn't really do anything a second time.
-							// The test fixture depends on global state which currently cannot be cleaned
-							// up.
-							TableHandle Handle = RegisterTestTableA(TedsInterface);
-							TestNotEqual("Expecting valid table handle", Handle, InvalidTableHandle);
+							for (int32 CreatedRowIndex = 0; CreatedRowIndex < RowCount; ++CreatedRowIndex)
+							{
+								RowHandle RowHandle = TedsInterface->AddRow(TestTableHandleA);
+								CreatedRows.Add(RowHandle);
+								TestNotEqual("Expecting valid row", RowHandle, InvalidRowHandle);
+							}
 						});
-					});
+					}
+				});
+
+				Describe("BatchAddRow(TableHandle, int32, TypedElementDataStorageCreationCallbackRef)", [this]()
+				{
+					for (int32 RowCount = 1; RowCount <= 2; ++RowCount)
+					{
+						It(FString::Printf(TEXT("should should be able to create %d rows when called once"), RowCount), [this, RowCount]()
+						{
+							auto RowInitializeCallback = [this](RowHandle RowHandle)
+							{
+								CreatedRows.Add(RowHandle);
+								TestNotEqual("Expecting valid row", RowHandle, InvalidRowHandle);
+							};
+										
+							RowHandle RowHandle = TedsInterface->BatchAddRow(TestTableHandleA, RowCount, RowInitializeCallback);
+										
+							TestEqual("Unexpected number of created rows", CreatedRows.Num(), RowCount);
+						});
+					}
 				});
 
+		
 				Describe("", [this]()
 				{
+					/**
+					 * Prerequisites for this batch of tests are:
+					 * - A test table is registered
+					 * - An add and remove observer is registered for FTestInt64Column
+					 */
 					BeforeEach([this]()
 					{
+						using namespace Queries;
+
 						TestTableHandleA = RegisterTestTableA(TedsInterface);
+						TestNotEqual("Expecting valid table handle", TestTableHandleA, InvalidTableHandle);
+
+						RegisterTestObserver<FTestColumnA>(
+							TEXT("Increment CallCount when FTestColumnA added"),
+							TedsInterface,
+							FObserver::EEvent::Add,
+							&DataColumnA_AddObserverCallCount,
+							&DataColumnA_AddObserverHandle);
+						TestTrue("Expect valid query observer handle", DataColumnA_AddObserverHandle != InvalidQueryHandle);
+					
+						RegisterTestObserver<FTestColumnC>(
+							TEXT("Increment CallCount when FTestColumnC added"),
+							TedsInterface,
+							FObserver::EEvent::Add,
+							&DataColumnC_AddObserverCallCount,
+							&DataColumnC_AddObserverHandle);
+						TestTrue("Expect valid query observer handle", DataColumnC_AddObserverHandle != InvalidQueryHandle);
+
+						RegisterTestObserver<FTestColumnB>(
+							TEXT("Increment CallCount when FTestColumnB removed"),
+							TedsInterface,
+							FObserver::EEvent::Remove,
+							&DataColumnB_RemoveObserverCallCount,
+							&DataColumnB_RemoveObserverHandle);
+						TestTrue("Expect valid query observer handle", DataColumnB_RemoveObserverHandle != InvalidQueryHandle);
+
+						RegisterTestObserver<FTestTagColumnA>(
+							TEXT("Increment CallCount when FTestColumnA added"),
+							TedsInterface,
+							FObserver::EEvent::Add,
+							&TagColumnA_AddObserverCallCount,
+							&TagColumnA_AddObserverHandle);
+						TestTrue("Expect valid query observer handle", TagColumnA_AddObserverHandle != InvalidQueryHandle);
+											
+						RegisterTestObserver<FTestTagColumnC>(
+							TEXT("Increment CallCount when FTestTagColumnC added"),
+							TedsInterface,
+							FObserver::EEvent::Add,
+							&TagColumnC_AddObserverCallCount, 
+							&TagColumnC_AddObserverHandle);
+						TestTrue("Expect valid query observer handle", TagColumnC_AddObserverHandle != InvalidQueryHandle);
+
+						RegisterTestObserver<FTestTagColumnB>(
+							TEXT("Increment CallCount when FTestTagColumnB removed"),
+							TedsInterface, FObserver::EEvent::Remove,
+							&TagColumnB_RemoveObserverCallCount,
+							&TagColumnB_RemoveObserverHandle);
+						TestTrue("Expect valid query observer handle", TagColumnB_RemoveObserverHandle != InvalidQueryHandle);					
 					});
 				
 					AfterEach([this]()
 					{
 						CleanupTestRows();
 					});
-				
+
 					Describe("AddRow(TableHandle)", [this]()
 					{
 						for (int32 RowCount = 1; RowCount <= 2; ++RowCount)
 						{
-							It(FString::Printf(TEXT("should should be able to create %d rows when called %d times"), RowCount, RowCount), [this, RowCount]()
+							It(FString::Printf(TEXT("should invoke Add observer %d times when called with table containing data column %d times"), RowCount, RowCount), [this, RowCount]()
 							{
+								int64 PreviousQueryCallCount = DataColumnA_AddObserverCallCount;
+							
 								for (int32 CreatedRowIndex = 0; CreatedRowIndex < RowCount; ++CreatedRowIndex)
 								{
 									RowHandle RowHandle = TedsInterface->AddRow(TestTableHandleA);
 									CreatedRows.Add(RowHandle);
-									TestNotEqual("Expecting valid row", RowHandle, InvalidRowHandle);
 								}
+
+								TestEqual(TEXT("Expect observer to be called correct number of times"), DataColumnA_AddObserverCallCount, PreviousQueryCallCount + RowCount);
 							});
 						}
 					});
@@ -187,399 +291,292 @@ namespace UE::Editor::DataStorage
 					{
 						for (int32 RowCount = 1; RowCount <= 2; ++RowCount)
 						{
-							It(FString::Printf(TEXT("should should be able to create %d rows when called once"), RowCount), [this, RowCount]()
+							It(FString::Printf(TEXT("should invoke Add observer %d times when called with table containing data column"), RowCount), [this, RowCount]()
 							{
+								int64 PreviousQueryCallCount = DataColumnA_AddObserverCallCount;
+							
 								auto RowInitializeCallback = [this](RowHandle RowHandle)
 								{
 									CreatedRows.Add(RowHandle);
-									TestNotEqual("Expecting valid row", RowHandle, InvalidRowHandle);
 								};
-											
+															
 								RowHandle RowHandle = TedsInterface->BatchAddRow(TestTableHandleA, RowCount, RowInitializeCallback);
-											
-								TestEqual("Unexpected number of created rows", CreatedRows.Num(), RowCount);
+							
+								TestEqual(TEXT("Expect observer to be called correct number of times"), DataColumnA_AddObserverCallCount, PreviousQueryCallCount + RowCount);
 							});
 						}
 					});
 
-			
 					Describe("", [this]()
 					{
-						/**
-						 * Prerequisites for this batch of tests are:
-						 * - A test table is registered
-						 * - An add and remove observer is registered for FTestInt64Column
-						 */
 						BeforeEach([this]()
 						{
-							using namespace Queries;
-
-							TestTableHandleA = RegisterTestTableA(TedsInterface);
-							TestNotEqual("Expecting valid table handle", TestTableHandleA, InvalidTableHandle);
-
-							RegisterTestObserver<FTestColumnA>(
-								TEXT("Increment CallCount when FTestColumnA added"),
-								TedsInterface,
-								FObserver::EEvent::Add,
-								&DataColumnA_AddObserverCallCount,
-								&DataColumnA_AddObserverHandle);
-							TestTrue("Expect valid query observer handle", DataColumnA_AddObserverHandle != InvalidQueryHandle);
-						
-							RegisterTestObserver<FTestColumnC>(
-								TEXT("Increment CallCount when FTestColumnC added"),
-								TedsInterface,
-								FObserver::EEvent::Add,
-								&DataColumnC_AddObserverCallCount,
-								&DataColumnC_AddObserverHandle);
-							TestTrue("Expect valid query observer handle", DataColumnC_AddObserverHandle != InvalidQueryHandle);
-
-							RegisterTestObserver<FTestColumnB>(
-								TEXT("Increment CallCount when FTestColumnB removed"),
-								TedsInterface,
-								FObserver::EEvent::Remove,
-								&DataColumnB_RemoveObserverCallCount,
-								&DataColumnB_RemoveObserverHandle);
-							TestTrue("Expect valid query observer handle", DataColumnB_RemoveObserverHandle != InvalidQueryHandle);
-
-							RegisterTestObserver<FTestTagColumnA>(
-								TEXT("Increment CallCount when FTestColumnA added"),
-								TedsInterface,
-								FObserver::EEvent::Add,
-								&TagColumnA_AddObserverCallCount,
-								&TagColumnA_AddObserverHandle);
-							TestTrue("Expect valid query observer handle", TagColumnA_AddObserverHandle != InvalidQueryHandle);
-												
-							RegisterTestObserver<FTestTagColumnC>(
-								TEXT("Increment CallCount when FTestTagColumnC added"),
-								TedsInterface,
-								FObserver::EEvent::Add,
-								&TagColumnC_AddObserverCallCount, 
-								&TagColumnC_AddObserverHandle);
-							TestTrue("Expect valid query observer handle", TagColumnC_AddObserverHandle != InvalidQueryHandle);
-
-							RegisterTestObserver<FTestTagColumnB>(
-								TEXT("Increment CallCount when FTestTagColumnB removed"),
-								TedsInterface, FObserver::EEvent::Remove,
-								&TagColumnB_RemoveObserverCallCount,
-								&TagColumnB_RemoveObserverHandle);
-							TestTrue("Expect valid query observer handle", TagColumnB_RemoveObserverHandle != InvalidQueryHandle);					
-						});
-					
-						AfterEach([this]()
-						{
-							CleanupTestRows();
+							CreateTestRows(1);
 						});
 
-						Describe("AddRow(TableHandle)", [this]()
+						Describe("RemoveRow(TableHandle)", [this]()
 						{
-							for (int32 RowCount = 1; RowCount <= 2; ++RowCount)
+							It("should remove row when called", [this]()
 							{
-								It(FString::Printf(TEXT("should invoke Add observer %d times when called with table containing data column %d times"), RowCount, RowCount), [this, RowCount]()
-								{
-									int64 PreviousQueryCallCount = DataColumnA_AddObserverCallCount;
-								
-									for (int32 CreatedRowIndex = 0; CreatedRowIndex < RowCount; ++CreatedRowIndex)
-									{
-										RowHandle RowHandle = TedsInterface->AddRow(TestTableHandleA);
-										CreatedRows.Add(RowHandle);
-									}
-
-									TestEqual(TEXT("Expect observer to be called correct number of times"), DataColumnA_AddObserverCallCount, PreviousQueryCallCount + RowCount);
-								});
-							}
-						});
+								RowHandle Row = CreatedRows[0];
+								TestTrue("Row is available", TedsInterface->IsRowAvailable(Row));
 
-						Describe("BatchAddRow(TableHandle, int32, TypedElementDataStorageCreationCallbackRef)", [this]()
-						{
-							for (int32 RowCount = 1; RowCount <= 2; ++RowCount)
-							{
-								It(FString::Printf(TEXT("should invoke Add observer %d times when called with table containing data column"), RowCount), [this, RowCount]()
-								{
-									int64 PreviousQueryCallCount = DataColumnA_AddObserverCallCount;
-								
-									auto RowInitializeCallback = [this](RowHandle RowHandle)
-									{
-										CreatedRows.Add(RowHandle);
-									};
-																
-									RowHandle RowHandle = TedsInterface->BatchAddRow(TestTableHandleA, RowCount, RowInitializeCallback);
-								
-									TestEqual(TEXT("Expect observer to be called correct number of times"), DataColumnA_AddObserverCallCount, PreviousQueryCallCount + RowCount);
-								});
-							}
-						});
+								TedsInterface->RemoveRow(Row);
+							
 
-						Describe("", [this]()
-						{
-							BeforeEach([this]()
-							{
-								CreateTestRows(1);
-							});
+								TestFalse("Row is not available", TedsInterface->IsRowAvailable(Row));
 
-							Describe("RemoveRow(TableHandle)", [this]()
+								CreatedRows.RemoveAt(0);
+							});
+						
+							It("should call observer of removed data column when called on row that has the column", [this]()
 							{
-								It("should remove row when called", [this]()
-								{
-									RowHandle Row = CreatedRows[0];
-									TestTrue("Row is available", TedsInterface->IsRowAvailable(Row));
-
-									TedsInterface->RemoveRow(Row);
-								
-
-									TestFalse("Row is not available", TedsInterface->IsRowAvailable(Row));
-
-									CreatedRows.RemoveAt(0);
-								});
+								int64 PreviousQueryCallCount = DataColumnB_RemoveObserverCallCount;
 							
-								It("should call observer of removed data column when called on row that has the column", [this]()
-								{
-									int64 PreviousQueryCallCount = DataColumnB_RemoveObserverCallCount;
-								
-									RowHandle Row = CreatedRows[0];
-									TedsInterface->RemoveRow(Row);
+								RowHandle Row = CreatedRows[0];
+								TedsInterface->RemoveRow(Row);
 
-									TestEqual("Expected observer called", DataColumnB_RemoveObserverCallCount, PreviousQueryCallCount + 1);
+								TestEqual("Expected observer called", DataColumnB_RemoveObserverCallCount, PreviousQueryCallCount + 1);
 
-									CreatedRows.RemoveAt(0);
-								});
-							
-								It("should call observer of removed tag column when called on row that has the column", [this]()
-								{
-									int64 PreviousQueryCallCount = TagColumnB_RemoveObserverCallCount;
-								
-									RowHandle Row = CreatedRows[0];
-									TedsInterface->RemoveRow(Row);
-								
-									TestEqual("Expected observer called", TagColumnB_RemoveObserverCallCount, PreviousQueryCallCount + 1);
-
-									CreatedRows.RemoveAt(0);
-								});
+								CreatedRows.RemoveAt(0);
 							});
 						
-							Describe("AddColumn(RowHandle, const UScriptStruct*)", [this]()
-							{					
-								It("should add data column when called", [this]()
-								{
-									RowHandle Row = CreatedRows[0];
-									TestFalse("Expect row to not have column about to be added",
-										TedsInterface->HasColumns(Row, MakeArrayView({FTestColumnC::StaticStruct()})));
+							It("should call observer of removed tag column when called on row that has the column", [this]()
+							{
+								int64 PreviousQueryCallCount = TagColumnB_RemoveObserverCallCount;
+							
+								RowHandle Row = CreatedRows[0];
+								TedsInterface->RemoveRow(Row);
+							
+								TestEqual("Expected observer called", TagColumnB_RemoveObserverCallCount, PreviousQueryCallCount + 1);
 
-									TedsInterface->AddColumn(Row, FTestColumnC::StaticStruct());
+								CreatedRows.RemoveAt(0);
+							});
+						});
+					
+						Describe("AddColumn(RowHandle, const UScriptStruct*)", [this]()
+						{					
+							It("should add data column when called", [this]()
+							{
+								RowHandle Row = CreatedRows[0];
+								TestFalse("Expect row to not have column about to be added",
+									TedsInterface->HasColumns(Row, MakeArrayView({FTestColumnC::StaticStruct()})));
 
-									TestTrue("Expect row to have column added",
-										TedsInterface->HasColumns(Row, MakeArrayView({FTestColumnC::StaticStruct()})));
+								TedsInterface->AddColumn(Row, FTestColumnC::StaticStruct());
 
-								});
+								TestTrue("Expect row to have column added",
+									TedsInterface->HasColumns(Row, MakeArrayView({FTestColumnC::StaticStruct()})));
 
-								It("should invoke registered observer when called with data column", [this]()
-								{
-									int64 PreviousQueryCallCount = DataColumnC_AddObserverCallCount;
-									RowHandle Row = CreatedRows[0];
-									TedsInterface->AddColumn(Row, FTestColumnC::StaticStruct());
+							});
+
+							It("should invoke registered observer when called with data column", [this]()
+							{
+								int64 PreviousQueryCallCount = DataColumnC_AddObserverCallCount;
+								RowHandle Row = CreatedRows[0];
+								TedsInterface->AddColumn(Row, FTestColumnC::StaticStruct());
 
-									TestEqual("Expect observer to have been called one time", DataColumnC_AddObserverCallCount, PreviousQueryCallCount + 1);
-								});
+								TestEqual("Expect observer to have been called one time", DataColumnC_AddObserverCallCount, PreviousQueryCallCount + 1);
+							});
 
-								It("should add tag column when called", [this]()
-								{
-									RowHandle Row = CreatedRows[0];
-									TestFalse("Expect row to not have column about to be added",
-										TedsInterface->HasColumns(Row, MakeArrayView({FTestTagColumnC::StaticStruct()})));
+							It("should add tag column when called", [this]()
+							{
+								RowHandle Row = CreatedRows[0];
+								TestFalse("Expect row to not have column about to be added",
+									TedsInterface->HasColumns(Row, MakeArrayView({FTestTagColumnC::StaticStruct()})));
 
-									TedsInterface->AddColumn(Row, FTestTagColumnC::StaticStruct());
+								TedsInterface->AddColumn(Row, FTestTagColumnC::StaticStruct());
 
-									TestTrue("Expect row to have column added",
-										TedsInterface->HasColumns(Row, MakeArrayView({FTestTagColumnC::StaticStruct()})));
+								TestTrue("Expect row to have column added",
+									TedsInterface->HasColumns(Row, MakeArrayView({FTestTagColumnC::StaticStruct()})));
 
-								});
+							});
 
-								It("should invoke registered observer when called with tag column", [this]()
-								{
-									int64 PreviousQueryCallCount = TagColumnC_AddObserverCallCount;
-									RowHandle Row = CreatedRows[0];
-									TedsInterface->AddColumn(Row, FTestTagColumnC::StaticStruct());
+							It("should invoke registered observer when called with tag column", [this]()
+							{
+								int64 PreviousQueryCallCount = TagColumnC_AddObserverCallCount;
+								RowHandle Row = CreatedRows[0];
+								TedsInterface->AddColumn(Row, FTestTagColumnC::StaticStruct());
 
-									TestEqual("Expect observer to have been called one time", TagColumnC_AddObserverCallCount, PreviousQueryCallCount + 1);
-								});
+								TestEqual("Expect observer to have been called one time", TagColumnC_AddObserverCallCount, PreviousQueryCallCount + 1);
 							});
-						
-							Describe("RemoveColumn(RowHandle, TConstArrayView<const UScriptStruct*>)", [this]()
+						});
+					
+						Describe("RemoveColumn(RowHandle, TConstArrayView<const UScriptStruct*>)", [this]()
+						{
+							It("should remove a single column when called", [this]()
 							{
-								It("should remove a single column when called", [this]()
-								{
-									RowHandle Row = CreatedRows[0];
-									TestTrue("Expected to have column about to be removed",
-										TedsInterface->HasColumns(Row, MakeArrayView({FTestColumnB::StaticStruct()})));
-								
-									TedsInterface->RemoveColumn(Row, FTestColumnB::StaticStruct());
+								RowHandle Row = CreatedRows[0];
+								TestTrue("Expected to have column about to be removed",
+									TedsInterface->HasColumns(Row, MakeArrayView({FTestColumnB::StaticStruct()})));
+							
+								TedsInterface->RemoveColumn(Row, FTestColumnB::StaticStruct());
 
-									TestFalse("Expected to no longer have removed column",
-										TedsInterface->HasColumns(Row, MakeArrayView({FTestColumnB::StaticStruct()})));
-								});
+								TestFalse("Expected to no longer have removed column",
+									TedsInterface->HasColumns(Row, MakeArrayView({FTestColumnB::StaticStruct()})));
+							});
 
-								It("should invoke registered data column observer when called", [this]()
-								{
-									int64 PreviousQueryCallCount = DataColumnB_RemoveObserverCallCount;
-									RowHandle Row = CreatedRows[0];
-								
-									TedsInterface->RemoveColumn(Row, FTestColumnB::StaticStruct());
-								
-									TestEqual("Expect observer to have been called one time", DataColumnB_RemoveObserverCallCount, PreviousQueryCallCount + 1);
-								});
-
-								It("should remove a single tag column when called", [this]()
-								{
-									RowHandle Row = CreatedRows[0];
-									TestTrue("Expected to have column about to be removed",
-										TedsInterface->HasColumns(Row, MakeArrayView({FTestTagColumnB::StaticStruct()})));
-														
-									TedsInterface->RemoveColumn(Row, FTestTagColumnB::StaticStruct());
+							It("should invoke registered data column observer when called", [this]()
+							{
+								int64 PreviousQueryCallCount = DataColumnB_RemoveObserverCallCount;
+								RowHandle Row = CreatedRows[0];
+							
+								TedsInterface->RemoveColumn(Row, FTestColumnB::StaticStruct());
+							
+								TestEqual("Expect observer to have been called one time", DataColumnB_RemoveObserverCallCount, PreviousQueryCallCount + 1);
+							});
 
-									TestFalse("Expected to no longer have removed column",
-										TedsInterface->HasColumns(Row, MakeArrayView({FTestTagColumnB::StaticStruct()})));
-								});
+							It("should remove a single tag column when called", [this]()
+							{
+								RowHandle Row = CreatedRows[0];
+								TestTrue("Expected to have column about to be removed",
+									TedsInterface->HasColumns(Row, MakeArrayView({FTestTagColumnB::StaticStruct()})));
+													
+								TedsInterface->RemoveColumn(Row, FTestTagColumnB::StaticStruct());
+
+								TestFalse("Expected to no longer have removed column",
+									TedsInterface->HasColumns(Row, MakeArrayView({FTestTagColumnB::StaticStruct()})));
+							});
 
-								It("should invoke registered tag column observer when called", [this]()
-								{
-									int64 PreviousQueryCallCount = TagColumnB_RemoveObserverCallCount;
-									RowHandle Row = CreatedRows[0];
-								
-									TedsInterface->RemoveColumn(Row, FTestTagColumnB::StaticStruct());
-								
-									TestEqual("Expect observer to have been called one time", TagColumnB_RemoveObserverCallCount, PreviousQueryCallCount + 1);
-								});
+							It("should invoke registered tag column observer when called", [this]()
+							{
+								int64 PreviousQueryCallCount = TagColumnB_RemoveObserverCallCount;
+								RowHandle Row = CreatedRows[0];
+							
+								TedsInterface->RemoveColumn(Row, FTestTagColumnB::StaticStruct());
+							
+								TestEqual("Expect observer to have been called one time", TagColumnB_RemoveObserverCallCount, PreviousQueryCallCount + 1);
 							});
+						});
 
-							Describe("AddColumns(RowHandle, TConstArrayView<const UScriptStruct*>)", [this]()
+						Describe("AddColumns(RowHandle, TConstArrayView<const UScriptStruct*>)", [this]()
+						{
+							It("should add a single column when called with single data column", [this]()
 							{
-								It("should add a single column when called with single data column", [this]()
-								{
-									TArray<const UScriptStruct*> Columns = {FTestColumnC::StaticStruct(), FTestColumnD::StaticStruct()};
-								
-									RowHandle Row = CreatedRows[0];
-									TestFalse("Expected to not have column about to be added",
-										TedsInterface->HasColumns(Row, Columns));
+								TArray<const UScriptStruct*> Columns = {FTestColumnC::StaticStruct(), FTestColumnD::StaticStruct()};
+							
+								RowHandle Row = CreatedRows[0];
+								TestFalse("Expected to not have column about to be added",
+									TedsInterface->HasColumns(Row, Columns));
 
-									TedsInterface->AddColumns(Row, Columns);
+								TedsInterface->AddColumns(Row, Columns);
 
-									TestTrue("Expected to have column about to be added", 
-									TedsInterface->HasColumns(Row, Columns));
-								});
+								TestTrue("Expected to have column about to be added", 
+								TedsInterface->HasColumns(Row, Columns));
+							});
 
-								It("should add multiple column when called with multiple data columns", [this]()
-								{
-									TArray<const UScriptStruct*> Columns = {FTestColumnC::StaticStruct(), FTestColumnD::StaticStruct()};
-								
-									RowHandle Row = CreatedRows[0];
-									TestFalse("Expected to not have column about to be added", 
+							It("should add multiple column when called with multiple data columns", [this]()
+							{
+								TArray<const UScriptStruct*> Columns = {FTestColumnC::StaticStruct(), FTestColumnD::StaticStruct()};
+							
+								RowHandle Row = CreatedRows[0];
+								TestFalse("Expected to not have column about to be added", 
+								TedsInterface->HasColumns(Row, Columns));
+
+								TedsInterface->AddColumns(Row, Columns);
+
+								TestTrue("Expected to have column about to be added",
 									TedsInterface->HasColumns(Row, Columns));
+							});
 
-									TedsInterface->AddColumns(Row, Columns);
+							It("should invoke registered data column observer when called", [this]()
+							{
+								TArray<const UScriptStruct*> Columns = {FTestColumnC::StaticStruct()};
+							
+								int64 PreviousQueryCallCount = DataColumnC_AddObserverCallCount;
+								RowHandle Row = CreatedRows[0];
+							
+								TedsInterface->AddColumns(Row, Columns);
 
-									TestTrue("Expected to have column about to be added",
-										TedsInterface->HasColumns(Row, Columns));
-								});
+								TestEqual("Expect observer to have been called one time", DataColumnC_AddObserverCallCount, PreviousQueryCallCount + 1);
+							});
 
-								It("should invoke registered data column observer when called", [this]()
-								{
-									TArray<const UScriptStruct*> Columns = {FTestColumnC::StaticStruct()};
-								
-									int64 PreviousQueryCallCount = DataColumnC_AddObserverCallCount;
-									RowHandle Row = CreatedRows[0];
-								
-									TedsInterface->AddColumns(Row, Columns);
+							It("should add a single column when called with single tag column", [this]()
+							{
+								TArray<const UScriptStruct*> Columns = {FTestTagColumnC::StaticStruct(), FTestTagColumnD::StaticStruct()};
+													
+								RowHandle Row = CreatedRows[0];
+								TestFalse("Expected to not have column about to be added",
+									TedsInterface->HasColumns(Row, Columns));
 
-									TestEqual("Expect observer to have been called one time", DataColumnC_AddObserverCallCount, PreviousQueryCallCount + 1);
-								});
+								TedsInterface->AddColumns(Row, Columns);
 
-								It("should add a single column when called with single tag column", [this]()
-								{
-									TArray<const UScriptStruct*> Columns = {FTestTagColumnC::StaticStruct(), FTestTagColumnD::StaticStruct()};
-														
-									RowHandle Row = CreatedRows[0];
-									TestFalse("Expected to not have column about to be added",
-										TedsInterface->HasColumns(Row, Columns));
+								TestTrue("Expected to have column about to be added", 
+								TedsInterface->HasColumns(Row, Columns));
+							});
 
-									TedsInterface->AddColumns(Row, Columns);
+							It("should add multiple column when called with multiple tag columns", [this]()
+							{
+								TArray<const UScriptStruct*> Columns = {FTestTagColumnC::StaticStruct(), FTestTagColumnD::StaticStruct()};
+													
+								RowHandle Row = CreatedRows[0];
+								TestFalse("Expected to not have column about to be added", 
+								TedsInterface->HasColumns(Row, Columns));
 
-									TestTrue("Expected to have column about to be added", 
-									TedsInterface->HasColumns(Row, Columns));
-								});
+								TedsInterface->AddColumns(Row, Columns);
 
-								It("should add multiple column when called with multiple tag columns", [this]()
-								{
-									TArray<const UScriptStruct*> Columns = {FTestTagColumnC::StaticStruct(), FTestTagColumnD::StaticStruct()};
-														
-									RowHandle Row = CreatedRows[0];
-									TestFalse("Expected to not have column about to be added", 
+								TestTrue("Expected to have column about to be added",
 									TedsInterface->HasColumns(Row, Columns));
+							});
 
-									TedsInterface->AddColumns(Row, Columns);
-
-									TestTrue("Expected to have column about to be added",
-										TedsInterface->HasColumns(Row, Columns));
-								});
+							It("should invoke registered tag column added observer when called", [this]()
+							{
+								TArray<const UScriptStruct*> Columns = {FTestTagColumnC::StaticStruct()};
+													
+								int64 PreviousQueryCallCount = TagColumnC_AddObserverCallCount;
+								RowHandle Row = CreatedRows[0];
+								TedsInterface->AddColumns(Row, Columns);
 
-								It("should invoke registered tag column added observer when called", [this]()
-								{
-									TArray<const UScriptStruct*> Columns = {FTestTagColumnC::StaticStruct()};
-														
-									int64 PreviousQueryCallCount = TagColumnC_AddObserverCallCount;
-									RowHandle Row = CreatedRows[0];
-									TedsInterface->AddColumns(Row, Columns);
-
-									TestEqual("Expect observer to have been called one time", TagColumnC_AddObserverCallCount, PreviousQueryCallCount + 1);
-								});
+								TestEqual("Expect observer to have been called one time", TagColumnC_AddObserverCallCount, PreviousQueryCallCount + 1);
 							});
+						});
 
-							Describe("RemoveColumns(RowHandle, TConstArrayView<const UScriptStruct*>)", [this]()
+						Describe("RemoveColumns(RowHandle, TConstArrayView<const UScriptStruct*>)", [this]()
+						{
+							It("should remove a single column when RemoveColumns called with single column", [this]()
 							{
-								It("should remove a single column when RemoveColumns called with single column", [this]()
-								{
-									TArray<const UScriptStruct*> Columns = {FTestColumnB::StaticStruct()};
-								
-									RowHandle Row = CreatedRows[0];
-									TestTrue("Expected to not column about to be removed",
-										TedsInterface->HasColumns(Row, Columns));
+								TArray<const UScriptStruct*> Columns = {FTestColumnB::StaticStruct()};
+							
+								RowHandle Row = CreatedRows[0];
+								TestTrue("Expected to not column about to be removed",
+									TedsInterface->HasColumns(Row, Columns));
 
-									TedsInterface->RemoveColumns(Row, Columns);
+								TedsInterface->RemoveColumns(Row, Columns);
 
-									TestFalse("Expected to not have column that was removed", 
-									TedsInterface->HasColumns(Row, MakeArrayView({FTestColumnB::StaticStruct()})));
-								});
+								TestFalse("Expected to not have column that was removed", 
+								TedsInterface->HasColumns(Row, MakeArrayView({FTestColumnB::StaticStruct()})));
+							});
 
-								It("should remove all columns when RemoveColumns called with multiple columns", [this]()
-								{
-									TArray<const UScriptStruct*> Columns = {FTestColumnA::StaticStruct(), FTestColumnB::StaticStruct()};
-								
-									RowHandle Row = CreatedRows[0];
-									TestTrue("Expected to have columns about to be removed", 
-									TedsInterface->HasColumns(Row, Columns));
+							It("should remove all columns when RemoveColumns called with multiple columns", [this]()
+							{
+								TArray<const UScriptStruct*> Columns = {FTestColumnA::StaticStruct(), FTestColumnB::StaticStruct()};
+							
+								RowHandle Row = CreatedRows[0];
+								TestTrue("Expected to have columns about to be removed", 
+								TedsInterface->HasColumns(Row, Columns));
 
-									TedsInterface->RemoveColumns(Row, Columns);
+								TedsInterface->RemoveColumns(Row, Columns);
 
-									TestFalse("Expected to have column about to be added",
-										TedsInterface->HasColumns(Row, Columns));
-								});
+								TestFalse("Expected to have column about to be added",
+									TedsInterface->HasColumns(Row, Columns));
+							});
 
-								It("should invoke registered observer when RemoveColumns called", [this]()
-								{
-									TArray<const UScriptStruct*> Columns = {FTestColumnB::StaticStruct()};
-								
-									int64 PreviousQueryCallCount = DataColumnB_RemoveObserverCallCount;
-									RowHandle Row = CreatedRows[0];
-									TedsInterface->RemoveColumns(Row, Columns);
-
-									TestEqual("Expect observer to have been called one time", DataColumnB_RemoveObserverCallCount, PreviousQueryCallCount + 1);
-								});
+							It("should invoke registered observer when RemoveColumns called", [this]()
+							{
+								TArray<const UScriptStruct*> Columns = {FTestColumnB::StaticStruct()};
+							
+								int64 PreviousQueryCallCount = DataColumnB_RemoveObserverCallCount;
+								RowHandle Row = CreatedRows[0];
+								TedsInterface->RemoveColumns(Row, Columns);
+
+								TestEqual("Expect observer to have been called one time", DataColumnB_RemoveObserverCallCount, PreviousQueryCallCount + 1);
 							});
 						});
 					});
-				
 				});
+			
 			});
-		}
-	} // namespace Tests
-} // namespace UE::Editor::DataStorage
+		});
+	}
+} // namespace UE::Editor::DataStorage::Tests
 
 #endif
\ No newline at end of file
diff --git a/Engine/Plugins/Experimental/EditorDataStorage/Source/TedsUI/Private/Processors/WidgetReferenceColumnUpdateProcessor.cpp b/Engine/Plugins/Experimental/EditorDataStorage/Source/TedsUI/Private/Processors/WidgetReferenceColumnUpdateProcessor.cpp
--- a/Engine/Plugins/Experimental/EditorDataStorage/Source/TedsUI/Private/Processors/WidgetReferenceColumnUpdateProcessor.cpp
+++ b/Engine/Plugins/Experimental/EditorDataStorage/Source/TedsUI/Private/Processors/WidgetReferenceColumnUpdateProcessor.cpp
@@ -5,13 +5,13 @@
 #include "Elements/Columns/TypedElementSlateWidgetColumns.h"
 #include "Elements/Framework/TypedElementQueryBuilder.h"
 
-void UWidgetReferenceColumnUpdateFactory::RegisterQueries(ITypedElementDataStorageInterface& DataStorage)
+void UWidgetReferenceColumnUpdateFactory::RegisterQueries(IEditorDataStorageProvider& DataStorage)
 {
 	RegisterDeleteRowOnWidgetDeleteQuery(DataStorage);
 	RegisterDeleteColumnOnWidgetDeleteQuery(DataStorage);
 }
 
-void UWidgetReferenceColumnUpdateFactory::RegisterDeleteRowOnWidgetDeleteQuery(ITypedElementDataStorageInterface& DataStorage) const
+void UWidgetReferenceColumnUpdateFactory::RegisterDeleteRowOnWidgetDeleteQuery(IEditorDataStorageProvider& DataStorage) const
 {
 	using namespace UE::Editor::DataStorage::Queries;
 
@@ -31,7 +31,7 @@ void UWidgetReferenceColumnUpdateFactory::RegisterDeleteRowOnWidgetDeleteQuery(I
     	.Compile());
 }
 
-void UWidgetReferenceColumnUpdateFactory::RegisterDeleteColumnOnWidgetDeleteQuery(ITypedElementDataStorageInterface& DataStorage) const
+void UWidgetReferenceColumnUpdateFactory::RegisterDeleteColumnOnWidgetDeleteQuery(IEditorDataStorageProvider& DataStorage) const
 {
 	using namespace UE::Editor::DataStorage::Queries;
 
diff --git a/Engine/Plugins/Experimental/EditorDataStorage/Source/TedsUI/Private/Processors/WidgetReferenceColumnUpdateProcessor.h b/Engine/Plugins/Experimental/EditorDataStorage/Source/TedsUI/Private/Processors/WidgetReferenceColumnUpdateProcessor.h
--- a/Engine/Plugins/Experimental/EditorDataStorage/Source/TedsUI/Private/Processors/WidgetReferenceColumnUpdateProcessor.h
+++ b/Engine/Plugins/Experimental/EditorDataStorage/Source/TedsUI/Private/Processors/WidgetReferenceColumnUpdateProcessor.h
@@ -13,16 +13,16 @@
  * the FTypedElementSlateWidgetReferenceDeletesRowTag was found.
  */
 UCLASS()
-class UWidgetReferenceColumnUpdateFactory : public UTypedElementDataStorageFactory
+class UWidgetReferenceColumnUpdateFactory : public UEditorDataStorageFactory
 {
 	GENERATED_BODY()
 
 public:
 	~UWidgetReferenceColumnUpdateFactory() override = default;
 
-	void RegisterQueries(ITypedElementDataStorageInterface& DataStorage) override;
+	void RegisterQueries(IEditorDataStorageProvider& DataStorage) override;
 
 private:
-	void RegisterDeleteRowOnWidgetDeleteQuery(ITypedElementDataStorageInterface& DataStorage) const;
-	void RegisterDeleteColumnOnWidgetDeleteQuery(ITypedElementDataStorageInterface& DataStorage) const;
+	void RegisterDeleteRowOnWidgetDeleteQuery(IEditorDataStorageProvider& DataStorage) const;
+	void RegisterDeleteColumnOnWidgetDeleteQuery(IEditorDataStorageProvider& DataStorage) const;
 };
diff --git a/Engine/Plugins/Experimental/EditorDataStorage/Source/TedsUI/Private/TedsUIModule.cpp b/Engine/Plugins/Experimental/EditorDataStorage/Source/TedsUI/Private/TedsUIModule.cpp
--- a/Engine/Plugins/Experimental/EditorDataStorage/Source/TedsUI/Private/TedsUIModule.cpp
+++ b/Engine/Plugins/Experimental/EditorDataStorage/Source/TedsUI/Private/TedsUIModule.cpp
@@ -10,7 +10,7 @@
 
 #define LOCTEXT_NAMESPACE "TedsUIModule"
 
-namespace UE::TypedElementsDataStorageUI::Private
+namespace UE::Editor::DataStorage::Private
 {
 	FAutoConsoleCommandWithOutputDevice PrintWidgetPurposesConsoleCommand(
 	TEXT("TEDS.UI.PrintWidgetPurposes"),
@@ -18,11 +18,11 @@ namespace UE::TypedElementsDataStorageUI::Private
 	FConsoleCommandWithOutputDeviceDelegate::CreateLambda([](FOutputDevice& Output)
 		{
 			UTypedElementRegistry* Registry = UTypedElementRegistry::GetInstance();
-			if (ITypedElementDataStorageUiInterface* UiStorage = Registry->GetMutableDataStorageUi())
+			if (IEditorDataStorageUiProvider* UiStorage = Registry->GetMutableDataStorageUi())
 			{
 				Output.Log(TEXT("The Typed Elements Data Storage has recorded the following widget purposes:"));
 				UiStorage->ListWidgetPurposes(
-					[&Output](FName Purpose, ITypedElementDataStorageUiInterface::EPurposeType, const FText& Description)
+					[&Output](FName Purpose, IEditorDataStorageUiProvider::EPurposeType, const FText& Description)
 					{
 						Output.Logf(TEXT("    %s - %s"), *Purpose.ToString(), *Description.ToString());
 					});
@@ -36,7 +36,7 @@ namespace UE::TypedElementsDataStorageUI::Private
 		bUseNewTEDSUIWidgets,
 		TEXT("If true, TEDS UI will use new attribute binding driven widgets (needs to be set at startup)")
 	);
-}
+} // namespace UE::Editor::DataStorage::Private
 
 void FTedsUIModule::StartupModule()
 {
diff --git a/Engine/Plugins/Experimental/EditorDataStorage/Source/TedsUI/Private/Widgets/CounterWidget.cpp b/Engine/Plugins/Experimental/EditorDataStorage/Source/TedsUI/Private/Widgets/CounterWidget.cpp
--- a/Engine/Plugins/Experimental/EditorDataStorage/Source/TedsUI/Private/Widgets/CounterWidget.cpp
+++ b/Engine/Plugins/Experimental/EditorDataStorage/Source/TedsUI/Private/Widgets/CounterWidget.cpp
@@ -43,7 +43,7 @@ UCounterWidgetFactory::UCounterWidgetFactory()
 	}
 }
 
-void UCounterWidgetFactory::RegisterQueries(ITypedElementDataStorageInterface& DataStorage)
+void UCounterWidgetFactory::RegisterQueries(IEditorDataStorageProvider& DataStorage)
 {
 	using namespace UE::Editor::DataStorage::Queries;
 	
@@ -77,14 +77,14 @@ void UCounterWidgetFactory::RegisterQueries(ITypedElementDataStorageInterface& D
 	.Compile());
 }
 
-void UCounterWidgetFactory::RegisterWidgetPurposes(ITypedElementDataStorageUiInterface& DataStorageUi) const
+void UCounterWidgetFactory::RegisterWidgetPurposes(IEditorDataStorageUiProvider& DataStorageUi) const
 {
-	DataStorageUi.RegisterWidgetPurpose(WigetPurpose, ITypedElementDataStorageUiInterface::EPurposeType::Generic,
+	DataStorageUi.RegisterWidgetPurpose(WigetPurpose, IEditorDataStorageUiProvider::EPurposeType::Generic,
 		LOCTEXT("ToolBarPurposeDescription", "Widgets added to the status bar at the bottom editor of the main editor window."));
 }
 
-void UCounterWidgetFactory::RegisterWidgetConstructors(ITypedElementDataStorageInterface& DataStorage,
-	ITypedElementDataStorageUiInterface& DataStorageUi) const
+void UCounterWidgetFactory::RegisterWidgetConstructors(IEditorDataStorageProvider& DataStorage,
+	IEditorDataStorageUiProvider& DataStorageUi) const
 {
 	using namespace UE::Editor::DataStorage::Queries;
 
@@ -115,7 +115,7 @@ void UCounterWidgetFactory::SetupMainWindowIntegrations(TSharedPtr<SWindow> Pare
 		checkf(Registry, TEXT(
 			"FTypedElementsDataStorageUiModule didn't find the UTypedElementRegistry during main window integration when it should be available."));
 
-		ITypedElementDataStorageUiInterface* UiInterface = Registry->GetMutableDataStorageUi();
+		IEditorDataStorageUiProvider* UiInterface = Registry->GetMutableDataStorageUi();
 		checkf(UiInterface, TEXT(
 			"FTypedElementsDataStorageUiModule tried to integrate with the main window before the "
 			"Typed Elements Data Storage UI interface is available."));
@@ -173,7 +173,7 @@ TSharedPtr<SWidget> FCounterWidgetConstructor::CreateWidget(const UE::Editor::Da
 		.Justification(ETextJustify::Center);
 }
 
-bool FCounterWidgetConstructor::SetColumns(ITypedElementDataStorageInterface* DataStorage, UE::Editor::DataStorage::RowHandle Row)
+bool FCounterWidgetConstructor::SetColumns(IEditorDataStorageProvider* DataStorage, UE::Editor::DataStorage::RowHandle Row)
 {
 	FCounterWidgetColumn* CounterColumn = DataStorage->GetColumn<FCounterWidgetColumn>(Row);
 	checkf(CounterColumn, TEXT("Added a new FCounterWidgetColumn to the Typed Elements Data Storage, but didn't get a valid pointer back."));
diff --git a/Engine/Plugins/Experimental/EditorDataStorage/Source/TedsUI/Private/Widgets/ExportedTextWidget.cpp b/Engine/Plugins/Experimental/EditorDataStorage/Source/TedsUI/Private/Widgets/ExportedTextWidget.cpp
--- a/Engine/Plugins/Experimental/EditorDataStorage/Source/TedsUI/Private/Widgets/ExportedTextWidget.cpp
+++ b/Engine/Plugins/Experimental/EditorDataStorage/Source/TedsUI/Private/Widgets/ExportedTextWidget.cpp
@@ -42,7 +42,7 @@ namespace UE::Editor
 			UpdateExportedTextWidget(FText::FromString(MoveTemp(Label)), Widget);
 		}
 		
-		static void UpdateExportedTextWidget(ITypedElementDataStorageInterface& DataStorage, FTypedElementSlateWidgetReferenceColumn& Widget,
+		static void UpdateExportedTextWidget(IEditorDataStorageProvider& DataStorage, FTypedElementSlateWidgetReferenceColumn& Widget,
 			const FTypedElementScriptStructTypeInfoColumn& TypeInfo, const FTypedElementRowReferenceColumn& ReferencedRow)
 		{
 			const UScriptStruct* StructType = TypeInfo.TypeInfo.Get();
@@ -52,7 +52,7 @@ namespace UE::Editor
 			}
 		}
 		
-		static DataStorage::RowHandle RegisterUpdateCallback(ITypedElementDataStorageInterface& DataStorage, const UScriptStruct* Target)
+		static DataStorage::RowHandle RegisterUpdateCallback(IEditorDataStorageProvider& DataStorage, const UScriptStruct* Target)
 		{
 			using namespace UE::Editor::DataStorage::Queries;
 
@@ -115,8 +115,8 @@ namespace UE::Editor
 	} // namespace DataStorage::UI::Private
 } // namespace UE::Editor
 
-void UExportedTextWidgetFactory::RegisterWidgetConstructors(ITypedElementDataStorageInterface& DataStorage,
-	ITypedElementDataStorageUiInterface& DataStorageUi) const
+void UExportedTextWidgetFactory::RegisterWidgetConstructors(IEditorDataStorageProvider& DataStorage,
+	IEditorDataStorageUiProvider& DataStorageUi) const
 {
 	DataStorageUi.RegisterWidgetFactory(FName(TEXT("General.Cell.Default")), FExportedTextWidgetConstructor::StaticStruct());
 }
@@ -145,7 +145,7 @@ const UE::Editor::DataStorage::Queries::FConditions* FExportedTextWidgetConstruc
 	return &MatchedColumn;
 }
 
-FString FExportedTextWidgetConstructor::CreateWidgetDisplayName(ITypedElementDataStorageInterface* DataStorage,
+FString FExportedTextWidgetConstructor::CreateWidgetDisplayName(IEditorDataStorageProvider* DataStorage,
 	UE::Editor::DataStorage::RowHandle Row) const
 {
 	if (FTypedElementScriptStructTypeInfoColumn* TypeInfoColumn = DataStorage->GetColumn<FTypedElementScriptStructTypeInfoColumn>(Row))
@@ -164,8 +164,8 @@ TSharedPtr<SWidget> FExportedTextWidgetConstructor::CreateWidget(const UE::Edito
 }
 
 bool FExportedTextWidgetConstructor::FinalizeWidget(
-	ITypedElementDataStorageInterface* DataStorage,
-	ITypedElementDataStorageUiInterface* DataStorageUi,
+	IEditorDataStorageProvider* DataStorage,
+	IEditorDataStorageUiProvider* DataStorageUi,
 	UE::Editor::DataStorage::RowHandle Row,
 	const TSharedPtr<SWidget>& Widget)
 {
diff --git a/Engine/Plugins/Experimental/EditorDataStorage/Source/TedsUI/Private/Widgets/GeneralWidgetRegistrationFactory.cpp b/Engine/Plugins/Experimental/EditorDataStorage/Source/TedsUI/Private/Widgets/GeneralWidgetRegistrationFactory.cpp
--- a/Engine/Plugins/Experimental/EditorDataStorage/Source/TedsUI/Private/Widgets/GeneralWidgetRegistrationFactory.cpp
+++ b/Engine/Plugins/Experimental/EditorDataStorage/Source/TedsUI/Private/Widgets/GeneralWidgetRegistrationFactory.cpp
@@ -12,21 +12,21 @@ const FName UGeneralWidgetRegistrationFactory::HeaderPurpose(TEXT("General.Heade
 const FName UGeneralWidgetRegistrationFactory::CellDefaultPurpose(TEXT("General.Cell.Default"));
 const FName UGeneralWidgetRegistrationFactory::HeaderDefaultPurpose(TEXT("General.Header.Default"));
 
-void UGeneralWidgetRegistrationFactory::RegisterWidgetPurposes(ITypedElementDataStorageUiInterface& DataStorageUi) const
+void UGeneralWidgetRegistrationFactory::RegisterWidgetPurposes(IEditorDataStorageUiProvider& DataStorageUi) const
 {
-	DataStorageUi.RegisterWidgetPurpose(CellPurpose, ITypedElementDataStorageUiInterface::EPurposeType::UniqueByNameAndColumn,
+	DataStorageUi.RegisterWidgetPurpose(CellPurpose, IEditorDataStorageUiProvider::EPurposeType::UniqueByNameAndColumn,
 		LOCTEXT("GeneralCellPurpose", "General purpose widgets that can be used as cells for specific columns or column combinations."));
 
-	DataStorageUi.RegisterWidgetPurpose(LargeCellPurpose, ITypedElementDataStorageUiInterface::EPurposeType::UniqueByNameAndColumn,
+	DataStorageUi.RegisterWidgetPurpose(LargeCellPurpose, IEditorDataStorageUiProvider::EPurposeType::UniqueByNameAndColumn,
 		LOCTEXT("GeneralLargeCellPurpose", "General purpose widgets that are specifically designed to be embedded in a space larger than a single cell."));
 
-	DataStorageUi.RegisterWidgetPurpose(HeaderPurpose, ITypedElementDataStorageUiInterface::EPurposeType::UniqueByNameAndColumn,
+	DataStorageUi.RegisterWidgetPurpose(HeaderPurpose, IEditorDataStorageUiProvider::EPurposeType::UniqueByNameAndColumn,
 		LOCTEXT("GeneralHeaderPurpose", "General purpose widget that can be used as a header."));
 
-	DataStorageUi.RegisterWidgetPurpose(HeaderDefaultPurpose, ITypedElementDataStorageUiInterface::EPurposeType::UniqueByName,
+	DataStorageUi.RegisterWidgetPurpose(HeaderDefaultPurpose, IEditorDataStorageUiProvider::EPurposeType::UniqueByName,
 		LOCTEXT("GeneralHeaderDefaultPurpose", "The default widget to use in headers if no other specialization is provided."));
 	
-	DataStorageUi.RegisterWidgetPurpose(CellDefaultPurpose, ITypedElementDataStorageUiInterface::EPurposeType::UniqueByName,
+	DataStorageUi.RegisterWidgetPurpose(CellDefaultPurpose, IEditorDataStorageUiProvider::EPurposeType::UniqueByName,
 		LOCTEXT("GeneralCellDefaultPurpose", "The default widget to use in cells if no other specialization is provided."));
 }
 
diff --git a/Engine/Plugins/Experimental/EditorDataStorage/Source/TedsUI/Private/Widgets/LabelWidget.cpp b/Engine/Plugins/Experimental/EditorDataStorage/Source/TedsUI/Private/Widgets/LabelWidget.cpp
--- a/Engine/Plugins/Experimental/EditorDataStorage/Source/TedsUI/Private/Widgets/LabelWidget.cpp
+++ b/Engine/Plugins/Experimental/EditorDataStorage/Source/TedsUI/Private/Widgets/LabelWidget.cpp
@@ -24,8 +24,8 @@
 // ULabelWidgetFactory
 //
 
-void ULabelWidgetFactory::RegisterWidgetConstructors(ITypedElementDataStorageInterface& DataStorage,
-	ITypedElementDataStorageUiInterface& DataStorageUi) const
+void ULabelWidgetFactory::RegisterWidgetConstructors(IEditorDataStorageProvider& DataStorage,
+	IEditorDataStorageUiProvider& DataStorageUi) const
 {
 	using namespace UE::Editor::DataStorage::Queries;
 
@@ -36,12 +36,12 @@ void ULabelWidgetFactory::RegisterWidgetConstructors(ITypedElementDataStorageInt
 		TColumn<FTypedElementLabelColumn>() || (TColumn<FTypedElementLabelColumn>() && TColumn<FTypedElementLabelHashColumn>()));
 }
 
-void ULabelWidgetFactory::RegisterWidgetPurposes(ITypedElementDataStorageUiInterface& DataStorageUi) const
+void ULabelWidgetFactory::RegisterWidgetPurposes(IEditorDataStorageUiProvider& DataStorageUi) const
 {
-	DataStorageUi.RegisterWidgetPurpose(FName(TEXT("General.RowLabel")), ITypedElementDataStorageUiInterface::EPurposeType::UniqueByNameAndColumn,
+	DataStorageUi.RegisterWidgetPurpose(FName(TEXT("General.RowLabel")), IEditorDataStorageUiProvider::EPurposeType::UniqueByNameAndColumn,
 	LOCTEXT("GeneralRowLabelPurpose", "Specific purpose to request a widget to display a user facing display name for a row."));
 	
-	DataStorageUi.RegisterWidgetPurpose(FName(TEXT("General.RowLabel.Default")), ITypedElementDataStorageUiInterface::EPurposeType::UniqueByName,
+	DataStorageUi.RegisterWidgetPurpose(FName(TEXT("General.RowLabel.Default")), IEditorDataStorageUiProvider::EPurposeType::UniqueByName,
 	LOCTEXT("GeneralRowLabelDefaultPurpose", "Default purpose to request a widget to display a user facing display name for a row."));
 }
 
@@ -64,8 +64,8 @@ TConstArrayView<const UScriptStruct*> FLabelWidgetConstructor::GetAdditionalColu
 	return Columns;
 }
 
-TSharedPtr<SWidget> FLabelWidgetConstructor::CreateWidget(ITypedElementDataStorageInterface* DataStorage,
-	ITypedElementDataStorageUiInterface* DataStorageUi, RowHandle TargetRow, RowHandle WidgetRow,
+TSharedPtr<SWidget> FLabelWidgetConstructor::CreateWidget(IEditorDataStorageProvider* DataStorage,
+	IEditorDataStorageUiProvider* DataStorageUi, RowHandle TargetRow, RowHandle WidgetRow,
 	const UE::Editor::DataStorage::FMetaDataView& Arguments)
 {
 	UE::Editor::DataStorage::FAttributeBinder Binder(TargetRow, DataStorage);
diff --git a/Engine/Plugins/Experimental/EditorDataStorage/Source/TedsUI/Private/Widgets/PackagePathWidget.cpp b/Engine/Plugins/Experimental/EditorDataStorage/Source/TedsUI/Private/Widgets/PackagePathWidget.cpp
--- a/Engine/Plugins/Experimental/EditorDataStorage/Source/TedsUI/Private/Widgets/PackagePathWidget.cpp
+++ b/Engine/Plugins/Experimental/EditorDataStorage/Source/TedsUI/Private/Widgets/PackagePathWidget.cpp
@@ -11,8 +11,8 @@
 // UPackagePathWidgetFactory
 //
 
-void UPackagePathWidgetFactory::RegisterWidgetConstructors(ITypedElementDataStorageInterface& DataStorage,
-	ITypedElementDataStorageUiInterface& DataStorageUi) const
+void UPackagePathWidgetFactory::RegisterWidgetConstructors(IEditorDataStorageProvider& DataStorage,
+	IEditorDataStorageUiProvider& DataStorageUi) const
 {
 	using namespace UE::Editor::DataStorage::Queries;
 
@@ -46,8 +46,8 @@ TSharedPtr<SWidget> FPackagePathWidgetConstructor::CreateWidget(const UE::Editor
 }
 
 bool FPackagePathWidgetConstructor::FinalizeWidget(
-	ITypedElementDataStorageInterface* DataStorage,
-	ITypedElementDataStorageUiInterface* DataStorageUi,
+	IEditorDataStorageProvider* DataStorage,
+	IEditorDataStorageUiProvider* DataStorageUi,
 	UE::Editor::DataStorage::RowHandle Row,
 	const TSharedPtr<SWidget>& Widget)
 {
@@ -78,8 +78,8 @@ FLoadedPackagePathWidgetConstructor::FLoadedPackagePathWidgetConstructor()
 }
 
 bool FLoadedPackagePathWidgetConstructor::FinalizeWidget(
-	ITypedElementDataStorageInterface* DataStorage,
-	ITypedElementDataStorageUiInterface* DataStorageUi,
+	IEditorDataStorageProvider* DataStorage,
+	IEditorDataStorageUiProvider* DataStorageUi,
 	UE::Editor::DataStorage::RowHandle Row,
 	const TSharedPtr<SWidget>& Widget)
 {
diff --git a/Engine/Plugins/Experimental/EditorDataStorage/Source/TedsUI/Private/Widgets/SlateColorWidget.cpp b/Engine/Plugins/Experimental/EditorDataStorage/Source/TedsUI/Private/Widgets/SlateColorWidget.cpp
--- a/Engine/Plugins/Experimental/EditorDataStorage/Source/TedsUI/Private/Widgets/SlateColorWidget.cpp
+++ b/Engine/Plugins/Experimental/EditorDataStorage/Source/TedsUI/Private/Widgets/SlateColorWidget.cpp
@@ -14,7 +14,7 @@
 namespace UE::SlateColorWidget::Private
 {
 	FReply SummonColorPicker(const FGeometry& Geometry, const FPointerEvent& PointerEvent,
-		ITypedElementDataStorageInterface* DataStorage, UE::Editor::DataStorage::RowHandle TargetRow)
+		IEditorDataStorageProvider* DataStorage, UE::Editor::DataStorage::RowHandle TargetRow)
 	{
 		if (PointerEvent.GetEffectingButton() != EKeys::LeftMouseButton)
 		{
@@ -40,8 +40,8 @@ namespace UE::SlateColorWidget::Private
 	}
 }
 
-void USlateColorWidgetFactory::RegisterWidgetConstructors(ITypedElementDataStorageInterface& DataStorage,
-														ITypedElementDataStorageUiInterface& DataStorageUi) const
+void USlateColorWidgetFactory::RegisterWidgetConstructors(IEditorDataStorageProvider& DataStorage,
+														IEditorDataStorageUiProvider& DataStorageUi) const
 {
 	using namespace UE::Editor::DataStorage::Queries;
 	DataStorageUi.RegisterWidgetFactory<FSlateColorWidgetConstructor>(TEXT("General.Cell"), TColumn<FSlateColorColumn>());
@@ -54,8 +54,8 @@ FSlateColorWidgetConstructor::FSlateColorWidgetConstructor()
 
 
 
-TSharedPtr<SWidget> FSlateColorWidgetConstructor::CreateWidget(ITypedElementDataStorageInterface* DataStorage,
-                                                               ITypedElementDataStorageUiInterface* DataStorageUi, RowHandle TargetRow, RowHandle WidgetRow,
+TSharedPtr<SWidget> FSlateColorWidgetConstructor::CreateWidget(IEditorDataStorageProvider* DataStorage,
+                                                               IEditorDataStorageUiProvider* DataStorageUi, RowHandle TargetRow, RowHandle WidgetRow,
                                                                const UE::Editor::DataStorage::FMetaDataView& Arguments)
 {
 	UE::Editor::DataStorage::FAttributeBinder Binder(TargetRow, DataStorage);
diff --git a/Engine/Plugins/Experimental/EditorDataStorage/Source/TedsUI/Private/Widgets/SlateColorWidget.h b/Engine/Plugins/Experimental/EditorDataStorage/Source/TedsUI/Private/Widgets/SlateColorWidget.h
--- a/Engine/Plugins/Experimental/EditorDataStorage/Source/TedsUI/Private/Widgets/SlateColorWidget.h
+++ b/Engine/Plugins/Experimental/EditorDataStorage/Source/TedsUI/Private/Widgets/SlateColorWidget.h
@@ -11,15 +11,15 @@ class ITypedElementDataStorageInterface;
 class UScriptStruct;
 
 UCLASS()
-class USlateColorWidgetFactory : public UTypedElementDataStorageFactory
+class USlateColorWidgetFactory : public UEditorDataStorageFactory
 {
 	GENERATED_BODY()
 
 public:
 	~USlateColorWidgetFactory() override = default;
 
-	void RegisterWidgetConstructors(ITypedElementDataStorageInterface& DataStorage,
-		ITypedElementDataStorageUiInterface& DataStorageUi) const override;
+	void RegisterWidgetConstructors(IEditorDataStorageProvider& DataStorage,
+		IEditorDataStorageUiProvider& DataStorageUi) const override;
 };
 
 // Widget to show and edit the color column in TEDS
@@ -33,8 +33,8 @@ public:
 	~FSlateColorWidgetConstructor() override = default;
 
 	virtual TSharedPtr<SWidget> CreateWidget(
-		ITypedElementDataStorageInterface* DataStorage,
-		ITypedElementDataStorageUiInterface* DataStorageUi,
+		IEditorDataStorageProvider* DataStorage,
+		IEditorDataStorageUiProvider* DataStorageUi,
 		RowHandle TargetRow,
 		RowHandle WidgetRow, 
 		const UE::Editor::DataStorage::FMetaDataView& Arguments) override;
diff --git a/Engine/Plugins/Experimental/EditorDataStorage/Source/TedsUI/Public/Widgets/CounterWidget.h b/Engine/Plugins/Experimental/EditorDataStorage/Source/TedsUI/Public/Widgets/CounterWidget.h
--- a/Engine/Plugins/Experimental/EditorDataStorage/Source/TedsUI/Public/Widgets/CounterWidget.h
+++ b/Engine/Plugins/Experimental/EditorDataStorage/Source/TedsUI/Public/Widgets/CounterWidget.h
@@ -9,12 +9,12 @@
 
 #include "CounterWidget.generated.h"
 
-class ITypedElementDataStorageInterface;
+class IEditorDataStorageProvider;
 class SWindow;
 class UScriptStruct;
 
 UCLASS()
-class TEDSUI_API UCounterWidgetFactory : public UTypedElementDataStorageFactory
+class TEDSUI_API UCounterWidgetFactory : public UEditorDataStorageFactory
 {
 	GENERATED_BODY()
 
@@ -24,10 +24,10 @@ public:
 	UCounterWidgetFactory();
 	~UCounterWidgetFactory() override = default;
 
-	void RegisterQueries(ITypedElementDataStorageInterface& DataStorage) override;
-	void RegisterWidgetPurposes(ITypedElementDataStorageUiInterface& DataStorageUi) const override;
-	void RegisterWidgetConstructors(ITypedElementDataStorageInterface& DataStorage,
-		ITypedElementDataStorageUiInterface& DataStorageUi) const override;
+	void RegisterQueries(IEditorDataStorageProvider& DataStorage) override;
+	void RegisterWidgetPurposes(IEditorDataStorageUiProvider& DataStorageUi) const override;
+	void RegisterWidgetConstructors(IEditorDataStorageProvider& DataStorage,
+		IEditorDataStorageUiProvider& DataStorageUi) const override;
 
 	static void EnableCounterWidgets();
 
@@ -61,7 +61,7 @@ public:
 
 protected:
 	TSharedPtr<SWidget> CreateWidget(const UE::Editor::DataStorage::FMetaDataView& Arguments) override;
-	bool SetColumns(ITypedElementDataStorageInterface* DataStorage, UE::Editor::DataStorage::RowHandle Row) override;
+	bool SetColumns(IEditorDataStorageProvider* DataStorage, UE::Editor::DataStorage::RowHandle Row) override;
 };
 
 USTRUCT(meta = (DisplayName = "Counter widget"))
diff --git a/Engine/Plugins/Experimental/EditorDataStorage/Source/TedsUI/Public/Widgets/ExportedTextWidget.h b/Engine/Plugins/Experimental/EditorDataStorage/Source/TedsUI/Public/Widgets/ExportedTextWidget.h
--- a/Engine/Plugins/Experimental/EditorDataStorage/Source/TedsUI/Public/Widgets/ExportedTextWidget.h
+++ b/Engine/Plugins/Experimental/EditorDataStorage/Source/TedsUI/Public/Widgets/ExportedTextWidget.h
@@ -8,19 +8,19 @@
 
 #include "ExportedTextWidget.generated.h"
 
-class ITypedElementDataStorageInterface;
+class IEditorDataStorageProvider;
 class UScriptStruct;
 
 UCLASS()
-class TEDSUI_API UExportedTextWidgetFactory : public UTypedElementDataStorageFactory
+class TEDSUI_API UExportedTextWidgetFactory : public UEditorDataStorageFactory
 {
 	GENERATED_BODY()
 
 public:
 	virtual ~UExportedTextWidgetFactory() override = default;
 
-	virtual void RegisterWidgetConstructors(ITypedElementDataStorageInterface& DataStorage,
-		ITypedElementDataStorageUiInterface& DataStorageUi) const override;
+	virtual void RegisterWidgetConstructors(IEditorDataStorageProvider& DataStorage,
+		IEditorDataStorageUiProvider& DataStorageUi) const override;
 
 	TSet<TWeakObjectPtr<const UScriptStruct>> RegisteredTypes;
 };
@@ -37,13 +37,13 @@ public:
 	virtual TConstArrayView<const UScriptStruct*> GetAdditionalColumnsList() const override;
 	virtual const UE::Editor::DataStorage::Queries::FConditions* GetQueryConditions() const override;
 	virtual FString CreateWidgetDisplayName(
-		ITypedElementDataStorageInterface* DataStorage, UE::Editor::DataStorage::RowHandle Row) const override;
+		IEditorDataStorageProvider* DataStorage, UE::Editor::DataStorage::RowHandle Row) const override;
 
 protected:
 	virtual TSharedPtr<SWidget> CreateWidget(const UE::Editor::DataStorage::FMetaDataView& Arguments) override;
 	virtual bool FinalizeWidget(
-		ITypedElementDataStorageInterface* DataStorage,
-		ITypedElementDataStorageUiInterface* DataStorageUi,
+		IEditorDataStorageProvider* DataStorage,
+		IEditorDataStorageUiProvider* DataStorageUi,
 		UE::Editor::DataStorage::RowHandle Row,
 		const TSharedPtr<SWidget>& Widget) override;
 
diff --git a/Engine/Plugins/Experimental/EditorDataStorage/Source/TedsUI/Public/Widgets/GeneralWidgetRegistrationFactory.h b/Engine/Plugins/Experimental/EditorDataStorage/Source/TedsUI/Public/Widgets/GeneralWidgetRegistrationFactory.h
--- a/Engine/Plugins/Experimental/EditorDataStorage/Source/TedsUI/Public/Widgets/GeneralWidgetRegistrationFactory.h
+++ b/Engine/Plugins/Experimental/EditorDataStorage/Source/TedsUI/Public/Widgets/GeneralWidgetRegistrationFactory.h
@@ -8,7 +8,7 @@
 #include "GeneralWidgetRegistrationFactory.generated.h"
 
 UCLASS()
-class TEDSUI_API UGeneralWidgetRegistrationFactory : public UTypedElementDataStorageFactory
+class TEDSUI_API UGeneralWidgetRegistrationFactory : public UEditorDataStorageFactory
 {
 	GENERATED_BODY()
 
@@ -21,5 +21,5 @@ public:
 
 	~UGeneralWidgetRegistrationFactory() override = default;
 
-	void RegisterWidgetPurposes(ITypedElementDataStorageUiInterface& DataStorageUi) const override;
+	void RegisterWidgetPurposes(IEditorDataStorageUiProvider& DataStorageUi) const override;
 };
diff --git a/Engine/Plugins/Experimental/EditorDataStorage/Source/TedsUI/Public/Widgets/LabelWidget.h b/Engine/Plugins/Experimental/EditorDataStorage/Source/TedsUI/Public/Widgets/LabelWidget.h
--- a/Engine/Plugins/Experimental/EditorDataStorage/Source/TedsUI/Public/Widgets/LabelWidget.h
+++ b/Engine/Plugins/Experimental/EditorDataStorage/Source/TedsUI/Public/Widgets/LabelWidget.h
@@ -9,20 +9,20 @@
 
 #include "LabelWidget.generated.h"
 
-class ITypedElementDataStorageInterface;
+class IEditorDataStorageProvider;
 class UScriptStruct;
 
 UCLASS()
-class ULabelWidgetFactory : public UTypedElementDataStorageFactory
+class ULabelWidgetFactory : public UEditorDataStorageFactory
 {
 	GENERATED_BODY()
 
 public:
 	~ULabelWidgetFactory() override = default;
 
-	TEDSUI_API void RegisterWidgetConstructors(ITypedElementDataStorageInterface& DataStorage,
-		ITypedElementDataStorageUiInterface& DataStorageUi) const override;
-	TEDSUI_API virtual void RegisterWidgetPurposes(ITypedElementDataStorageUiInterface& DataStorageUi) const override;
+	TEDSUI_API void RegisterWidgetConstructors(IEditorDataStorageProvider& DataStorage,
+		IEditorDataStorageUiProvider& DataStorageUi) const override;
+	TEDSUI_API virtual void RegisterWidgetPurposes(IEditorDataStorageUiProvider& DataStorageUi) const override;
 };
 
 USTRUCT()
@@ -37,8 +37,8 @@ public:
 	TEDSUI_API virtual TConstArrayView<const UScriptStruct*> GetAdditionalColumnsList() const override;
 
 	TEDSUI_API virtual TSharedPtr<SWidget> CreateWidget(
-		ITypedElementDataStorageInterface* DataStorage,
-		ITypedElementDataStorageUiInterface* DataStorageUi,
+		IEditorDataStorageProvider* DataStorage,
+		IEditorDataStorageUiProvider* DataStorageUi,
 		UE::Editor::DataStorage::RowHandle TargetRow,
 		UE::Editor::DataStorage::RowHandle WidgetRow, 
 		const UE::Editor::DataStorage::FMetaDataView& Arguments) override;
