[Teds] Fixing some extraneous strings discovered with find in files (i.e. not class names)

diff --git a/Engine/Plugins/Experimental/EditorDataStorage/Source/TedsCore/Private/GlobalLock.h b/Engine/Plugins/Experimental/EditorDataStorage/Source/TedsCore/Private/GlobalLock.h
--- a/Engine/Plugins/Experimental/EditorDataStorage/Source/TedsCore/Private/GlobalLock.h
+++ b/Engine/Plugins/Experimental/EditorDataStorage/Source/TedsCore/Private/GlobalLock.h
@@ -6,8 +6,6 @@
 #include "HAL/CriticalSection.h"
 #include "HAL/Platform.h"
 
-class UTypedElementDatabase;
-
 namespace UE::Editor::DataStorage
 {
 	enum class EGlobalLockStatus : uint8
@@ -46,7 +44,6 @@ namespace UE::Editor::DataStorage
 	 */
 	class FGlobalLock
 	{
-		friend class ::UTypedElementDatabase;
 	public:
 		static void SharedLock(EGlobalLockScope Scope);
 		static void ExclusiveLock(EGlobalLockScope Scope);
diff --git a/Engine/Plugins/Experimental/EditorDataStorage/Source/TedsCore/Private/TypedElementDatabase.cpp b/Engine/Plugins/Experimental/EditorDataStorage/Source/TedsCore/Private/TypedElementDatabase.cpp
--- a/Engine/Plugins/Experimental/EditorDataStorage/Source/TedsCore/Private/TypedElementDatabase.cpp
+++ b/Engine/Plugins/Experimental/EditorDataStorage/Source/TedsCore/Private/TypedElementDatabase.cpp
@@ -72,7 +72,7 @@ FAutoConsoleCommandWithOutputDevice PrintSupportedColumnsConsoleCommand(
 			Output.Log(TEXT("End of Typed Elements Data Storage supported column list."));
 		}));
 
-namespace TypedElementDatabasePrivate
+namespace UE::Editor::DataStorage::Private
 {
 	struct ColumnsToBitSetsResult
 	{
@@ -100,7 +100,7 @@ namespace TypedElementDatabasePrivate
 		}
 		return Result;
 	}
-}
+} // namespace UE::Editor::DataStorage::Private
 
 void UEditorDataStorage::Initialize()
 {
@@ -580,6 +580,7 @@ void* UEditorDataStorage::GetColumnData(RowHandle Row, const UScriptStruct* Colu
 
 void UEditorDataStorage::AddColumns(RowHandle Row, TConstArrayView<const UScriptStruct*> Columns)
 {
+	using namespace UE::Editor::DataStorage;
 	if (ActiveEditorEntityManager)
 	{
 		FMassEntityHandle Entity = FMassEntityHandle::FromNumber(Row);
@@ -587,11 +588,11 @@ void UEditorDataStorage::AddColumns(RowHandle Row, TConstArrayView<const UScript
 
 		FMassFragmentBitSet FragmentsToAdd;
 		FMassTagBitSet TagsToAdd;
-		if (TypedElementDatabasePrivate::ColumnsToBitSets(Columns, FragmentsToAdd, TagsToAdd).MustUpdate())
+		if (Private::ColumnsToBitSets(Columns, FragmentsToAdd, TagsToAdd).MustUpdate())
 		{
 			if (ActiveEditorEntityManager->IsEntityActive(Entity))
 			{
-				UE::Editor::DataStorage::Legacy::FCommandBuffer::Execute_AddColumnsCommand(*ActiveEditorEntityManager, Row, FragmentsToAdd, TagsToAdd);
+				Legacy::FCommandBuffer::Execute_AddColumnsCommand(*ActiveEditorEntityManager, Row, FragmentsToAdd, TagsToAdd);
 			}
 			else
 			{
@@ -630,6 +631,8 @@ void UEditorDataStorage::RemoveColumn(RowHandle Row, const UE::Editor::DataStora
 
 void UEditorDataStorage::RemoveColumns(RowHandle Row, TConstArrayView<const UScriptStruct*> Columns)
 {
+	using namespace UE::Editor::DataStorage;
+
 	FMassEntityHandle Entity = FMassEntityHandle::FromNumber(Row);
 	if (ActiveEditorEntityManager)
 	{
@@ -637,11 +640,11 @@ void UEditorDataStorage::RemoveColumns(RowHandle Row, TConstArrayView<const UScr
 
 		FMassFragmentBitSet FragmentsToRemove;
 		FMassTagBitSet TagsToRemove;
-		if (TypedElementDatabasePrivate::ColumnsToBitSets(Columns, FragmentsToRemove, TagsToRemove).MustUpdate())
+		if (Private::ColumnsToBitSets(Columns, FragmentsToRemove, TagsToRemove).MustUpdate())
 		{
 			if (ActiveEditorEntityManager->IsEntityActive(Entity))
 			{
-				UE::Editor::DataStorage::Legacy::FCommandBuffer::Execute_RemoveColumnsCommand(*ActiveEditorEntityManager, Row, FragmentsToRemove, TagsToRemove);
+				Legacy::FCommandBuffer::Execute_RemoveColumnsCommand(*ActiveEditorEntityManager, Row, FragmentsToRemove, TagsToRemove);
 			}
 			else
 			{
@@ -654,6 +657,7 @@ void UEditorDataStorage::RemoveColumns(RowHandle Row, TConstArrayView<const UScr
 void UEditorDataStorage::AddRemoveColumns(RowHandle Row,
 	TConstArrayView<const UScriptStruct*> ColumnsToAdd, TConstArrayView<const UScriptStruct*> ColumnsToRemove)
 {
+	using namespace UE::Editor::DataStorage;
 	if (ActiveEditorEntityManager)
 	{
 		FMassEntityHandle Entity = FMassEntityHandle::FromNumber(Row);
@@ -664,18 +668,18 @@ void UEditorDataStorage::AddRemoveColumns(RowHandle Row,
 		FMassTagBitSet TagsToRemove;
 		FMassFragmentBitSet FragmentsToRemove;
 
-		bool bMustAddColumns = TypedElementDatabasePrivate::ColumnsToBitSets(ColumnsToAdd, FragmentsToAdd, TagsToAdd).MustUpdate();
-		bool bMustRemoveColumns = TypedElementDatabasePrivate::ColumnsToBitSets(ColumnsToRemove, FragmentsToRemove, TagsToRemove).MustUpdate();
+		bool bMustAddColumns = Private::ColumnsToBitSets(ColumnsToAdd, FragmentsToAdd, TagsToAdd).MustUpdate();
+		bool bMustRemoveColumns = Private::ColumnsToBitSets(ColumnsToRemove, FragmentsToRemove, TagsToRemove).MustUpdate();
 		
 		if (ActiveEditorEntityManager->IsEntityActive(Entity))
 		{
 			if (bMustAddColumns)
 			{
-				UE::Editor::DataStorage::Legacy::FCommandBuffer::Execute_AddColumnsCommand(*ActiveEditorEntityManager, Row, FragmentsToAdd, TagsToAdd);
+				Legacy::FCommandBuffer::Execute_AddColumnsCommand(*ActiveEditorEntityManager, Row, FragmentsToAdd, TagsToAdd);
 			}
 			if (bMustRemoveColumns)
 			{
-				UE::Editor::DataStorage::Legacy::FCommandBuffer::Execute_RemoveColumnsCommand(*ActiveEditorEntityManager, Row, FragmentsToRemove, TagsToRemove);
+				Legacy::FCommandBuffer::Execute_RemoveColumnsCommand(*ActiveEditorEntityManager, Row, FragmentsToRemove, TagsToRemove);
 			}
 		}
 		else
@@ -703,10 +707,10 @@ void UEditorDataStorage::BatchAddRemoveColumns(TConstArrayView<RowHandle> Rows,
 		FMassTagBitSet TagsToAdd;
 		FMassTagBitSet TagsToRemove;
 
-		namespace TEDP = TypedElementDatabasePrivate;
+		using namespace UE::Editor::DataStorage;
 
-		TEDP::ColumnsToBitSetsResult AddResult = TEDP::ColumnsToBitSets(ColumnsToAdd, FragmentsToAdd, TagsToAdd);
-		TEDP::ColumnsToBitSetsResult RemoveResult = TEDP::ColumnsToBitSets(ColumnsToRemove, FragmentsToRemove, TagsToRemove);
+		Private::ColumnsToBitSetsResult AddResult = Private::ColumnsToBitSets(ColumnsToAdd, FragmentsToAdd, TagsToAdd);
+		Private::ColumnsToBitSetsResult RemoveResult = Private::ColumnsToBitSets(ColumnsToRemove, FragmentsToRemove, TagsToRemove);
 		
 		if (AddResult.MustUpdate() || RemoveResult.MustUpdate())
 		{
@@ -714,7 +718,7 @@ void UEditorDataStorage::BatchAddRemoveColumns(TConstArrayView<RowHandle> Rows,
 			using EntityArchetypeLookup = TMap<FMassArchetypeHandle, EntityHandleArray, TInlineSetAllocator<32>>;
 			using ArchetypeEntityArray = TArray<FMassArchetypeEntityCollection, TInlineAllocator<32>>;
 
-			UE::Editor::DataStorage::Legacy::FCommandBuffer& CommandBuffer = Environment->GetDirectDeferredCommands();
+			Legacy::FCommandBuffer& CommandBuffer = Environment->GetDirectDeferredCommands();
 			
 			// Sort rows (entities) into to matching table (archetype) bucket.
 			EntityArchetypeLookup LookupTable;
diff --git a/Engine/Plugins/Experimental/EditorDataStorage/Source/TedsCore/Private/TypedElementDatabaseCompatibility.h b/Engine/Plugins/Experimental/EditorDataStorage/Source/TedsCore/Private/TypedElementDatabaseCompatibility.h
--- a/Engine/Plugins/Experimental/EditorDataStorage/Source/TedsCore/Private/TypedElementDatabaseCompatibility.h
+++ b/Engine/Plugins/Experimental/EditorDataStorage/Source/TedsCore/Private/TypedElementDatabaseCompatibility.h
@@ -23,8 +23,6 @@ class IEditorDataStorageProvider;
 struct FMassActorManager;
 class UTypedElementMementoSystem;
 
-enum class ETypedElementDatabaseCompatibilityObjectType : uint8;
-
 namespace UE::Editor::DataStorage
 {
 	class FEnvironment;
diff --git a/Engine/Plugins/Experimental/EditorDataStorage/Source/TedsCore/Private/TypedElementsDataStorage.cpp b/Engine/Plugins/Experimental/EditorDataStorage/Source/TedsCore/Private/TypedElementsDataStorage.cpp
--- a/Engine/Plugins/Experimental/EditorDataStorage/Source/TedsCore/Private/TypedElementsDataStorage.cpp
+++ b/Engine/Plugins/Experimental/EditorDataStorage/Source/TedsCore/Private/TypedElementsDataStorage.cpp
@@ -16,7 +16,7 @@
 #include "TypedElementDataStorageSharedColumn.h"
 #include "UObject/UObjectGlobals.h"
 
-#define LOCTEXT_NAMESPACE "FTypedElementsDataStorageModule"
+#define LOCTEXT_NAMESPACE "FEditorDataStorageModule"
 
 // MASS uses CDO in a few places, making it a difficult to consistently register Type Element's Columns and Tags
 // as they may have not been set up to impersonate MASS' Fragments and Tags yet. There are currently no longer
@@ -66,7 +66,7 @@ void FEditorDataStorageModule::StartupModule()
 	}
 
 	// Load the dependent TypedElementFramework module (holding TypedElementRegistry) here so that it is guaranteed to be available in Shutdown
-	// and it is shutdown AFTER FTypedElementsDataStorageModule
+	// and it is shutdown AFTER FEditorDataStorageModule
 	FModuleManager::Get().LoadModule(TEXT("TypedElementFramework"));
 	
 	ImpersonateMassTagsAndFragments();
@@ -93,7 +93,7 @@ void FEditorDataStorageModule::StartupModule()
 				// Register the various DataStorage instances.
 				UTypedElementRegistry* Registry = UTypedElementRegistry::GetInstance();
 				checkf(Registry, TEXT(
-					"FTypedElementsDataStorageModule tried to register itself, but there was no Typed Element Registry to register to."));
+					"FEditorDataStorageModule tried to register itself, but there was no Typed Element Registry to register to."));
 				Registry->SetDataStorage(DataStorage.Get());
 				Registry->SetDataStorageCompatibility(DataStorageCompatibility.Get());
 				Registry->SetDataStorageUi(DataStorageUi.Get());
diff --git a/Engine/Plugins/Experimental/EditorDataStorage/Source/TedsCore/Tests/TypedElementDataStorageInterfaceTests.cpp b/Engine/Plugins/Experimental/EditorDataStorage/Source/TedsCore/Tests/TypedElementDataStorageInterfaceTests.cpp
--- a/Engine/Plugins/Experimental/EditorDataStorage/Source/TedsCore/Tests/TypedElementDataStorageInterfaceTests.cpp
+++ b/Engine/Plugins/Experimental/EditorDataStorage/Source/TedsCore/Tests/TypedElementDataStorageInterfaceTests.cpp
@@ -181,7 +181,7 @@ namespace UE::Editor::DataStorage::Tests
 					}
 				});
 
-				Describe("BatchAddRow(TableHandle, int32, TypedElementDataStorageCreationCallbackRef)", [this]()
+				Describe("BatchAddRow(TableHandle, int32, EditorDataStorageCreationCallbackRef)", [this]()
 				{
 					for (int32 RowCount = 1; RowCount <= 2; ++RowCount)
 					{
@@ -287,7 +287,7 @@ namespace UE::Editor::DataStorage::Tests
 						}
 					});
 
-					Describe("BatchAddRow(TableHandle, int32, TypedElementDataStorageCreationCallbackRef)", [this]()
+					Describe("BatchAddRow(TableHandle, int32, EditorDataStorageCreationCallbackRef)", [this]()
 					{
 						for (int32 RowCount = 1; RowCount <= 2; ++RowCount)
 						{
diff --git a/Engine/Plugins/Experimental/EditorDataStorage/Source/TedsUI/Private/Widgets/CounterWidget.cpp b/Engine/Plugins/Experimental/EditorDataStorage/Source/TedsUI/Private/Widgets/CounterWidget.cpp
--- a/Engine/Plugins/Experimental/EditorDataStorage/Source/TedsUI/Private/Widgets/CounterWidget.cpp
+++ b/Engine/Plugins/Experimental/EditorDataStorage/Source/TedsUI/Private/Widgets/CounterWidget.cpp
@@ -113,12 +113,12 @@ void UCounterWidgetFactory::SetupMainWindowIntegrations(TSharedPtr<SWindow> Pare
 	{
 		UTypedElementRegistry* Registry = UTypedElementRegistry::GetInstance();
 		checkf(Registry, TEXT(
-			"FTypedElementsDataStorageUiModule didn't find the UTypedElementRegistry during main window integration when it should be available."));
+			"FEditorDataStorageUiModule didn't find the UTypedElementRegistry during main window integration when it should be available."));
 
 		IEditorDataStorageUiProvider* UiInterface = Registry->GetMutableDataStorageUi();
 		checkf(UiInterface, TEXT(
-			"FTypedElementsDataStorageUiModule tried to integrate with the main window before the "
-			"Typed Elements Data Storage UI interface is available."));
+			"FEditorDataStorageUiModule tried to integrate with the main window before the "
+			"TEDS UI interface is available."));
 
 		UToolMenu* Menu = UToolMenus::Get()->ExtendMenu(WigetPurpose);
 
diff --git a/Engine/Plugins/Experimental/EditorDataStorage/Source/TedsUI/Private/Widgets/SlateColorWidget.h b/Engine/Plugins/Experimental/EditorDataStorage/Source/TedsUI/Private/Widgets/SlateColorWidget.h
--- a/Engine/Plugins/Experimental/EditorDataStorage/Source/TedsUI/Private/Widgets/SlateColorWidget.h
+++ b/Engine/Plugins/Experimental/EditorDataStorage/Source/TedsUI/Private/Widgets/SlateColorWidget.h
@@ -7,7 +7,8 @@
 
 #include "SlateColorWidget.generated.h"
 
-class ITypedElementDataStorageInterface;
+class IEditorDataStorageProvider;
+class IEditorDataStorageUiProvider;
 class UScriptStruct;
 
 UCLASS()
diff --git a/Engine/Plugins/Experimental/EditorDataStorageFeatures/Source/TedsDebugger/Private/QueryEditor/TedsQueryEditorModel.cpp b/Engine/Plugins/Experimental/EditorDataStorageFeatures/Source/TedsDebugger/Private/QueryEditor/TedsQueryEditorModel.cpp
--- a/Engine/Plugins/Experimental/EditorDataStorageFeatures/Source/TedsDebugger/Private/QueryEditor/TedsQueryEditorModel.cpp
+++ b/Engine/Plugins/Experimental/EditorDataStorageFeatures/Source/TedsDebugger/Private/QueryEditor/TedsQueryEditorModel.cpp
@@ -23,8 +23,8 @@ namespace UE::Editor::DataStorage::Debug::QueryEditor
 		Id = INDEX_NONE;
 	}
 
-	FTedsQueryEditorModel::FTedsQueryEditorModel(IEditorDataStorageProvider& InTypedElementDataStorage)
-		: TypedElementDataStorage(InTypedElementDataStorage)
+	FTedsQueryEditorModel::FTedsQueryEditorModel(IEditorDataStorageProvider& InDataStorageProvider)
+		: EditorDataStorageProvider(InDataStorageProvider)
 	{
 	}
 
@@ -93,12 +93,12 @@ namespace UE::Editor::DataStorage::Debug::QueryEditor
 
 	IEditorDataStorageProvider& FTedsQueryEditorModel::GetTedsInterface()
 	{
-		return TypedElementDataStorage;
+		return EditorDataStorageProvider;
 	}
 	
 	const IEditorDataStorageProvider& FTedsQueryEditorModel::GetTedsInterface() const
 	{
-		return TypedElementDataStorage;
+		return EditorDataStorageProvider;
 	}
 
 	FQueryDescription FTedsQueryEditorModel::GenerateQueryDescription()
diff --git a/Engine/Plugins/Experimental/EditorDataStorageFeatures/Source/TedsDebugger/Private/QueryEditor/TedsQueryEditorModel.h b/Engine/Plugins/Experimental/EditorDataStorageFeatures/Source/TedsDebugger/Private/QueryEditor/TedsQueryEditorModel.h
--- a/Engine/Plugins/Experimental/EditorDataStorageFeatures/Source/TedsDebugger/Private/QueryEditor/TedsQueryEditorModel.h
+++ b/Engine/Plugins/Experimental/EditorDataStorageFeatures/Source/TedsDebugger/Private/QueryEditor/TedsQueryEditorModel.h
@@ -48,7 +48,7 @@ namespace UE::Editor::DataStorage::Debug::QueryEditor
 	class FTedsQueryEditorModel
 	{
 	public:
-		explicit FTedsQueryEditorModel(IEditorDataStorageProvider& InTypedElementDataStorage);
+		explicit FTedsQueryEditorModel(IEditorDataStorageProvider& InDataStorageProvider);
 
 		void Reset();
 		IEditorDataStorageProvider& GetTedsInterface();
@@ -105,7 +105,7 @@ namespace UE::Editor::DataStorage::Debug::QueryEditor
 		FConditionEntryInternal* FindEntryByHandle(FConditionEntryHandle Handle);
 		const FConditionEntryInternal* FindEntryByHandle(FConditionEntryHandle Handle) const;
 
-		IEditorDataStorageProvider& TypedElementDataStorage;
+		IEditorDataStorageProvider& EditorDataStorageProvider;
 
 		int32 IdGenerator = 0;
 
diff --git a/Engine/Plugins/Experimental/EditorDataStorageFeatures/Source/TedsStyling/Private/Widgets/SlateBrushWidget.h b/Engine/Plugins/Experimental/EditorDataStorageFeatures/Source/TedsStyling/Private/Widgets/SlateBrushWidget.h
--- a/Engine/Plugins/Experimental/EditorDataStorageFeatures/Source/TedsStyling/Private/Widgets/SlateBrushWidget.h
+++ b/Engine/Plugins/Experimental/EditorDataStorageFeatures/Source/TedsStyling/Private/Widgets/SlateBrushWidget.h
@@ -7,7 +7,8 @@
 
 #include "SlateBrushWidget.generated.h"
 
-class ITypedElementDataStorageInterface;
+class IEditorDataStorageProvider;
+class IEditorDataStorageUiProvider;
 class UScriptStruct;
 
 UCLASS()
diff --git a/Engine/Source/Runtime/TypedElementFramework/Public/Elements/Common/TypedElementHandles.h b/Engine/Source/Runtime/TypedElementFramework/Public/Elements/Common/TypedElementHandles.h
--- a/Engine/Source/Runtime/TypedElementFramework/Public/Elements/Common/TypedElementHandles.h
+++ b/Engine/Source/Runtime/TypedElementFramework/Public/Elements/Common/TypedElementHandles.h
@@ -17,12 +17,12 @@ namespace UE::Editor::DataStorage
 
 	using QueryHandle = uint64;
 	static constexpr QueryHandle InvalidQueryHandle = TNumericLimits<QueryHandle>::Max();
-} // namespace TypedElementDataStorage
+} // namespace UE::Editor::DataStorage
 
 /*
- * FTedsRowHandle is a strongly typed wrapper around TypedElementDataStorage::RowHandle and should only be used in cases where you need the extra info.
+ * FTedsRowHandle is a strongly typed wrapper around UE::Editor::DataStorage::RowHandle and should only be used in cases where you need the extra info.
  * E.g for reflection/UHT or for template specializing something that needs to know the semantics of the row handle.
- * For all other cases, you should use the regular typedef TypedElementDataStorage::RowHandle
+ * For all other cases, you should use the regular typedef UE::Editor::DataStorage::RowHandle
  */
 USTRUCT()
 struct FTedsRowHandle
diff --git a/Engine/Source/Runtime/TypedElementFramework/Public/Elements/Interfaces/TypedElementDataStorageInterface.h b/Engine/Source/Runtime/TypedElementFramework/Public/Elements/Interfaces/TypedElementDataStorageInterface.h
--- a/Engine/Source/Runtime/TypedElementFramework/Public/Elements/Interfaces/TypedElementDataStorageInterface.h
+++ b/Engine/Source/Runtime/TypedElementFramework/Public/Elements/Interfaces/TypedElementDataStorageInterface.h
@@ -308,14 +308,14 @@ public:
 	/**
 	 * Finds the type information for a dynamic column.
 	 * If the dynamic column has not been generated, then return nullptr
-	 * The TemplateType may be a typed derived from either FTypedElementDataStorageColumn or FTypedElementDataStorageTag, anything else will return nullptr
+	 * The TemplateType may be a typed derived from either UE::Editor::DataStorage::FColumn or UE::Editor::DataStorage::FTag, anything else will return nullptr
 	 */
 	virtual const UScriptStruct* FindDynamicColumn(const UE::Editor::DataStorage::FDynamicColumnDescription& Description) const = 0;
 
 	/**
 	 * Generates a new dynamic column from a Template.  A dynamic column is uniquely identified using the given template and an Identifier
 	 * This function is idempotent - multiple calls with the same parameters will result in subsequent calls returning the same type
-	 * The TemplateType may be a typed derived from either FTypedElementDataStorageColumn or FTypedElementDataStorageTag
+	 * The TemplateType may be a typed derived from either UE::Editor::DataStorage::FColumn or UE::Editor::DataStorage::FTag
 	 */
 	virtual const UScriptStruct* GenerateDynamicColumn(const UE::Editor::DataStorage::FDynamicColumnDescription& Description) = 0;
 	
