Print more information from TEDS when a processor is executed with an unregistered query handle

#rnx

diff --git a/Engine/Plugins/Experimental/EditorDataStorage/Source/TedsCore/Private/Processors/TypedElementProcessorAdaptors.cpp b/Engine/Plugins/Experimental/EditorDataStorage/Source/TedsCore/Private/Processors/TypedElementProcessorAdaptors.cpp
--- a/Engine/Plugins/Experimental/EditorDataStorage/Source/TedsCore/Private/Processors/TypedElementProcessorAdaptors.cpp
+++ b/Engine/Plugins/Experimental/EditorDataStorage/Source/TedsCore/Private/Processors/TypedElementProcessorAdaptors.cpp
@@ -1168,7 +1168,7 @@ void FTypedElementQueryProcessorData::Execute(FMassEntityManager& EntityManager,
 	
 	FExtendedQuery* StoredQuery = QueryStore->GetMutable(ParentQuery);
 
-	checkf(StoredQuery, TEXT("A query callback was registered for execution without an associated query."));
+	checkf(StoredQuery, TEXT("A query callback was registered for execution without an associated query. Processor: [%s]"), *GetProcessorName());
 	
 	IEditorDataStorageProvider::FQueryDescription& Description = StoredQuery->Description;
 	if (Description.Callback.ActivationCount > 0)
