TEDS UI: Misc follow up fixes to the change added to support dynamic columns in TEDS UI

diff --git a/Engine/Plugins/Experimental/EditorDataStorage/Source/TedsCore/Private/TypedElementDatabaseUI.cpp b/Engine/Plugins/Experimental/EditorDataStorage/Source/TedsCore/Private/TypedElementDatabaseUI.cpp
--- a/Engine/Plugins/Experimental/EditorDataStorage/Source/TedsCore/Private/TypedElementDatabaseUI.cpp
+++ b/Engine/Plugins/Experimental/EditorDataStorage/Source/TedsCore/Private/TypedElementDatabaseUI.cpp
@@ -473,7 +473,9 @@ void UEditorDataStorageUi::CreateWidgetConstructors_LongestMatch(const TArray<FW
 	TArray<TWeakObjectPtr<const UScriptStruct>> MatchedColumns;
 	for (auto FactoryIt = WidgetFactories.CreateConstIterator(); FactoryIt && !Columns.IsEmpty(); ++FactoryIt)
 	{
-		if (FactoryIt->GetConditions(Storage).MinimumColumnMatchRequired() > Columns.Num())
+		const UE::Editor::DataStorage::Queries::FConditions& Conditions = FactoryIt->GetConditions(Storage);
+		
+		if (Conditions.MinimumColumnMatchRequired() > Columns.Num())
 		{
 			// There are more columns required for this factory than there are in the requested columns list so skip this
 			// factory.
@@ -482,7 +484,7 @@ void UEditorDataStorageUi::CreateWidgetConstructors_LongestMatch(const TArray<FW
 
 		MatchedColumns.Reset();
 		
-		if (FactoryIt->GetConditions(Storage).Verify(MatchedColumns, Columns, true))
+		if (Conditions.Verify(MatchedColumns, Columns, true))
 		{
 			// Remove the found columns from the requested list.
 			Algo::SortBy(MatchedColumns, [](const TWeakObjectPtr<const UScriptStruct>& Column) { return Column.Get(); });
@@ -508,7 +510,7 @@ void UEditorDataStorageUi::CreateWidgetConstructors_LongestMatch(const TArray<FW
 				--ColumnsEnd;
 			}
 			
-			if (!CreateSingleWidgetConstructor(FactoryIt->Constructor, Arguments, MoveTemp(MatchedColumns), FactoryIt->GetConditions(Storage), Callback))
+			if (!CreateSingleWidgetConstructor(FactoryIt->Constructor, Arguments, MoveTemp(MatchedColumns), Conditions, Callback))
 			{
 				return;
 			}
@@ -524,23 +526,25 @@ void UEditorDataStorageUi::CreateWidgetConstructors_ExactMatch(const TArray<FWid
 	TArray<TWeakObjectPtr<const UScriptStruct>> MatchedColumns;
 	for (const FWidgetFactory& Factory : WidgetFactories)
 	{
+		const UE::Editor::DataStorage::Queries::FConditions& Conditions = Factory.GetConditions(Storage);
+		
 		// If there are more matches required that there are columns, then there will never be an exact match.
 		// Less than the column count can still result in a match that covers all columns.
-		if (Factory.GetConditions(Storage).MinimumColumnMatchRequired() > ColumnCount)
+		if (Conditions.MinimumColumnMatchRequired() > ColumnCount)
 		{
 			continue;
 		}
 
 		MatchedColumns.Reset();
 
-		if (Factory.GetConditions(Storage).Verify(MatchedColumns, Columns, true))
+		if (Conditions.Verify(MatchedColumns, Columns, true))
 		{
 			Algo::SortBy(MatchedColumns, [](const TWeakObjectPtr<const UScriptStruct>& Column) { return Column.Get(); });
 			MatchedColumns.SetNum(Algo::Unique(MatchedColumns), EAllowShrinking::No);
 			if (MatchedColumns.Num() == Columns.Num())
 			{
 				Columns.Reset();
-				CreateSingleWidgetConstructor(Factory.Constructor, Arguments, MoveTemp(MatchedColumns), Factory.GetConditions(Storage), Callback);
+				CreateSingleWidgetConstructor(Factory.Constructor, Arguments, MoveTemp(MatchedColumns), Conditions, Callback);
 				return;
 			}
 		}
@@ -559,7 +563,9 @@ void UEditorDataStorageUi::CreateWidgetConstructors_SingleMatch(const TArray<FWi
 	{
 		for (; FactoryIt != FactoryEnd; ++FactoryIt)
 		{
-			TConstArrayView<TWeakObjectPtr<const UScriptStruct>> ColumnData = (*FactoryIt).GetConditions(Storage).GetColumns();
+			const UE::Editor::DataStorage::Queries::FConditions& Conditions = (*FactoryIt).GetConditions(Storage);
+			
+			TConstArrayView<TWeakObjectPtr<const UScriptStruct>> ColumnData = Conditions.GetColumns();
 			if (ColumnData.Num() > 1)
 			{
 				// Moved passed the point where factories only have a single column.
@@ -575,7 +581,7 @@ void UEditorDataStorageUi::CreateWidgetConstructors_SingleMatch(const TArray<FWi
 			{
 				Columns.RemoveAt(ColumnIndex);
 				CreateSingleWidgetConstructor((*FactoryIt).Constructor, Arguments, 
-					TArray<TWeakObjectPtr<const UScriptStruct>>(ColumnData), (*FactoryIt).GetConditions(Storage), Callback);
+					TArray<TWeakObjectPtr<const UScriptStruct>>(ColumnData), Conditions, Callback);
 				// Match was found so move on to the next column in the column.
 				break;
 			}
@@ -603,15 +609,15 @@ UEditorDataStorageUi::FWidgetFactory::FWidgetFactory(TUniquePtr<FTypedElementWid
 
 UEditorDataStorageUi::FWidgetFactory::FWidgetFactory(const UScriptStruct* InConstructor, 
 	UE::Editor::DataStorage::Queries::FConditions&& InColumns)
-	: Columns(MoveTemp(InColumns))
-	, Constructor(InConstructor)
+	: Constructor(InConstructor)
+	, Columns(MoveTemp(InColumns))
 {
 }
 
 UEditorDataStorageUi::FWidgetFactory::FWidgetFactory(TUniquePtr<FTypedElementWidgetConstructor>&& InConstructor, 
 	UE::Editor::DataStorage::Queries::FConditions&& InColumns)
-	: Columns(MoveTemp(InColumns))
-	, Constructor(MoveTemp(InConstructor))
+	: Constructor(MoveTemp(InConstructor))
+	, Columns(MoveTemp(InColumns))
 {
 	checkf(std::get<TUniquePtr<FTypedElementWidgetConstructor>>(Constructor)->GetTypeInfo(),
 		TEXT("Widget constructor registered that didn't contain valid type information."));
diff --git a/Engine/Source/Runtime/TypedElementFramework/Public/Elements/Common/TypedElementQueryConditions.h b/Engine/Source/Runtime/TypedElementFramework/Public/Elements/Common/TypedElementQueryConditions.h
--- a/Engine/Source/Runtime/TypedElementFramework/Public/Elements/Common/TypedElementQueryConditions.h
+++ b/Engine/Source/Runtime/TypedElementFramework/Public/Elements/Common/TypedElementQueryConditions.h
@@ -28,7 +28,7 @@ namespace UE::Editor::DataStorage::Queries
 	struct TColumn final : public FColumnBase
 	{
 		template <typename U = T> requires (!std::is_same_v<U, void>)
-			constexpr TColumn(const FName& Identifier = NAME_None) : FColumnBase(T::StaticStruct(), Identifier) {}
+			constexpr explicit TColumn(const FName& Identifier = NAME_None) : FColumnBase(T::StaticStruct(), Identifier) {}
 		
 		template <typename U = T> requires (std::is_same_v<U, void>)
 			constexpr explicit TColumn(TWeakObjectPtr<const UScriptStruct> ColumnTypeInfo, const FName& Identifier = NAME_None)
diff --git a/Engine/Source/Runtime/TypedElementFramework/Public/Elements/Framework/TypedElementAttributeBinding.inl b/Engine/Source/Runtime/TypedElementFramework/Public/Elements/Framework/TypedElementAttributeBinding.inl
--- a/Engine/Source/Runtime/TypedElementFramework/Public/Elements/Framework/TypedElementAttributeBinding.inl
+++ b/Engine/Source/Runtime/TypedElementFramework/Public/Elements/Framework/TypedElementAttributeBinding.inl
@@ -17,9 +17,12 @@ namespace UE::Editor::DataStorage
 				{
 					return DynamicColumn;
 				}
+
+				// The dynamic column wasn't found on this row
+				return nullptr;
 			}
 
-			// If we weren't provided a dynamic column identifier, or the dynamic column could not be found simply fallback to the non-dynamic column
+			// If we weren't provided a dynamic column identifier, we are looking for the regular non-dynamic column
 			return DataStorage->GetColumn<ColumnType>(Row);
 		}
 		
