Fixed a crash when the TEDS Compat command buffer has exactly one command.

diff --git a/Engine/Plugins/Experimental/EditorDataStorage/Source/TedsCore/Private/Commands/EditorDataStorageCommandBuffer.inl b/Engine/Plugins/Experimental/EditorDataStorage/Source/TedsCore/Private/Commands/EditorDataStorageCommandBuffer.inl
--- a/Engine/Plugins/Experimental/EditorDataStorage/Source/TedsCore/Private/Commands/EditorDataStorageCommandBuffer.inl
+++ b/Engine/Plugins/Experimental/EditorDataStorage/Source/TedsCore/Private/Commands/EditorDataStorageCommandBuffer.inl
@@ -150,7 +150,7 @@ namespace UE::Editor::DataStorage
 			Left = MoveToNextNonNop(0);
 			Right = MoveToNextNonNop(Left);
 		}
-		else if (Commands.CommandReferences[1]->template IsType<FNopCommand>())
+		else if (Commands.CommandReferences.IsValidIndex(1) && Commands.CommandReferences[1]->template IsType<FNopCommand>())
 		{
 			Right = MoveToNextNonNop(1);
 		}
@@ -162,12 +162,12 @@ namespace UE::Editor::DataStorage
 		, Right(InLeft + 1)
 		, Commands(InCommands)
 	{
-		if (!Commands.CommandReferences.IsEmpty() && Commands.CommandReferences[Left]->template IsType<FNopCommand>())
+		if (Commands.CommandReferences.IsValidIndex(Left) && Commands.CommandReferences[Left]->template IsType<FNopCommand>())
 		{
 			Left = MoveToNextNonNop(Left);
 			Right = MoveToNextNonNop(Left);
 		}
-		else if (Commands.CommandReferences[Right]->template IsType<FNopCommand>())
+		else if (Commands.CommandReferences.IsValidIndex(Right) && Commands.CommandReferences[Right]->template IsType<FNopCommand>())
 		{
 			Right = MoveToNextNonNop(Right);
 		}
@@ -180,12 +180,12 @@ namespace UE::Editor::DataStorage
 		, Right(InLeft + 1)
 		, Commands(InCommands)
 	{
-		if (!Commands.CommandReferences.IsEmpty() && Commands.CommandReferences[Left]->template IsType<FNopCommand>())
+		if (Commands.CommandReferences.IsValidIndex(Left) && Commands.CommandReferences[Left]->template IsType<FNopCommand>())
 		{
 			Left = MoveToNextNonNop(Left);
 			Right = MoveToNextNonNop(Left);
 		}
-		else if (Commands.CommandReferences[Right]->template IsType<FNopCommand>())
+		else if (Commands.CommandReferences.IsValidIndex(Right) && Commands.CommandReferences[Right]->template IsType<FNopCommand>())
 		{
 			Right = MoveToNextNonNop(Right);
 		}
diff --git a/Engine/Plugins/Experimental/EditorDataStorage/Source/TedsCore/Tests/Commands/EditorDataStorageCommandBufferTests.cpp b/Engine/Plugins/Experimental/EditorDataStorage/Source/TedsCore/Tests/Commands/EditorDataStorageCommandBufferTests.cpp
--- a/Engine/Plugins/Experimental/EditorDataStorage/Source/TedsCore/Tests/Commands/EditorDataStorageCommandBufferTests.cpp
+++ b/Engine/Plugins/Experimental/EditorDataStorage/Source/TedsCore/Tests/Commands/EditorDataStorageCommandBufferTests.cpp
@@ -409,6 +409,35 @@ TEST_CASE_NAMED(EditorDataStorage_FCommandBuffer_FOptimizer, "Editor::DataStorag
 	TestCommandBuffer CommandBuffer;
 	CommandBuffer.Initialize(ScratchBuffer);
 
+	SECTION("Optimizing empty buffer")
+	{
+		TestCommandBuffer::FCollection Commands;
+		CommandBuffer.Collect(Commands);
+		
+		// Expecting no crash.
+		TestCommandBuffer::FOptimizer Optimizer(Commands);
+	}
+
+	SECTION("Optimizing buffer with single command")
+	{
+		TestCommandBuffer::FCollection Commands;
+		CommandBuffer.Collect(Commands);
+		Commands.AddCommand(FTestCommandA{ .Value = 1 });
+
+		// Expecting no crash.
+		TestCommandBuffer::FOptimizer Optimizer(Commands);
+	}
+
+	SECTION("Optimizing buffer with nop command")
+	{
+		TestCommandBuffer::FCollection Commands;
+		CommandBuffer.Collect(Commands);
+		Commands.AddCommand<FNopCommand>();
+
+		// Expecting no crash.
+		TestCommandBuffer::FOptimizer Optimizer(Commands);
+	}
+
 	SECTION("Constructor corrects for nops before left")
 	{
 		TestCommandBuffer::FCollection Commands;
