[TEDS] Enable QueryCondition matching for Queries
* Overrides the All/Any/None matching logic of Mass to use the boolean expression engine from TEDS to match archetypes
* If QueryConditions are used, then All/Any/None logic registered for the query is ignored

diff --git a/Engine/Plugins/Experimental/EditorDataStorage/Source/TedsCore/Private/Queries/TypedElementExtendedQueryStore.cpp b/Engine/Plugins/Experimental/EditorDataStorage/Source/TedsCore/Private/Queries/TypedElementExtendedQueryStore.cpp
--- a/Engine/Plugins/Experimental/EditorDataStorage/Source/TedsCore/Private/Queries/TypedElementExtendedQueryStore.cpp
+++ b/Engine/Plugins/Experimental/EditorDataStorage/Source/TedsCore/Private/Queries/TypedElementExtendedQueryStore.cpp
@@ -9,10 +9,95 @@
 #include "Misc/OutputDevice.h"
 #include "Processors/TypedElementProcessorAdaptors.h"
 
+
+
 namespace UE::Editor::DataStorage
 {
+	namespace Private
+	{
+		bool bUseTedsMatchFunction = true;
+		static FAutoConsoleVariableRef CVarUseTedsMatchFunction(
+			TEXT("TEDS.Feature.UseTedsMatchFunction"), bUseTedsMatchFunction,
+			TEXT("Override the Mass query's match function to use the TEDS match function"),
+			ECVF_Default);
+	}
 	const FQueryDescription FExtendedQueryStore::EmptyDescription{};
 
+	class FEnvironmentQueryConditionCompileContext : public Queries::IQueryConditionCompileContext
+	{
+	public:
+		explicit FEnvironmentQueryConditionCompileContext(FEnvironment& InEnvironment)
+			: Environment(InEnvironment)
+		{}
+		FEnvironmentQueryConditionCompileContext(const FEnvironmentQueryConditionCompileContext&) = delete;
+		FEnvironmentQueryConditionCompileContext(FEnvironmentQueryConditionCompileContext&&) = delete;
+		FEnvironmentQueryConditionCompileContext& operator=(const FEnvironmentQueryConditionCompileContext&) = delete;
+		FEnvironmentQueryConditionCompileContext& operator=(FEnvironmentQueryConditionCompileContext&&) = delete;
+		
+		virtual const UScriptStruct* GenerateDynamicColumn(const FDynamicColumnDescription& Description) const override
+		{
+			return Environment.GenerateDynamicColumn(*Description.TemplateType, Description.Identifier);
+		}
+	private:
+		FEnvironment& Environment;
+	};
+
+	struct FTedsQueryArchetypeMatchOverride
+	{
+		bool Match(const FMassArchetypeCompositionDescriptor& CompositionDescriptor) const;
+		
+		FExtendedQueryStore::Handle QueryHandle;
+		FEnvironment& Environment;
+	};
+	
+	bool FTedsQueryArchetypeMatchOverride::Match(const FMassArchetypeCompositionDescriptor& CompositionDescriptor) const
+	{
+		FExtendedQuery* Query = Environment.GetQueryStore().Get(QueryHandle);
+		if (Query->Description.Conditions)
+		{
+			TArray<TWeakObjectPtr<const UScriptStruct>, TInlineAllocator<64>> ColumnTypes;
+			auto FragmentIterator = CompositionDescriptor.Fragments.GetIndexIterator(true);
+			while (FragmentIterator)
+			{
+				const UScriptStruct* FragmentStruct = CompositionDescriptor.Fragments.GetTypeAtIndex(*FragmentIterator);
+				ColumnTypes.Add(FragmentStruct);
+				++FragmentIterator;
+			}
+			
+			// Check if the archetype matches the selected columns
+			for (int32 Index = 0, SelectionTypesCount = Query->Description.SelectionTypes.Num(); Index < SelectionTypesCount; ++Index)
+			{
+				if (Query->Description.SelectionAccessTypes[Index] != EQueryAccessType::OptionalReadOnly)
+				{
+					if(!ColumnTypes.Contains(Query->Description.SelectionTypes[Index]))
+					{
+						// Archetype doesn't have the required selected column, so reject the archetype outright.
+						return false;
+					}
+				}
+			}
+
+			// Then ensure it also matches the query conditions
+			auto TagIterator = CompositionDescriptor.Tags.GetIndexIterator(true);			
+			while (TagIterator)
+			{
+				const UScriptStruct* TagStruct = CompositionDescriptor.Tags.GetTypeAtIndex(*TagIterator);
+				ColumnTypes.Add(TagStruct);
+				++TagIterator;
+			}
+
+			constexpr bool bAvailableColumnsAreSorted = false;
+			const bool bMatchesQueryConditions = Query->Description.Conditions->Verify(MakeArrayView(ColumnTypes.GetData(),  ColumnTypes.Num()), bAvailableColumnsAreSorted);
+
+			return bMatchesQueryConditions;
+		}
+		else
+		{
+			return Query->GetQuery()->DoesArchetypeMatchRequirements(CompositionDescriptor);
+		}
+
+	}
+
 	FExtendedQueryStore::Handle FExtendedQueryStore::RegisterQuery(
 		FQueryDescription Query,
 		FEnvironment& Environment,
@@ -24,10 +109,22 @@ namespace UE::Editor::DataStorage
 		StoredQuery.Description = MoveTemp(Query);
 
 		FMassEntityQuery& NativeQuery = SetupNativeQuery(StoredQuery.Description, StoredQuery);
+
+		if (Private::bUseTedsMatchFunction)
+		{
+			const FTedsQueryArchetypeMatchOverride Context
+			{
+				.QueryHandle = Result,
+				.Environment = Environment
+			};
+
+			NativeQuery.SetArchetypeMatchOverride(Context);
+		}
+		
 		bool bContinueSetup =               SetupDynamicColumns(StoredQuery.Description, Environment);
 		bContinueSetup = bContinueSetup &&	SetupSelectedColumns(StoredQuery.Description, NativeQuery);
 		bContinueSetup = bContinueSetup && SetupChunkFilters(Result, StoredQuery.Description, Environment, NativeQuery);
-		bContinueSetup = bContinueSetup && SetupConditions(StoredQuery.Description, NativeQuery);
+		bContinueSetup = bContinueSetup && SetupConditions(StoredQuery.Description, Environment, NativeQuery);
 		bContinueSetup = bContinueSetup && SetupDependencies(StoredQuery.Description, NativeQuery);
 		bContinueSetup = bContinueSetup && SetupTickGroupDefaults(StoredQuery.Description);
 		bContinueSetup = bContinueSetup && SetupProcessors(Result, StoredQuery, Environment, EntityManager, PhaseManager);
@@ -423,6 +520,8 @@ namespace UE::Editor::DataStorage
 		 * UHT doesn't allow for templates so it had to be done in an explicit way.
 		 */
 
+		StoredQuery.QueryReference = nullptr;
+
 		if (Query.Action == FQueryDescription::EActionType::Select)
 		{
 			switch (Query.Callback.Type)
@@ -467,6 +566,7 @@ namespace UE::Editor::DataStorage
 					return StoredQuery.NativeQuery;
 				}
 				StoredQuery.Processor.Reset(Processor);
+				StoredQuery.QueryReference = &Processor->GetQuery();
 				return Processor->GetQuery();
 			}
 			case EQueryCallbackType::ObserveAdd:
@@ -509,6 +609,7 @@ namespace UE::Editor::DataStorage
 					return StoredQuery.NativeQuery;
 				}
 				StoredQuery.Processor.Reset(Observer);
+				StoredQuery.QueryReference = &Observer->GetQuery();
 				return Observer->GetQuery();
 			}
 			case EQueryCallbackType::PhasePreparation:
@@ -520,6 +621,7 @@ namespace UE::Editor::DataStorage
 				break;
 			}
 		}
+		
 		return StoredQuery.NativeQuery;
 	}
 
@@ -572,49 +674,96 @@ namespace UE::Editor::DataStorage
 		}
 	}
 
-	bool FExtendedQueryStore::SetupConditions(FQueryDescription& Query, FMassEntityQuery& NativeQuery)
+	bool FExtendedQueryStore::SetupConditions(FQueryDescription& Query, FEnvironment& Environment, FMassEntityQuery& NativeQuery)
 	{
-		if (Query.ConditionTypes.IsEmpty())
+		if (Private::bUseTedsMatchFunction && Query.Conditions.IsSet())
 		{
-			return true;
-		}
+			// Will ignore the All/Any/None conditions if in this branch, but let the user know via an ensure anyways
+			ensureAlwaysMsgf(Query.ConditionTypes.IsEmpty(), TEXT("Queries do not support using both QueryConditions and All/Any/None conditions"));
+				
+			// Use the QueryConditions to setup the condition requirements for Mass
+			Queries::FConditions& Conditions = *Query.Conditions;
+			const FEnvironmentQueryConditionCompileContext CompileContext(Environment);
+			Conditions.Compile(CompileContext);
+
+			const TConstArrayView<TWeakObjectPtr<const UScriptStruct>>& Columns = Conditions.GetColumns();
+
+			// Note: need to remove duplicates to appease Mass
+			TSet<const UScriptStruct*, DefaultKeyFuncs<const UScriptStruct*>, TInlineSetAllocator<128>> DedupedColumns;
+			for (const auto& Column : Columns)
+			{
+				DedupedColumns.Add(Column.Get());
+			}
 
-		if (ensureMsgf(Query.ConditionTypes.Num() == Query.ConditionOperators.Num(),
-			TEXT("The types and operators for a typed element query have gone out of sync.")))
-		{
-			const FQueryDescription::FOperator* Operand = Query.ConditionOperators.GetData();
-			for (FQueryDescription::EOperatorType Type : Query.ConditionTypes)
+			// Everything is going to be Any, Mass' internals checks requirements 
+			// We override the archetype matching function anyways so Mass will not use the All, Any or None
+			const EMassFragmentPresence Presence = EMassFragmentPresence::Any;
+			for (const UScriptStruct* Column : DedupedColumns)
 			{
-				EMassFragmentPresence Presence;
-				switch (Type)
+				if (Column->IsChildOf(FEditorDataStorageTag::StaticStruct()))
 				{
-				case FQueryDescription::EOperatorType::SimpleAll:
-					Presence = EMassFragmentPresence::All;
-					break;
-				case FQueryDescription::EOperatorType::SimpleAny:
-					Presence = EMassFragmentPresence::Any;
-					break;
-				case FQueryDescription::EOperatorType::SimpleNone:
-					Presence = EMassFragmentPresence::None;
-					break;
-				default:
-					continue;
+					NativeQuery.AddTagRequirement(*Column, Presence);
 				}
-
-				if (Operand->Type->IsChildOf(FMassTag::StaticStruct()))
+				else if (Column->IsChildOf(FEditorDataStorageColumn::StaticStruct()))
 				{
-					NativeQuery.AddTagRequirement(*(Operand->Type), Presence);
+					// Check that the requirement wasn't already added as a "Selected" column
+					const bool IsSelected = Query.SelectionTypes.Contains(Column);
+					if (ensureMsgf(!IsSelected, TEXT("Cannot add '%s' as a condition if it is a selected column"), *Column->GetName()))
+					{
+						NativeQuery.AddRequirement(Column, EMassFragmentAccess::None, Presence);
+					}
+					else
+					{
+						return false;
+					}
 				}
-				else if (Operand->Type->IsChildOf(FMassFragment::StaticStruct()))
+			}
+			return true;
+		}
+		else
+		{
+			if (Query.ConditionTypes.IsEmpty())
+			{
+				return true;
+			}
+			
+			if (ensureMsgf(Query.ConditionTypes.Num() == Query.ConditionOperators.Num(),
+				TEXT("The types and operators for a typed element query have gone out of sync.")))
+			{
+				const FQueryDescription::FOperator* Operand = Query.ConditionOperators.GetData();
+				for (FQueryDescription::EOperatorType Type : Query.ConditionTypes)
 				{
-					NativeQuery.AddRequirement(Operand->Type.Get(), EMassFragmentAccess::None, Presence);
-				}
+					EMassFragmentPresence Presence;
+					switch (Type)
+					{
+					case FQueryDescription::EOperatorType::SimpleAll:
+						Presence = EMassFragmentPresence::All;
+						break;
+					case FQueryDescription::EOperatorType::SimpleAny:
+						Presence = EMassFragmentPresence::Any;
+						break;
+					case FQueryDescription::EOperatorType::SimpleNone:
+						Presence = EMassFragmentPresence::None;
+						break;
+					default:
+						continue;
+					}
 
-				++Operand;
+					if (Operand->Type->IsChildOf(FMassTag::StaticStruct()))
+					{
+						NativeQuery.AddTagRequirement(*(Operand->Type), Presence);
+					}
+					else if (Operand->Type->IsChildOf(FMassFragment::StaticStruct()))
+					{
+						NativeQuery.AddRequirement(Operand->Type.Get(), EMassFragmentAccess::None, Presence);
+					}
+
+					++Operand;
+				}
+				return true;
 			}
-			return true;
+			return false;
 		}
-		return false;
 	}
 
 	bool FExtendedQueryStore::SetupChunkFilters(
diff --git a/Engine/Plugins/Experimental/EditorDataStorage/Source/TedsCore/Private/Queries/TypedElementExtendedQueryStore.h b/Engine/Plugins/Experimental/EditorDataStorage/Source/TedsCore/Private/Queries/TypedElementExtendedQueryStore.h
--- a/Engine/Plugins/Experimental/EditorDataStorage/Source/TedsCore/Private/Queries/TypedElementExtendedQueryStore.h
+++ b/Engine/Plugins/Experimental/EditorDataStorage/Source/TedsCore/Private/Queries/TypedElementExtendedQueryStore.h
@@ -22,6 +22,20 @@ namespace UE::Editor::DataStorage
 		FMassEntityQuery NativeQuery; // Used if there's no processor bound.
 		IEditorDataStorageProvider::FQueryDescription Description;
 		TStrongObjectPtr<UMassProcessor> Processor;
+		FMassEntityQuery* QueryReference;
+		
+		FMassEntityQuery* GetQuery()
+		{
+			if (Processor)
+			{
+				return QueryReference;
+			}
+			else
+			{
+				return &NativeQuery;
+			}
+		}
+		
 	};
 
 	/**
@@ -167,7 +181,7 @@ namespace UE::Editor::DataStorage
 		FMassEntityQuery& SetupNativeQuery(IEditorDataStorageProvider::FQueryDescription& Query, FExtendedQuery& StoredQuery);
 		bool SetupDynamicColumns(IEditorDataStorageProvider::FQueryDescription& Query, FEnvironment& Environment);
 		bool SetupSelectedColumns(IEditorDataStorageProvider::FQueryDescription& Query, FMassEntityQuery& NativeQuery);
-		bool SetupConditions(IEditorDataStorageProvider::FQueryDescription& Query, FMassEntityQuery& NativeQuery);
+		bool SetupConditions(IEditorDataStorageProvider::FQueryDescription& Query, FEnvironment& Environment, FMassEntityQuery& NativeQuery);
 		bool SetupChunkFilters(Handle QueryHandle, IEditorDataStorageProvider::FQueryDescription& Query, FEnvironment& Environment, FMassEntityQuery& NativeQuery);
 		bool SetupDependencies(IEditorDataStorageProvider::FQueryDescription& Query, FMassEntityQuery& NativeQuery);
 		bool SetupTickGroupDefaults(IEditorDataStorageProvider::FQueryDescription& Query);
diff --git a/Engine/Plugins/Experimental/EditorDataStorage/Source/TedsCore/Tests/TypedElementDataStorageInterfaceTests.cpp b/Engine/Plugins/Experimental/EditorDataStorage/Source/TedsCore/Tests/TypedElementDataStorageInterfaceTests.cpp
--- a/Engine/Plugins/Experimental/EditorDataStorage/Source/TedsCore/Tests/TypedElementDataStorageInterfaceTests.cpp
+++ b/Engine/Plugins/Experimental/EditorDataStorage/Source/TedsCore/Tests/TypedElementDataStorageInterfaceTests.cpp
@@ -307,6 +307,70 @@ namespace UE::Editor::DataStorage::Tests
 						}
 					});
 
+					Describe("RegisterQuery", [this]()
+					{
+						It("QueryCondition", [this]()
+						{
+							using namespace Queries;
+							auto Query = TedsInterface->RegisterQuery(
+								Select()
+								.ReadOnly<FTestColumnD>()
+								.Where((TColumn<FTestColumnA>() && TColumn<FTestColumnB>()) || TColumn<FTestColumnC>())
+								.Compile());
+
+							if(!TestNotEqual(TEXT("Expect RegisterQuery to provide a valid query handle"), Query, InvalidQueryHandle))
+							{
+								return;
+							}
+
+							// Should not match
+							{
+								// Expected to have FTestColumnA && FTestColumnB but not FTestColumnD
+								RowHandle Handle = TedsInterface->AddRow(TestTableHandleA);
+								CreatedRows.Add(Handle);
+							}
+
+							// Should match
+							{
+								// Expected to have FTestColumnA && FTestColumnB
+								RowHandle Handle = TedsInterface->AddRow(TestTableHandleA);
+								CreatedRows.Add(Handle);
+								TedsInterface->AddColumn<FTestColumnD>(Handle, FTestColumnD{});
+							}
+
+							// Should not match
+							{
+								// Expected to have FTestColumnA only
+								RowHandle Handle = TedsInterface->AddRow(TestTableHandleA);
+								CreatedRows.Add(Handle);
+								TedsInterface->RemoveColumn<FTestColumnB>(Handle);
+								TedsInterface->AddColumn<FTestColumnD>(Handle, FTestColumnD{});
+							}
+
+							// Should match
+							{
+								// Expected to have FTestColumnC only
+								RowHandle Handle = TedsInterface->AddRow(TestTableHandleA);
+								CreatedRows.Add(Handle);
+								TedsInterface->RemoveColumn<FTestColumnA>(Handle);
+								TedsInterface->RemoveColumn<FTestColumnB>(Handle);
+								TedsInterface->AddColumn<FTestColumnC>(Handle, FTestColumnC{});
+								TedsInterface->AddColumn<FTestColumnD>(Handle, FTestColumnD{});
+							}
+
+							FQueryResult Result = TedsInterface->RunQuery(
+								Query,
+								CreateDirectQueryCallbackBinding([](IDirectQueryContext& Context, const RowHandle* RowHandles, const FTestColumnD* DColumns)
+							{
+									
+							}));
+
+							TestEqual(TEXT("Expect number of rows matched equal to created rows"), Result.Count, 2);
+
+							TedsInterface->UnregisterQuery(Query);
+						});
+					});
+
 					Describe("", [this]()
 					{
 						BeforeEach([this]()
diff --git a/Engine/Source/Runtime/MassEntity/Private/MassEntityQuery.cpp b/Engine/Source/Runtime/MassEntity/Private/MassEntityQuery.cpp
--- a/Engine/Source/Runtime/MassEntity/Private/MassEntityQuery.cpp
+++ b/Engine/Source/Runtime/MassEntity/Private/MassEntityQuery.cpp
@@ -123,7 +123,23 @@ void FMassEntityQuery::CacheArchetypes(const FMassEntityManager& InEntityManager
 	if (bUpdateArchetypes)
 	{
 		TArray<FMassArchetypeHandle> NewValidArchetypes;
-		InEntityManager.GetMatchingArchetypes(*this, NewValidArchetypes, LastUpdatedArchetypeDataVersion);
+#if WITH_EDITOR
+		if (bHasArchetypeMatchOverride)
+		{
+			InEntityManager.ForEachArchetype(static_cast<int32>(LastUpdatedArchetypeDataVersion), TNumericLimits<int32>::Max() /*last*/, [this, &NewValidArchetypes](const FMassEntityManager& EntityManager, const FMassArchetypeHandle& Handle)
+			{
+				const FMassArchetypeCompositionDescriptor& Composition = EntityManager.GetArchetypeComposition(Handle);
+				if (ArchetypeMatchOverride.Match(&ArchetypeMatchOverride.Data, Composition))
+				{
+					NewValidArchetypes.Add(Handle);
+				}
+			});
+		}
+		else
+#endif
+		{
+			InEntityManager.GetMatchingArchetypes(*this, NewValidArchetypes, LastUpdatedArchetypeDataVersion);
+		}
 		LastUpdatedArchetypeDataVersion = InEntityManager.GetArchetypeDataVersion();
 		if (NewValidArchetypes.Num())
 		{
diff --git a/Engine/Source/Runtime/MassEntity/Public/MassEntityManager.h b/Engine/Source/Runtime/MassEntity/Public/MassEntityManager.h
--- a/Engine/Source/Runtime/MassEntity/Public/MassEntityManager.h
+++ b/Engine/Source/Runtime/MassEntity/Public/MassEntityManager.h
@@ -623,6 +623,22 @@ private:
 
 public:
 
+#if WITH_EDITOR || WITH_MASSENTITY_DEBUG
+	template<typename CallableT>
+	void ForEachArchetype(int32 BeginRange, int32 EndRange, const CallableT& Callable) const
+	{
+		if (EndRange > AllArchetypes.Num())
+		{
+			EndRange = AllArchetypes.Num();
+		}
+		for (int32 Cursor = BeginRange; Cursor < EndRange; ++Cursor)
+		{
+			FMassArchetypeHandle Handle(AllArchetypes[Cursor]);
+			Callable(*this, Handle);
+		}
+	}
+#endif
+
 	template<typename T, typename... TArgs>
 	UE_DEPRECATED(5.5, "This method will no longer be exposed. Use GetOrCreateSharedFragment instead.")
 	const FSharedStruct& GetOrCreateSharedFragmentByHash(const uint32 Hash, TArgs&&... InArgs)
diff --git a/Engine/Source/Runtime/MassEntity/Public/MassEntityQuery.h b/Engine/Source/Runtime/MassEntity/Public/MassEntityQuery.h
--- a/Engine/Source/Runtime/MassEntity/Public/MassEntityQuery.h
+++ b/Engine/Source/Runtime/MassEntity/Public/MassEntityQuery.h
@@ -11,6 +11,25 @@
 #include "MassRequirements.h"
 #include "MassEntityQuery.generated.h"
 
+#ifndef WITH_ARCHETYPE_MATCH_OVERRIDE
+#define WITH_ARCHETYPE_MATCH_OVERRIDE (WITH_EDITOR)
+#endif //  WITH_ARCHETYPE_MATCH_OVERRIDE
+
+#ifndef MASS_ACHETYPE_OVERRIDE_MAX_SIZE
+#define MASS_ACHETYPE_OVERRIDE_MAX_SIZE 16
+#endif
+	
+#ifndef MASS_ACHETYPE_OVERRIDE_MAX_ALIGNMENT
+#define MASS_ACHETYPE_OVERRIDE_MAX_ALIGNMENT 8
+#endif
+
+#if WITH_ARCHETYPE_MATCH_OVERRIDE
+template<typename T>
+concept TArchetypeMatchOverrideConcept = requires(const FMassArchetypeCompositionDescriptor& Descriptor)
+{
+	{ static_cast<const T*>(nullptr)->Match(Descriptor) } -> std::convertible_to<bool>;
+};
+#endif
 
 /** 
  *  FMassEntityQuery is a structure that is used to trigger calculations on cached set of valid archetypes as described 
@@ -24,9 +43,21 @@ USTRUCT()
 struct MASSENTITY_API FMassEntityQuery : public FMassFragmentRequirements, public FMassSubsystemRequirements
 {
 	GENERATED_BODY()
-
+	
+	static constexpr int32 ArchetypeMatchOverrideSize = MASS_ACHETYPE_OVERRIDE_MAX_SIZE;
+	static constexpr uint32 ArchetypeMatchOverrideAlignment = MASS_ACHETYPE_OVERRIDE_MAX_ALIGNMENT;
+	
 	friend struct FMassDebugger;
 
+private:
+	struct FArchetypeMatchOverride
+    {
+		using MatchFunction = bool (*)(const void* /*Context*/, const FMassArchetypeCompositionDescriptor&);
+		
+		MatchFunction Match;
+    	TAlignedBytes<ArchetypeMatchOverrideSize, ArchetypeMatchOverrideAlignment> Data;
+    };
+
 public:
 	enum EParallelForMode
 	{
@@ -126,6 +157,11 @@ public:
 	void ClearChunkFilter() { ChunkCondition.Reset(); }
 
 	bool HasChunkFilter() const { return bool(ChunkCondition); }
+	
+#if WITH_ARCHETYPE_MATCH_OVERRIDE
+	template<typename T> requires TArchetypeMatchOverrideConcept<T>
+	void SetArchetypeMatchOverride(const T& Override);
+#endif
 
 	/** 
 	 * If ArchetypeHandle is among ValidArchetypes then the function retrieves requirements mapping cached for it,
@@ -175,11 +211,41 @@ private:
 	 */
 	uint8 bAllowParallelCommands : 1 = true;
 	uint8 bRequiresMutatingWorldAccess : 1 = false;
+#if WITH_ARCHETYPE_MATCH_OVERRIDE
+	uint8 bHasArchetypeMatchOverride : 1 = false;
+#endif
 
 	EMassExecutionContextType ExpectedContextType = EMassExecutionContextType::Local;
 
 #if WITH_MASSENTITY_DEBUG
 	uint8 bRegistered : 1;
 #endif // WITH_MASSENTITY_DEBUG
+
+#if WITH_ARCHETYPE_MATCH_OVERRIDE
+	FArchetypeMatchOverride ArchetypeMatchOverride;
+#endif
 };
 
+#if WITH_ARCHETYPE_MATCH_OVERRIDE
+template <typename T> requires TArchetypeMatchOverrideConcept<T>
+void FMassEntityQuery::SetArchetypeMatchOverride(const T& Context)
+{
+	static_assert(sizeof(T) <= sizeof(FArchetypeMatchOverride));
+	static_assert(alignof(T) <= alignof(FArchetypeMatchOverride));
+	static_assert(std::is_trivially_copyable_v<T>);
+	static_assert(std::is_trivially_destructible_v<T>);
+
+	check(!bHasArchetypeMatchOverride);
+	bHasArchetypeMatchOverride = true;
+
+	ArchetypeMatchOverride.Match = [](const void* TypeErasedContext, const FMassArchetypeCompositionDescriptor& Descriptor)->bool
+	{
+		const T* Context = static_cast<const T*>(TypeErasedContext);
+		return Context->Match(Descriptor);	
+	};
+	FMemory::Memcpy(&ArchetypeMatchOverride.Data, &Context, sizeof(T));
+}
+#endif
+
+#undef MASS_ACHETYPE_OVERRIDE_MAX_SIZE
+#undef MASS_ACHETYPE_OVERRIDE_MAX_ALIGNMENT
\ No newline at end of file
diff --git a/Engine/Source/Runtime/TypedElementFramework/Private/Elements/Framework/TypedElementQueryBuilder.cpp b/Engine/Source/Runtime/TypedElementFramework/Private/Elements/Framework/TypedElementQueryBuilder.cpp
--- a/Engine/Source/Runtime/TypedElementFramework/Private/Elements/Framework/TypedElementQueryBuilder.cpp
+++ b/Engine/Source/Runtime/TypedElementFramework/Private/Elements/Framework/TypedElementQueryBuilder.cpp
@@ -393,6 +393,22 @@ namespace UE::Editor::DataStorage::Queries
 	 * Select
 	 */
 
+	FDependency FQueryConditionQuery::DependsOn()
+	{
+		return FDependency{ Query };
+	}
+
+	FQueryDescription&& FQueryConditionQuery::Compile()
+	{
+		return MoveTemp(*Query);
+	}
+
+	FQueryConditionQuery::FQueryConditionQuery(FQueryDescription* InQuery)
+		: Query(InQuery)
+	{
+		
+	}
+
 	Select::Select()
 	{
 		Query.Action = FQueryDescription::EActionType::Select;
@@ -497,6 +513,12 @@ namespace UE::Editor::DataStorage::Queries
 		return FDependency{ &Query };
 	}
 
+	FQueryConditionQuery Select::Where(const Queries::FConditions& Condition)
+	{
+		Query.Conditions.Emplace(Condition);
+		return FQueryConditionQuery(&Query);
+	}
+
 	FQueryDescription&& Select::Compile()
 	{
 		return MoveTemp(Query);
diff --git a/Engine/Source/Runtime/TypedElementFramework/Public/Elements/Common/TypedElementQueryDescription.h b/Engine/Source/Runtime/TypedElementFramework/Public/Elements/Common/TypedElementQueryDescription.h
--- a/Engine/Source/Runtime/TypedElementFramework/Public/Elements/Common/TypedElementQueryDescription.h
+++ b/Engine/Source/Runtime/TypedElementFramework/Public/Elements/Common/TypedElementQueryDescription.h
@@ -6,9 +6,11 @@
 #include "Containers/ContainerAllocationPolicies.h"
 #include "Elements/Common/TypedElementHandles.h"
 #include "Elements/Common/TypedElementCommonTypes.h"
+#include "Elements/Common/TypedElementQueryConditions.h"
 #include "Elements/Common/TypedElementQueryTypes.h"
 #include "Elements/Interfaces/TypedElementQueryStorageInterfaces.h"
 #include "Elements/Framework/TypedElementMetaData.h"
+#include "Misc/Optional.h"
 #include "Templates/Function.h"
 #include "UObject/NameTypes.h"
 #include "UObject/WeakObjectPtr.h"
@@ -107,6 +109,8 @@ namespace UE::Editor::DataStorage
 
 		TArray<FValueTagData> ValueTags;
 
+		TOptional<UE::Editor::DataStorage::Queries::FConditions> Conditions;
+
 		TArray<TWeakObjectPtr<const UClass>, TInlineAllocator<NumInlineDependencies>> DependencyTypes;
 		TArray<EQueryDependencyFlags, TInlineAllocator<NumInlineDependencies>> DependencyFlags;
 		/** Cached instances of the dependencies. This will always match the count of the other Dependency*Types, but may contain null pointers. */
diff --git a/Engine/Source/Runtime/TypedElementFramework/Public/Elements/Framework/TypedElementQueryBuilder.h b/Engine/Source/Runtime/TypedElementFramework/Public/Elements/Framework/TypedElementQueryBuilder.h
--- a/Engine/Source/Runtime/TypedElementFramework/Public/Elements/Framework/TypedElementQueryBuilder.h
+++ b/Engine/Source/Runtime/TypedElementFramework/Public/Elements/Framework/TypedElementQueryBuilder.h
@@ -121,6 +121,7 @@ namespace UE::Editor::DataStorage::Queries
 		friend class Count;
 		friend class Select;
 		friend class FSimpleQuery;
+		friend class FQueryConditionQuery;
 	public:
 		template<typename... TargetTypes>
 		FDependency& ReadOnly();
@@ -354,6 +355,20 @@ namespace UE::Editor::DataStorage::Queries
 		const Dependency& GetCachedDependency() const;
 	};
 
+	class FQueryConditionQuery
+	{
+	public:
+		friend class Count;
+		friend class Select;
+
+		TYPEDELEMENTFRAMEWORK_API FDependency DependsOn();
+		TYPEDELEMENTFRAMEWORK_API FQueryDescription&& Compile();
+	private:
+		TYPEDELEMENTFRAMEWORK_API explicit FQueryConditionQuery(FQueryDescription* Query);
+
+		FQueryDescription* Query;
+	};
+
 	// Explicitly not following the naming convention in order to present this as a query that can be read as such.
 	class Select final
 	{
@@ -404,6 +419,9 @@ namespace UE::Editor::DataStorage::Queries
 		template<TDataColumnType Target>
 		Select& ReadWrite(const FName& Identifier);
 
+		// Query Condition API
+		TYPEDELEMENTFRAMEWORK_API FQueryConditionQuery Where(const Queries::FConditions& Condition);
+
 		TYPEDELEMENTFRAMEWORK_API FQueryDescription&& Compile();
 		TYPEDELEMENTFRAMEWORK_API FSimpleQuery Where();
 		TYPEDELEMENTFRAMEWORK_API FDependency DependsOn();
