TEDS UI: Fix column virtualization not working for the export text widget

diff --git a/Engine/Plugins/Experimental/EditorDataStorage/Source/TedsUI/Private/Widgets/ExportedTextWidget.cpp b/Engine/Plugins/Experimental/EditorDataStorage/Source/TedsUI/Private/Widgets/ExportedTextWidget.cpp
--- a/Engine/Plugins/Experimental/EditorDataStorage/Source/TedsUI/Private/Widgets/ExportedTextWidget.cpp
+++ b/Engine/Plugins/Experimental/EditorDataStorage/Source/TedsUI/Private/Widgets/ExportedTextWidget.cpp
@@ -142,7 +142,12 @@ TConstArrayView<const UScriptStruct*> FExportedTextWidgetConstructor::GetAdditio
 const UE::Editor::DataStorage::Queries::FConditions* FExportedTextWidgetConstructor::GetQueryConditions() const
 {
 	// For the exported text widget, the query condition we are matched against is the column we are exporting text for
-	return &MatchedColumn;
+	if(MatchedColumn.IsCompiled() && !MatchedColumn.IsEmpty())
+	{
+		return &MatchedColumn;
+	}
+
+	return nullptr;
 }
 
 FString FExportedTextWidgetConstructor::CreateWidgetDisplayName(IEditorDataStorageProvider* DataStorage,
@@ -158,6 +163,21 @@ FString FExportedTextWidgetConstructor::CreateWidgetDisplayName(IEditorDataStora
 	}
 }
 
+TSharedPtr<SWidget> FExportedTextWidgetConstructor::ConstructFinalWidget(RowHandle Row, IEditorDataStorageProvider* DataStorage,
+	IEditorDataStorageUiProvider* DataStorageUi, const UE::Editor::DataStorage::FMetaDataView& Arguments)
+{
+	FTypedElementScriptStructTypeInfoColumn& TypeInfoColumn = *DataStorage->GetColumn<FTypedElementScriptStructTypeInfoColumn>(Row);
+
+	using namespace UE::Editor::DataStorage::Queries;
+	
+	// NOTE: We are currently assuming that an instance of FExportedTextWidgetConstructor will only be used to show the same type info for all rows
+	// matched with it. This isn't ideal but it's better than nothing since we need some sort of matched conditions for column based virtualization to work.
+	// TEDS UI TODO: We should work around it by refactoring this into an STedsWidget in the future so it can store the column conditions per instance
+	MatchedColumn = FConditions(TColumn(TypeInfoColumn.TypeInfo)).Compile(FEditorStorageQueryConditionCompileContext(DataStorage));
+
+	return FTypedElementWidgetConstructor::ConstructFinalWidget(Row, DataStorage, DataStorageUi, Arguments);
+}
+
 TSharedPtr<SWidget> FExportedTextWidgetConstructor::CreateWidget(const UE::Editor::DataStorage::FMetaDataView& Arguments)
 {
 	return SNew(STextBlock);
@@ -172,11 +192,6 @@ bool FExportedTextWidgetConstructor::FinalizeWidget(
 	using namespace UE::Editor;
 	FTypedElementScriptStructTypeInfoColumn& TypeInfoColumn = *DataStorage->GetColumn<FTypedElementScriptStructTypeInfoColumn>(Row);
 
-	// NOTE: We are currently assuming that an instance of FExportedTextWidgetConstructor will only be used to show the same type info
-	// which isn't ideal but it's better than nothing since we need some sort of matched conditions for column based virtualization to work.
-	// TEDS UI TODO: We should work around it by refactoring this into an STedsWidget in the future so it can store the column conditions per instance
-	MatchedColumn = DataStorage::Queries::FConditions(DataStorage::Queries::TColumn(TypeInfoColumn.TypeInfo));
-
 	if (TypeInfoColumn.TypeInfo->IsChildOf(DataStorage::FTag::StaticStruct()))
 	{
 		DataStorage::UI::Private::UpdateExportedTextWidget(
diff --git a/Engine/Plugins/Experimental/EditorDataStorage/Source/TedsUI/Public/Widgets/ExportedTextWidget.h b/Engine/Plugins/Experimental/EditorDataStorage/Source/TedsUI/Public/Widgets/ExportedTextWidget.h
--- a/Engine/Plugins/Experimental/EditorDataStorage/Source/TedsUI/Public/Widgets/ExportedTextWidget.h
+++ b/Engine/Plugins/Experimental/EditorDataStorage/Source/TedsUI/Public/Widgets/ExportedTextWidget.h
@@ -39,6 +39,12 @@ public:
 	virtual FString CreateWidgetDisplayName(
 		IEditorDataStorageProvider* DataStorage, UE::Editor::DataStorage::RowHandle Row) const override;
 
+	virtual TSharedPtr<SWidget> ConstructFinalWidget(
+		RowHandle Row, /** The row the widget will be stored in. */
+		IEditorDataStorageProvider* DataStorage,
+		IEditorDataStorageUiProvider* DataStorageUi,
+		const UE::Editor::DataStorage::FMetaDataView& Arguments) override;
+	
 protected:
 	virtual TSharedPtr<SWidget> CreateWidget(const UE::Editor::DataStorage::FMetaDataView& Arguments) override;
 	virtual bool FinalizeWidget(
diff --git a/Engine/Plugins/Experimental/EditorDataStorageFeatures/Source/TedsTableViewer/Private/TedsTableViewerColumn.cpp b/Engine/Plugins/Experimental/EditorDataStorageFeatures/Source/TedsTableViewer/Private/TedsTableViewerColumn.cpp
--- a/Engine/Plugins/Experimental/EditorDataStorageFeatures/Source/TedsTableViewer/Private/TedsTableViewerColumn.cpp
+++ b/Engine/Plugins/Experimental/EditorDataStorageFeatures/Source/TedsTableViewer/Private/TedsTableViewerColumn.cpp
@@ -34,8 +34,14 @@ namespace UE::Editor::DataStorage
 			StorageCompatibility = Registry->GetMutableDataStorageCompatibility();
 		}
 
-		RegisterQueries();
+		// Store the matched columns as a query condition that requires all of them (i.e AND's them)
+		for(const TWeakObjectPtr<const UScriptStruct>& Column : MatchedColumns)
+		{
+			MatchedColumnConditions = MatchedColumnConditions && Queries::TColumn(Column);
+		}
+		MatchedColumnConditions.Compile(Queries::FEditorStorageQueryConditionCompileContext(Storage));
 
+		RegisterQueries();
 	}
 	
 	FTedsTableViewerColumn::~FTedsTableViewerColumn()
@@ -257,7 +263,19 @@ namespace UE::Editor::DataStorage
 			}
 			
 			// Check if the row now matches the query conditions for this widget
-			const bool bMatchesQueryConditions = CellWidgetConstructor->GetQueryConditions() && Storage->MatchesColumns(Pair.Key, *CellWidgetConstructor->GetQueryConditions());
+			bool bMatchesQueryConditions = true;
+
+			// First we try to match against the conditions provided by the widget constructor if possible
+			const Queries::FConditions* WidgetConstructorConditions = CellWidgetConstructor->GetQueryConditions();
+			if(WidgetConstructorConditions && WidgetConstructorConditions->IsCompiled())
+			{
+				bMatchesQueryConditions = bMatchesQueryConditions && Storage->MatchesColumns(Pair.Key, *WidgetConstructorConditions);
+			}
+			// If the widget constructor didn't provide any conditions, try to match against the columns we were provided on init
+			else
+			{
+				bMatchesQueryConditions = bMatchesQueryConditions && Storage->MatchesColumns(Pair.Key, MatchedColumnConditions);
+			}
 			
 			// If we are adding a column that we are monitoring and it now matches, or if we are removing a column that we are monitoring and it now
 			// stops matching, there is a potential need for widget update
diff --git a/Engine/Plugins/Experimental/EditorDataStorageFeatures/Source/TedsTableViewer/Public/TedsTableViewerColumn.h b/Engine/Plugins/Experimental/EditorDataStorageFeatures/Source/TedsTableViewer/Public/TedsTableViewerColumn.h
--- a/Engine/Plugins/Experimental/EditorDataStorageFeatures/Source/TedsTableViewer/Public/TedsTableViewerColumn.h
+++ b/Engine/Plugins/Experimental/EditorDataStorageFeatures/Source/TedsTableViewer/Public/TedsTableViewerColumn.h
@@ -3,6 +3,7 @@
 #pragma once
 #include "Elements/Common/TypedElementHandles.h"
 #include "Elements/Framework/TypedElementMetaData.h"
+#include "Elements/Framework/TypedElementQueryBuilder.h"
 #include "Templates/SharedPointer.h"
 #include "UObject/NameTypes.h"
 #include "Widgets/Views/SHeaderRow.h"
@@ -65,6 +66,9 @@ namespace UE::Editor::DataStorage
 
 		// Teds Columns this widget constructor matched with
 		TArray<TWeakObjectPtr<const UScriptStruct>> MatchedColumns;
+
+		// The matched columns stored as a query condition for quick access
+		Queries::FConditions MatchedColumnConditions;
 		
 		// The Metadata used to create widgets
 		FMetaDataView WidgetMetaData;
diff --git a/Engine/Source/Runtime/TypedElementFramework/Private/Elements/Interfaces/TypedElementDataStorageUiInterface.cpp b/Engine/Source/Runtime/TypedElementFramework/Private/Elements/Interfaces/TypedElementDataStorageUiInterface.cpp
--- a/Engine/Source/Runtime/TypedElementFramework/Private/Elements/Interfaces/TypedElementDataStorageUiInterface.cpp
+++ b/Engine/Source/Runtime/TypedElementFramework/Private/Elements/Interfaces/TypedElementDataStorageUiInterface.cpp
@@ -108,9 +108,9 @@ TSharedPtr<SWidget> FTypedElementWidgetConstructor::ConstructFinalWidget(
 		bool bConstructWidget = DataStorage->IsRowAssigned(RowReference->Row);
 
 		// If the original row matches this widgets query conditions currently, create the actual internal widget
-		if (QueryConditions)
+		if (const UE::Editor::DataStorage::Queries::FConditions* MatchedQueryConditions = GetQueryConditions())
 		{
-			bConstructWidget &= DataStorage->MatchesColumns(RowReference->Row, *QueryConditions);
+			bConstructWidget &= DataStorage->MatchesColumns(RowReference->Row, *MatchedQueryConditions);
 		}
 		
 		if (bConstructWidget)
diff --git a/Engine/Source/Runtime/TypedElementFramework/Public/Elements/Interfaces/TypedElementDataStorageUiInterface.h b/Engine/Source/Runtime/TypedElementFramework/Public/Elements/Interfaces/TypedElementDataStorageUiInterface.h
--- a/Engine/Source/Runtime/TypedElementFramework/Public/Elements/Interfaces/TypedElementDataStorageUiInterface.h
+++ b/Engine/Source/Runtime/TypedElementFramework/Public/Elements/Interfaces/TypedElementDataStorageUiInterface.h
@@ -66,7 +66,7 @@ public:
 	 *	Construct() can be called later to (re)create the internal widget if ever required.
 	 *	@see Construct
 	 */
-	TYPEDELEMENTFRAMEWORK_API TSharedPtr<SWidget> ConstructFinalWidget(
+	TYPEDELEMENTFRAMEWORK_API virtual TSharedPtr<SWidget> ConstructFinalWidget(
 		RowHandle Row, /** The row the widget will be stored in. */
 		IEditorDataStorageProvider* DataStorage,
 		IEditorDataStorageUiProvider* DataStorageUi,
