[TEDS] Deprecate accessor functions for TEDS features inside the TypedElementRegistry class

diff --git a/Engine/Plugins/Developer/PerforceSourceControl/Source/PerforceSourceControl/Private/PerforceSourceControlOperations.cpp b/Engine/Plugins/Developer/PerforceSourceControl/Source/PerforceSourceControl/Private/PerforceSourceControlOperations.cpp
--- a/Engine/Plugins/Developer/PerforceSourceControl/Source/PerforceSourceControl/Private/PerforceSourceControlOperations.cpp
+++ b/Engine/Plugins/Developer/PerforceSourceControl/Source/PerforceSourceControl/Private/PerforceSourceControlOperations.cpp
@@ -25,8 +25,8 @@
 #include "Elements/Columns/TypedElementMiscColumns.h"
 #include "Elements/Columns/TypedElementPackageColumns.h"
 #include "Elements/Columns/TypedElementRevisionControlColumns.h"
+#include "Elements/Common/EditorDataStorageFeatures.h"
 #include "Elements/Framework/TypedElementIndexHasher.h"
-#include "Elements/Framework/TypedElementRegistry.h"
 #include "Elements/Interfaces/TypedElementDataStorageInterface.h"
 #include "ProfilingDebugging/CpuProfilerTrace.h"
 
@@ -2031,15 +2031,10 @@ bool FPerforceUpdateStatusWorker::UpdateStates() const
 
 	auto UpdateDataStorage = [](const FPerforceSourceControlState& State)
 	{
+		using namespace UE::Editor::DataStorage;
 		using namespace UE::Editor::DataStorage::Queries;
 
-		UTypedElementRegistry* Registry = UTypedElementRegistry::GetInstance();
-		if (!Registry)
-		{
-			return;
-		}
-
-		IEditorDataStorageProvider* DataStorage = Registry->GetMutableDataStorage();
+		IEditorDataStorageProvider* DataStorage = GetMutableDataStorageFeature<IEditorDataStorageProvider>(StorageFeatureName);
 		if (!DataStorage)
 		{
 			return;
diff --git a/Engine/Plugins/Experimental/Animation/HierarchyTable/Source/Editor/Private/HierarchyTableEditorToolkit.cpp b/Engine/Plugins/Experimental/Animation/HierarchyTable/Source/Editor/Private/HierarchyTableEditorToolkit.cpp
--- a/Engine/Plugins/Experimental/Animation/HierarchyTable/Source/Editor/Private/HierarchyTableEditorToolkit.cpp
+++ b/Engine/Plugins/Experimental/Animation/HierarchyTable/Source/Editor/Private/HierarchyTableEditorToolkit.cpp
@@ -5,11 +5,12 @@
 #include "Modules/ModuleManager.h"
 #include "PropertyEditorModule.h"
 #include "TedsOutlinerModule.h"
-#include "Elements/Framework/TypedElementRegistry.h"
 #include "Elements/Columns/TypedElementLabelColumns.h"
 #include "Elements/Columns/TypedElementTypeInfoColumns.h"
 #include "Elements/Columns/TypedElementSelectionColumns.h"
 #include "Elements/Columns/TypedElementMiscColumns.h"
+#include "Elements/Common/EditorDataStorageFeatures.h"
+#include "Elements/Framework/TypedElementQueryBuilder.h"
 #include "Elements/Interfaces/TypedElementDataStorageInterface.h"
 #include "SceneOutlinerPublicTypes.h"
 #include "TedsOutlinerMode.h"
@@ -88,10 +89,10 @@ void FHierarchyTableEditorToolkit::RegisterTabSpawners(const TSharedRef<class FT
 
 void FHierarchyTableEditorToolkit::OnClose()
 {
-	UTypedElementRegistry* Registry = UTypedElementRegistry::GetInstance();
-	IEditorDataStorageProvider* DSI = Registry->GetMutableDataStorage();
+	using namespace UE::Editor::DataStorage;
+	IEditorDataStorageProvider* DSI = GetMutableDataStorageFeature<IEditorDataStorageProvider>(StorageFeatureName);
 
-	for (const TTuple<int32, UE::Editor::DataStorage::RowHandle>& Row : EntryIndexToHandleMap)
+	for (const TTuple<int32, RowHandle>& Row : EntryIndexToHandleMap)
 	{
 		DSI->RemoveRow(Row.Value);
 	}
@@ -106,18 +107,15 @@ void FHierarchyTableEditorToolkit::UnregisterTabSpawners(const TSharedRef<class
 
 TSharedRef<SWidget> FHierarchyTableEditorToolkit::CreateTedsOutliner()
 {
-	UTypedElementRegistry* Registry = UTypedElementRegistry::GetInstance();
-	checkf(Registry, TEXT("Unable to initialize the table viewer before TEDS is initialized."));
+	using namespace UE::Editor::DataStorage::Queries;
+	using namespace UE::Editor::Outliner;
 
-	if (!Registry->AreDataStorageInterfacesSet())
+	if (!AreEditorDataStorageFeaturesEnabled())
 	{
 		return SNew(STextBlock)
 			.Text(INVTEXT("You need to enable the Typed Element Data Storage plugin to see the table viewer!"));
 	}
 
-	using namespace UE::Editor::DataStorage::Queries;
-	using namespace UE::Editor::Outliner;
-
 	if (!ensure(HierarchyTable->TableType))
 	{
 		return SNullWidget::NullWidget;
@@ -138,7 +136,7 @@ TSharedRef<SWidget> FHierarchyTableEditorToolkit::CreateTedsOutliner()
 			.ReadOnly(HierarchyTableTypeColumns)
 		.Compile();
 
-	InitialColumnQuery = Registry->GetMutableDataStorage()->RegisterQuery(MoveTemp(ColumnQueryDescription));
+	InitialColumnQuery = GetMutableDataStorageFeature<IEditorDataStorageProvider>(StorageFeatureName)->RegisterQuery(MoveTemp(ColumnQueryDescription));
 
 	FSceneOutlinerInitializationOptions InitOptions;
 	InitOptions.bShowHeaderRow = true;
@@ -161,7 +159,7 @@ TSharedRef<SWidget> FHierarchyTableEditorToolkit::CreateTedsOutliner()
 
 	FTedsOutlinerModule& TedsOutlinerModule = FModuleManager::GetModuleChecked<FTedsOutlinerModule>("TedsOutliner");
 
-	IEditorDataStorageProvider* DSI = Registry->GetMutableDataStorage();
+	IEditorDataStorageProvider* DSI = GetMutableDataStorageFeature<IEditorDataStorageProvider>(StorageFeatureName);
 	static TableHandle Table = DSI->FindTable(FName("Editor_HierarchyTableTable"));
 	
 	TArray<UScriptStruct*> BaseHierarchyTableTypeColumns = Handler->GetColumns();
@@ -325,8 +323,7 @@ void FHierarchyTableEditorToolkit::AddEntry(const FName Identifier, const EHiera
 	}
 
 	using namespace UE::Editor::DataStorage;
-	UTypedElementRegistry* Registry = UTypedElementRegistry::GetInstance();
-	IEditorDataStorageProvider* DSI = Registry->GetMutableDataStorage();
+	IEditorDataStorageProvider* DSI = GetMutableDataStorageFeature<IEditorDataStorageProvider>(StorageFeatureName);
 	static TableHandle Table = DSI->FindTable(FName("Editor_HierarchyTableTable"));
 
 	int32 ParentIndex = 0; // root of hierarchy
diff --git a/Engine/Plugins/Experimental/EditorDataStorage/Source/TedsCore/Private/Compatibility/TedsTypedElementActorHandleFactory.cpp b/Engine/Plugins/Experimental/EditorDataStorage/Source/TedsCore/Private/Compatibility/TedsTypedElementActorHandleFactory.cpp
--- a/Engine/Plugins/Experimental/EditorDataStorage/Source/TedsCore/Private/Compatibility/TedsTypedElementActorHandleFactory.cpp
+++ b/Engine/Plugins/Experimental/EditorDataStorage/Source/TedsCore/Private/Compatibility/TedsTypedElementActorHandleFactory.cpp
@@ -5,10 +5,11 @@
 #include "Compatibility/TedsTypedElementBridge.h"
 #include "Compatibility/Columns/TypedElement.h"
 #include "Elements/Columns/TypedElementCompatibilityColumns.h"
+#include "Elements/Common/EditorDataStorageFeatures.h"
 #include "Elements/Framework/EngineElementsLibrary.h"
 #include "Elements/Framework/TypedElementList.h"
 #include "Elements/Framework/TypedElementQueryBuilder.h"
-#include "Elements/Framework/TypedElementRegistry.h"
+#include "Elements/Interfaces/TypedElementDataStorageInterface.h"
 #include "GameFramework/Actor.h"
 
 void UTypedElementActorHandleDataStorageFactory::PreRegister(IEditorDataStorageProvider& DataStorage)
@@ -75,7 +76,7 @@ void UTypedElementActorHandleDataStorageFactory::HandleBridgeEnabled(bool bEnabl
 {
 	using namespace UE::Editor::DataStorage;
 
-	IEditorDataStorageProvider* DataStorage = UTypedElementRegistry::GetInstance()->GetMutableDataStorage();
+	IEditorDataStorageProvider* DataStorage = GetMutableDataStorageFeature<IEditorDataStorageProvider>(StorageFeatureName);
 
 	if (bEnabled)
 	{
diff --git a/Engine/Plugins/Experimental/EditorDataStorage/Source/TedsCore/Private/Compatibility/TedsTypedElementBridgeQueries.cpp b/Engine/Plugins/Experimental/EditorDataStorage/Source/TedsCore/Private/Compatibility/TedsTypedElementBridgeQueries.cpp
--- a/Engine/Plugins/Experimental/EditorDataStorage/Source/TedsCore/Private/Compatibility/TedsTypedElementBridgeQueries.cpp
+++ b/Engine/Plugins/Experimental/EditorDataStorage/Source/TedsCore/Private/Compatibility/TedsTypedElementBridgeQueries.cpp
@@ -5,10 +5,11 @@
 #include "Compatibility/TedsTypedElementBridge.h"
 #include "Compatibility/Columns/TypedElement.h"
 #include "Elements/Columns/TypedElementCompatibilityColumns.h"
+#include "Elements/Common/EditorDataStorageFeatures.h"
 #include "Elements/Framework/EngineElementsLibrary.h"
 #include "Elements/Framework/TypedElementList.h"
 #include "Elements/Framework/TypedElementQueryBuilder.h"
-#include "Elements/Framework/TypedElementRegistry.h"
+#include "Elements/Interfaces/TypedElementDataStorageInterface.h"
 #include "HAL/IConsoleManager.h"
 
 namespace UE::Editor::DataStorage::Compatibility::Private
@@ -90,17 +91,18 @@ void UTypedElementBridgeDataStorageFactory::CleanupTypedElementColumns(IEditorDa
 
 void UTypedElementBridgeDataStorageFactory::HandleOnEnabled(IConsoleVariable* CVar)
 {
-	IEditorDataStorageProvider* DataStorage = UTypedElementRegistry::GetInstance()->GetMutableDataStorage();
+	using namespace UE::Editor::DataStorage;
+	IEditorDataStorageProvider* DataStorage = GetMutableDataStorageFeature<IEditorDataStorageProvider>(StorageFeatureName);
 	bool bIsEnabled = CVar->GetBool();
 
 	if (bIsEnabled)
 	{
 		RegisterQuery_NewUObject(*DataStorage);
-		UE::Editor::DataStorage::Compatibility::OnTypedElementBridgeEnabled().Broadcast(bIsEnabled);
+		Compatibility::OnTypedElementBridgeEnabled().Broadcast(bIsEnabled);
 	}
 	else
 	{
-		UE::Editor::DataStorage::Compatibility::OnTypedElementBridgeEnabled().Broadcast(bIsEnabled);
+		Compatibility::OnTypedElementBridgeEnabled().Broadcast(bIsEnabled);
 		UnregisterQuery_NewUObject(*DataStorage);
 		CleanupTypedElementColumns(*DataStorage);
 	}
diff --git a/Engine/Plugins/Experimental/EditorDataStorage/Source/TedsCore/Private/Compatibility/TypedElementActorViewportProcessors.cpp b/Engine/Plugins/Experimental/EditorDataStorage/Source/TedsCore/Private/Compatibility/TypedElementActorViewportProcessors.cpp
--- a/Engine/Plugins/Experimental/EditorDataStorage/Source/TedsCore/Private/Compatibility/TypedElementActorViewportProcessors.cpp
+++ b/Engine/Plugins/Experimental/EditorDataStorage/Source/TedsCore/Private/Compatibility/TypedElementActorViewportProcessors.cpp
@@ -6,8 +6,9 @@
 #include "Elements/Columns/TypedElementCompatibilityColumns.h"
 #include "Elements/Columns/TypedElementViewportColumns.h"
 #include "Elements/Columns/TypedElementMiscColumns.h"
+#include "Elements/Common/EditorDataStorageFeatures.h"
 #include "Elements/Framework/TypedElementQueryBuilder.h"
-#include "Elements/Framework/TypedElementRegistry.h"
+#include "Elements/Interfaces/TypedElementDataStorageInterface.h"
 #include "HAL/IConsoleManager.h"
 #include "Elements/Columns/TypedElementSelectionColumns.h"
 #include "ProfilingDebugging/CpuProfilerTrace.h"
@@ -23,7 +24,7 @@ namespace UE::Editor::DataStorage::Private
 
 				TRACE_CPUPROFILER_EVENT_SCOPE(TEDS.Debug.AddOverlayColorToSelectionCommand);
 
-				if (IEditorDataStorageProvider* DataStorage = UTypedElementRegistry::GetInstance()->GetMutableDataStorage())
+				if (IEditorDataStorageProvider* DataStorage = GetMutableDataStorageFeature<IEditorDataStorageProvider>(StorageFeatureName))
 				{
 					static QueryHandle OverlayQuery = InvalidQueryHandle;
 					if (OverlayQuery == InvalidQueryHandle)
@@ -81,7 +82,7 @@ namespace UE::Editor::DataStorage::Private
 
 				TRACE_CPUPROFILER_EVENT_SCOPE(TEDS.Debug.AddOverlayColorToSelectionCommand);
 
-				if (IEditorDataStorageProvider* DataStorage = UTypedElementRegistry::GetInstance()->GetMutableDataStorage())
+				if (IEditorDataStorageProvider* DataStorage = GetMutableDataStorageFeature<IEditorDataStorageProvider>(StorageFeatureName))
 				{
 					static QueryHandle OverlayQuery = InvalidQueryHandle;
 					if (OverlayQuery == InvalidQueryHandle)
@@ -132,7 +133,7 @@ namespace UE::Editor::DataStorage::Private
 
 				TRACE_CPUPROFILER_EVENT_SCOPE(TEDS.Debug.AddOverlayColorToSelectionCommand);
 
-				if (IEditorDataStorageProvider* DataStorage = UTypedElementRegistry::GetInstance()->GetMutableDataStorage())
+				if (IEditorDataStorageProvider* DataStorage = GetMutableDataStorageFeature<IEditorDataStorageProvider>(StorageFeatureName))
 				{
 					static QueryHandle OverlayQuery = InvalidQueryHandle;
 					if (OverlayQuery == InvalidQueryHandle)
diff --git a/Engine/Plugins/Experimental/EditorDataStorage/Source/TedsCore/Private/Debug/TypedElementDatabaseDebugCommands.cpp b/Engine/Plugins/Experimental/EditorDataStorage/Source/TedsCore/Private/Debug/TypedElementDatabaseDebugCommands.cpp
--- a/Engine/Plugins/Experimental/EditorDataStorage/Source/TedsCore/Private/Debug/TypedElementDatabaseDebugCommands.cpp
+++ b/Engine/Plugins/Experimental/EditorDataStorage/Source/TedsCore/Private/Debug/TypedElementDatabaseDebugCommands.cpp
@@ -4,9 +4,9 @@
 #include "Elements/Columns/TypedElementCompatibilityColumns.h"
 #include "Elements/Columns/TypedElementLabelColumns.h"
 #include "Elements/Columns/TypedElementTypeInfoColumns.h"
+#include "Elements/Common/EditorDataStorageFeatures.h"
 #include "Elements/Common/TypedElementDataStorageLog.h"
 #include "Elements/Framework/TypedElementQueryBuilder.h"
-#include "Elements/Framework/TypedElementRegistry.h"
 #include "Elements/Framework/TypedElementTestColumns.h"
 #include "Elements/Interfaces/TypedElementDataStorageInterface.h"
 #include "Elements/Interfaces/TypedElementDataStorageCompatibilityInterface.h"
@@ -65,7 +65,7 @@ namespace UE::Editor::DataStorage::Private
 	{
 		using namespace UE::Editor::DataStorage::Queries;
 
-		if (IEditorDataStorageProvider* DataStorage = UTypedElementRegistry::GetInstance()->GetMutableDataStorage())
+		if (IEditorDataStorageProvider* DataStorage = GetMutableDataStorageFeature<IEditorDataStorageProvider>(StorageFeatureName))
 		{
 			static QueryHandle LabelQuery = [DataStorage]
 			{
@@ -136,7 +136,7 @@ FAutoConsoleCommandWithOutputDevice PrintObjectTypeInformationConsoleCommand(
 
 			TRACE_CPUPROFILER_EVENT_SCOPE(TEDS.Debug.PrintObjectTypeInfo);
 
-			if (IEditorDataStorageProvider* DataStorage = UTypedElementRegistry::GetInstance()->GetMutableDataStorage())
+			if (IEditorDataStorageProvider* DataStorage = GetMutableDataStorageFeature<IEditorDataStorageProvider>(StorageFeatureName))
 			{
 				FString Message;
 				Output.Log(TEXT("The Typed Elements Data Storage has the types:"));
@@ -179,6 +179,7 @@ FAutoConsoleCommandWithOutputDevice ListExtensionsConsoleCommand(
 	FConsoleCommandWithOutputDeviceDelegate::CreateLambda([](FOutputDevice& Output)
 		{
 			TRACE_CPUPROFILER_EVENT_SCOPE(TEDS.Debug.ListExtensions);
+			using namespace UE::Editor::DataStorage;
 
 			FString Message;
 			auto RecordExtensions = [&Message](FName Extension)
@@ -188,19 +189,17 @@ FAutoConsoleCommandWithOutputDevice ListExtensionsConsoleCommand(
 				Message += TEXT('\n');
 			};
 
-			UTypedElementRegistry* Registry = UTypedElementRegistry::GetInstance();
-
-			if (const IEditorDataStorageProvider* DataStorage = Registry->GetDataStorage())
+			if (const IEditorDataStorageProvider* DataStorage = GetDataStorageFeature<IEditorDataStorageProvider>(StorageFeatureName))
 			{
 				Message = TEXT("Data Storage Extensions: \n");
 				DataStorage->ListExtensions(RecordExtensions);
 			}
-			if (const IEditorDataStorageCompatibilityProvider* DataStorageCompat = Registry->GetDataStorageCompatibility())
+			if (const IEditorDataStorageCompatibilityProvider* DataStorageCompat = GetDataStorageFeature<IEditorDataStorageCompatibilityProvider>(CompatibilityFeatureName))
 			{
 				Message += TEXT("Data Storage Compatibility Extensions: \n");
 				DataStorageCompat->ListExtensions(RecordExtensions);
 			}
-			if (const IEditorDataStorageUiProvider* DataStorageUi = Registry->GetDataStorageUi())
+			if (const IEditorDataStorageUiProvider* DataStorageUi = GetDataStorageFeature<IEditorDataStorageUiProvider>(UiFeatureName))
 			{
 				Message += TEXT("Data Storage UI Extensions: \n");
 				DataStorageUi->ListExtensions(RecordExtensions);
@@ -217,10 +216,11 @@ static FAutoConsoleCommand CVarCreateRow(
 	TEXT("Argument: \n"),
 	FConsoleCommandWithArgsDelegate::CreateLambda([](const TArray<FString>& Args)
 	{
-		IEditorDataStorageProvider* DataStorage = UTypedElementRegistry::GetInstance()->GetMutableDataStorage();
-		static UE::Editor::DataStorage::TableHandle Table = DataStorage->RegisterTable<FTestColumnA>(FName(TEXT("Debug.CreateRow Table")));
+		using namespace UE::Editor::DataStorage;
+		IEditorDataStorageProvider* DataStorage = GetMutableDataStorageFeature<IEditorDataStorageProvider>(StorageFeatureName);
+		static TableHandle Table = DataStorage->RegisterTable<FTestColumnA>(FName(TEXT("Debug.CreateRow Table")));
 
-		const UE::Editor::DataStorage::RowHandle RowHandle = DataStorage->AddRow(Table);
+		const RowHandle RowHandle = DataStorage->AddRow(Table);
 
 		UE_LOG(LogEditorDataStorage, Warning, TEXT("Added Row %llu"), static_cast<uint64>(RowHandle));
 	}));
@@ -230,7 +230,8 @@ static FAutoConsoleCommand CVarAddDynamicColumnTag(
 	TEXT("Argument: Row, Identifier\n"),
 	FConsoleCommandWithArgsDelegate::CreateLambda([](const TArray<FString>& Args)
 	{
-		IEditorDataStorageProvider* DataStorage = UTypedElementRegistry::GetInstance()->GetMutableDataStorage();
+		using namespace UE::Editor::DataStorage;
+		IEditorDataStorageProvider* DataStorage = GetMutableDataStorageFeature<IEditorDataStorageProvider>(StorageFeatureName);
 		
 		if (Args.Num() != 2)
 		{
@@ -239,7 +240,7 @@ static FAutoConsoleCommand CVarAddDynamicColumnTag(
 		}
 
 		const uint64 RowAsU64 = FCString::Strtoui64(*Args[0], nullptr, 10);
-		const UE::Editor::DataStorage::RowHandle Row = RowAsU64;
+		const RowHandle Row = RowAsU64;
 		const FName Identifier = FName(*Args[1]);
 
 		DataStorage->AddColumn<FTestDynamicTag>(Row, Identifier);
@@ -252,7 +253,8 @@ static FAutoConsoleCommand CVarAddDynamicColumn(
 	TEXT("Argument: Row, Identifier\n"),
 	FConsoleCommandWithArgsDelegate::CreateLambda([](const TArray<FString>& Args)
 	{
-		IEditorDataStorageProvider* DataStorage = UTypedElementRegistry::GetInstance()->GetMutableDataStorage();
+		using namespace UE::Editor::DataStorage;
+		IEditorDataStorageProvider* DataStorage = GetMutableDataStorageFeature<IEditorDataStorageProvider>(StorageFeatureName);
 		
 		if (Args.Num() != 2)
 		{
@@ -261,7 +263,7 @@ static FAutoConsoleCommand CVarAddDynamicColumn(
 		}
 
 		const uint64 RowAsU64 = FCString::Strtoui64(*Args[0], nullptr, 10);
-		const UE::Editor::DataStorage::RowHandle Row = RowAsU64;
+		const RowHandle Row = RowAsU64;
 		const FName Identifier = FName(*Args[1]);
 		
 		bool bUseDefaultApi = false;
@@ -283,7 +285,8 @@ static FAutoConsoleCommand CVarRemoveDynamicColumn(
 	TEXT("Argument: Row, Identifier\n"),
 	FConsoleCommandWithArgsDelegate::CreateLambda([](const TArray<FString>& Args)
 	{
-		IEditorDataStorageProvider* DataStorage = UTypedElementRegistry::GetInstance()->GetMutableDataStorage();
+		using namespace UE::Editor::DataStorage;
+		IEditorDataStorageProvider* DataStorage = GetMutableDataStorageFeature<IEditorDataStorageProvider>(StorageFeatureName);
 		
 		if (Args.Num() != 2)
 		{
@@ -292,7 +295,7 @@ static FAutoConsoleCommand CVarRemoveDynamicColumn(
 		}
 
 		const uint64 RowAsU64 = FCString::Strtoui64(*Args[0], nullptr, 10);
-		const UE::Editor::DataStorage::RowHandle Row = RowAsU64;
+		const RowHandle Row = RowAsU64;
 		const FName Identifier = FName(*Args[1]);
 		
 		DataStorage->RemoveColumn<FTestDynamicColumn>(Row, Identifier);
@@ -305,7 +308,8 @@ static FAutoConsoleCommand CVarAddToDynamicColumn(
 	TEXT("Argument: Row, TagId, Value, [optional] MethodId\n"),
 	FConsoleCommandWithArgsDelegate::CreateLambda([](const TArray<FString>& Args)
 	{
-		IEditorDataStorageProvider* DataStorage = UTypedElementRegistry::GetInstance()->GetMutableDataStorage();
+		using namespace UE::Editor::DataStorage;
+		IEditorDataStorageProvider* DataStorage = GetMutableDataStorageFeature<IEditorDataStorageProvider>(StorageFeatureName);
 		
 		if (Args.Num() < 3 || Args.Num() > 4)
 		{
@@ -314,7 +318,7 @@ static FAutoConsoleCommand CVarAddToDynamicColumn(
 		}
 
 		const uint64 RowAsU64 = FCString::Strtoui64(*Args[0], nullptr, 10);
-		const UE::Editor::DataStorage::RowHandle Row = RowAsU64;
+		const RowHandle Row = RowAsU64;
 		const FName TagId = FName(*Args[1]);
 		const uint64 Value = FCString::Strtoui64(*Args[2], nullptr, 10);
 		const uint64 MethodId = Args.Num() >= 4 ? FCString::Strtoui64(*Args[3], nullptr, 10) : 0;
@@ -357,7 +361,8 @@ static FAutoConsoleCommand CVarPrintDynamicColumn(
 	TEXT("Argument: Row, TagId\n"),
 	FConsoleCommandWithArgsDelegate::CreateLambda([](const TArray<FString>& Args)
 	{
-		IEditorDataStorageProvider* DataStorage = UTypedElementRegistry::GetInstance()->GetMutableDataStorage();
+		using namespace UE::Editor::DataStorage;
+		IEditorDataStorageProvider* DataStorage = GetMutableDataStorageFeature<IEditorDataStorageProvider>(StorageFeatureName);
 		
 		if (Args.Num() != 2)
 		{
@@ -365,7 +370,7 @@ static FAutoConsoleCommand CVarPrintDynamicColumn(
 			return;
 		}
 		const uint64 RowAsU64 = FCString::Strtoui64(*Args[0], nullptr, 10);
-		const UE::Editor::DataStorage::RowHandle Row = RowAsU64;
+		const RowHandle Row = RowAsU64;
 		const FName TagId = FName(*Args[1]);
 
 		FTestDynamicColumn* Column = DataStorage->GetColumn<FTestDynamicColumn>(Row, TagId);
@@ -413,7 +418,7 @@ static FAutoConsoleCommand CVarPrintDynamicColumnWithQuery(
 			return;
 		}
 		
-		IEditorDataStorageProvider* DataStorage = UTypedElementRegistry::GetInstance()->GetMutableDataStorage();
+		IEditorDataStorageProvider* DataStorage = GetMutableDataStorageFeature<IEditorDataStorageProvider>(StorageFeatureName);
 		
 		const FName Identifier(*Args[0]);
 
@@ -424,8 +429,8 @@ static FAutoConsoleCommand CVarPrintDynamicColumnWithQuery(
 			Compile());
 
 		TStringBuilder<1024> StringBuilder;
-		UE::Editor::DataStorage::FQueryResult Result = DataStorage->RunQuery(Query,CreateDirectQueryCallbackBinding(
-			[Identifier, &StringBuilder](UE::Editor::DataStorage::IDirectQueryContext& Context, const RowHandle* Rows)
+		FQueryResult Result = DataStorage->RunQuery(Query,CreateDirectQueryCallbackBinding(
+			[Identifier, &StringBuilder](IDirectQueryContext& Context, const RowHandle* Rows)
 			{
 				const TArrayView<const RowHandle> RowView = MakeConstArrayView(Rows, Context.GetRowCount());
 				// Get pointer to the start of the range of columns to process
@@ -473,7 +478,7 @@ static FAutoConsoleCommand CVarCountDynamicTagWithQuery(
 			return;
 		}
 				
-		IEditorDataStorageProvider* DataStorage = UTypedElementRegistry::GetInstance()->GetMutableDataStorage();
+		IEditorDataStorageProvider* DataStorage = GetMutableDataStorageFeature<IEditorDataStorageProvider>(StorageFeatureName);
 				
 		const FName Identifier(*Args[0]);
 
@@ -514,7 +519,7 @@ static FAutoConsoleCommand CVarRegisterListDynamicColumnQuery(
 			return;
 		}
 				
-		IEditorDataStorageProvider* DataStorage = UTypedElementRegistry::GetInstance()->GetMutableDataStorage();
+		IEditorDataStorageProvider* DataStorage = GetMutableDataStorageFeature<IEditorDataStorageProvider>(StorageFeatureName);
 				
 		const FName Identifier(*Args[0]);
 		const FName ActivationGroup(*Args[1]);
@@ -552,8 +557,9 @@ static FAutoConsoleCommand CVarActivateListDynamicColumnQuery(
 			UE_LOG(LogEditorDataStorage, Error, TEXT("Invalid number of arguments. ActivationGroup"));
 			return;
 		}
-						
-		IEditorDataStorageProvider* DataStorage = UTypedElementRegistry::GetInstance()->GetMutableDataStorage();
+
+		using namespace UE::Editor::DataStorage;
+		IEditorDataStorageProvider* DataStorage = GetMutableDataStorageFeature<IEditorDataStorageProvider>(StorageFeatureName);
 		
 		const FName ActivationGroup(*Args[0]);
 
@@ -567,7 +573,7 @@ static FAutoConsoleCommand CVarAddValueTag(
 	FConsoleCommandWithArgsDelegate::CreateLambda([](const TArray<FString>& Args)
 	{
 		using namespace UE::Editor::DataStorage;
-		IEditorDataStorageProvider* DataStorage = UTypedElementRegistry::GetInstance()->GetMutableDataStorage();
+		IEditorDataStorageProvider* DataStorage = GetMutableDataStorageFeature<IEditorDataStorageProvider>(StorageFeatureName);
 
 		if (Args.Num() != 3)
 		{
@@ -601,7 +607,8 @@ static FAutoConsoleCommand CVarRemoveValueTag(
 	TEXT("Argument: Row, Group\n"),
 	FConsoleCommandWithArgsDelegate::CreateLambda([](const TArray<FString>& Args)
 	{
-		IEditorDataStorageProvider* DataStorage = UTypedElementRegistry::GetInstance()->GetMutableDataStorage();
+		using namespace UE::Editor::DataStorage;
+		IEditorDataStorageProvider* DataStorage = GetMutableDataStorageFeature<IEditorDataStorageProvider>(StorageFeatureName);
 
 		if (Args.Num() != 2)
 		{
@@ -633,7 +640,7 @@ static FAutoConsoleCommand CVarMatchValueTag(
 	{
 		using namespace UE::Editor::DataStorage::Queries;
 		
-		IEditorDataStorageProvider* DataStorage = UTypedElementRegistry::GetInstance()->GetMutableDataStorage();
+		IEditorDataStorageProvider* DataStorage = GetMutableDataStorageFeature<IEditorDataStorageProvider>(StorageFeatureName);
 
 		if (Args.Num() < 1 || Args.Num() > 2)
 		{
@@ -687,7 +694,7 @@ static FAutoConsoleCommand CVarAddValueTagFromEnum(
 	{
 		using namespace UE::Editor::DataStorage;
 
-		IEditorDataStorageProvider* DataStorage = UTypedElementRegistry::GetInstance()->GetMutableDataStorage();
+		IEditorDataStorageProvider* DataStorage = GetMutableDataStorageFeature<IEditorDataStorageProvider>(StorageFeatureName);
 
 		if (Args.Num() < 1 || Args.Num() > 2)
 		{
@@ -725,7 +732,7 @@ static FAutoConsoleCommand CVarRemoveValueTagFromEnum(
 	FConsoleCommandWithArgsDelegate::CreateLambda([](const TArray<FString>& Args)
 	{
 		using namespace UE::Editor::DataStorage;
-		IEditorDataStorageProvider* DataStorage = UTypedElementRegistry::GetInstance()->GetMutableDataStorage();
+		IEditorDataStorageProvider* DataStorage = GetMutableDataStorageFeature<IEditorDataStorageProvider>(StorageFeatureName);
 
 		if (Args.Num() != 1)
 		{
@@ -747,7 +754,7 @@ static FAutoConsoleCommand CVarMatchValueTagFromEnum(
 	{
 		using namespace UE::Editor::DataStorage::Queries;
 		
-		IEditorDataStorageProvider* DataStorage = UTypedElementRegistry::GetInstance()->GetMutableDataStorage();
+		IEditorDataStorageProvider* DataStorage = GetMutableDataStorageFeature<IEditorDataStorageProvider>(StorageFeatureName);
 
 		if (Args.Num() > 1)
 		{
@@ -816,4 +823,4 @@ static FAutoConsoleCommand CVarMatchValueTagFromEnum(
 		UE_LOG(LogEditorDataStorage, Warning, TEXT("Processed %llu rows"), static_cast<uint64>(Count));
 		
 	}),
-	ECVF_Default);
\ No newline at end of file
+	ECVF_Default);
diff --git a/Engine/Plugins/Experimental/EditorDataStorage/Source/TedsCore/Private/Queries/TypedElementAlertQueries.cpp b/Engine/Plugins/Experimental/EditorDataStorage/Source/TedsCore/Private/Queries/TypedElementAlertQueries.cpp
--- a/Engine/Plugins/Experimental/EditorDataStorage/Source/TedsCore/Private/Queries/TypedElementAlertQueries.cpp
+++ b/Engine/Plugins/Experimental/EditorDataStorage/Source/TedsCore/Private/Queries/TypedElementAlertQueries.cpp
@@ -6,8 +6,9 @@
 #include "Elements/Columns/TypedElementHiearchyColumns.h"
 #include "Elements/Columns/TypedElementMiscColumns.h"
 #include "Elements/Columns/TypedElementSelectionColumns.h"
+#include "Elements/Common/EditorDataStorageFeatures.h"
 #include "Elements/Framework/TypedElementQueryBuilder.h"
-#include "Elements/Framework/TypedElementRegistry.h"
+#include "Elements/Interfaces/TypedElementDataStorageInterface.h"
 #include "HAL/IConsoleManager.h"
 #include "Math/UnrealMathUtility.h"
 
@@ -21,7 +22,7 @@ FAutoConsoleCommand AddRandomAlertToRowConsoleCommand(
 
 			TRACE_CPUPROFILER_EVENT_SCOPE(TEDS.Debug.AddRandomAlertToSelectedRows);
 
-			if (IEditorDataStorageProvider* DataStorage = UTypedElementRegistry::GetInstance()->GetMutableDataStorage())
+			if (IEditorDataStorageProvider* DataStorage = GetMutableDataStorageFeature<IEditorDataStorageProvider>(StorageFeatureName))
 			{
 				static QueryHandle Query = [DataStorage]
 				{
@@ -82,7 +83,7 @@ FAutoConsoleCommand ClearAllAlertsConsoleCommand(
 			static TTypedElementColumnTypeList<FTypedElementSyncBackToWorldTag> BatchAddColumns;
 			static TTypedElementColumnTypeList<FTypedElementAlertColumn, FTypedElementChildAlertColumn> BatchRemoveColumns;
 
-			if (IEditorDataStorageProvider* DataStorage = UTypedElementRegistry::GetInstance()->GetMutableDataStorage())
+			if (IEditorDataStorageProvider* DataStorage = GetMutableDataStorageFeature<IEditorDataStorageProvider>(StorageFeatureName))
 			{
 				static QueryHandle AlertInfoQuery = [DataStorage]
 				{
@@ -118,7 +119,7 @@ FAutoConsoleCommand ClearSelectedAlertsConsoleCommand(
 
 			TRACE_CPUPROFILER_EVENT_SCOPE(TEDS.Debug.ClearSelectedAlerts);
 
-			if (IEditorDataStorageProvider* DataStorage = UTypedElementRegistry::GetInstance()->GetMutableDataStorage())
+			if (IEditorDataStorageProvider* DataStorage = GetMutableDataStorageFeature<IEditorDataStorageProvider>(StorageFeatureName))
 			{
 				static QueryHandle AlertQuery = [DataStorage]
 				{
diff --git a/Engine/Plugins/Experimental/EditorDataStorage/Source/TedsCore/Private/TypedElementDatabase.cpp b/Engine/Plugins/Experimental/EditorDataStorage/Source/TedsCore/Private/TypedElementDatabase.cpp
--- a/Engine/Plugins/Experimental/EditorDataStorage/Source/TedsCore/Private/TypedElementDatabase.cpp
+++ b/Engine/Plugins/Experimental/EditorDataStorage/Source/TedsCore/Private/TypedElementDatabase.cpp
@@ -4,8 +4,9 @@
 
 #include "Editor.h"
 #include "EditorDataStorageSettings.h"
+#include "Elements/Common/EditorDataStorageFeatures.h"
 #include "Elements/Interfaces/TypedElementDataStorageFactory.h"
-#include "Elements/Framework/TypedElementRegistry.h"
+#include "Elements/Interfaces/TypedElementDataStorageInterface.h"
 #include "Engine/World.h"
 #include "GlobalLock.h"
 #include "MassEntityEditorSubsystem.h"
@@ -32,12 +33,10 @@ FAutoConsoleCommandWithOutputDevice PrintQueryCallbacksConsoleCommand(
 	TEXT("Prints out a list of all processors."),
 	FConsoleCommandWithOutputDeviceDelegate::CreateLambda([](FOutputDevice& Output)
 		{
-			if (UTypedElementRegistry* Registry = UTypedElementRegistry::GetInstance())
+			using namespace UE::Editor::DataStorage;
+			if (IEditorDataStorageProvider* DataStorage = GetMutableDataStorageFeature<IEditorDataStorageProvider>(StorageFeatureName))
 			{
-				if (IEditorDataStorageProvider* DataStorage = Registry->GetMutableDataStorage())
-				{
-					DataStorage->DebugPrintQueryCallbacks(Output);
-				}
+				DataStorage->DebugPrintQueryCallbacks(Output);
 			}
 		}));
 
diff --git a/Engine/Plugins/Experimental/EditorDataStorage/Source/TedsCore/Tests/DynamicColumnTests.cpp b/Engine/Plugins/Experimental/EditorDataStorage/Source/TedsCore/Tests/DynamicColumnTests.cpp
--- a/Engine/Plugins/Experimental/EditorDataStorage/Source/TedsCore/Tests/DynamicColumnTests.cpp
+++ b/Engine/Plugins/Experimental/EditorDataStorage/Source/TedsCore/Tests/DynamicColumnTests.cpp
@@ -1,12 +1,15 @@
 // Copyright Epic Games, Inc. All Rights Reserved.
 
-#include "Elements/Framework/TypedElementRegistry.h"
 #include "Elements/Framework/TypedElementTestColumns.h"
 #include "Elements/Interfaces/TypedElementDataStorageInterface.h"
 #include "Debug/TypedElementDatabaseDebugTypes.h"
 #include "Misc/CoreDelegates.h"
 #if WITH_TESTS
 
+#include "Elements/Common/EditorDataStorageFeatures.h"
+#include "Elements/Framework/TypedElementQueryBuilder.h"
+#include "Elements/Interfaces/TypedElementDataStorageInterface.h"
+
 #include "Misc/AutomationTest.h"
 
 namespace UE::Editor::DataStorage::Tests
@@ -98,8 +101,7 @@ namespace UE::Editor::DataStorage::Tests
 	{
 		BeforeEach([this]()
 		{
-			UTypedElementRegistry* TypedElementRegistry = UTypedElementRegistry::GetInstance();
-			TedsInterface = TypedElementRegistry->GetMutableDataStorage();
+			TedsInterface = GetMutableDataStorageFeature<IEditorDataStorageProvider>(StorageFeatureName);
 			TestTable = RegisterTestTable();
 			Identifiers = {TEXT("StaticMesh"), TEXT("Animation"), TEXT("AudioClip")};
 			TestTrue("", TedsInterface != nullptr);
diff --git a/Engine/Plugins/Experimental/EditorDataStorage/Source/TedsCore/Tests/TypedElementAttributeBindingTests.cpp b/Engine/Plugins/Experimental/EditorDataStorage/Source/TedsCore/Tests/TypedElementAttributeBindingTests.cpp
--- a/Engine/Plugins/Experimental/EditorDataStorage/Source/TedsCore/Tests/TypedElementAttributeBindingTests.cpp
+++ b/Engine/Plugins/Experimental/EditorDataStorage/Source/TedsCore/Tests/TypedElementAttributeBindingTests.cpp
@@ -3,8 +3,9 @@
 
 #include "Elements/Framework/TypedElementAttributeBinding.h"
 #if WITH_TESTS
+#include "Elements/Common/EditorDataStorageFeatures.h"
 #include "Elements/Framework/TypedElementTestColumns.h"
-#include "Elements/Framework/TypedElementRegistry.h"
+#include "Elements/Interfaces/TypedElementDataStorageInterface.h"
 #include "Misc/AutomationTest.h"
 
 namespace UE::Editor::DataStorage::Tests
@@ -50,8 +51,7 @@ namespace UE::Editor::DataStorage::Tests
 	{
 		BeforeEach([this]()
 		{
-			UTypedElementRegistry* TypedElementRegistry = UTypedElementRegistry::GetInstance();
-			TedsInterface = TypedElementRegistry->GetMutableDataStorage();
+			TedsInterface = GetMutableDataStorageFeature<IEditorDataStorageProvider>(StorageFeatureName);
 			TestTrue("", TedsInterface != nullptr);
 			
 			TestTableHandle = RegisterTestTable();
diff --git a/Engine/Plugins/Experimental/EditorDataStorage/Source/TedsCore/Tests/TypedElementDataStorageInterfaceTests.cpp b/Engine/Plugins/Experimental/EditorDataStorage/Source/TedsCore/Tests/TypedElementDataStorageInterfaceTests.cpp
--- a/Engine/Plugins/Experimental/EditorDataStorage/Source/TedsCore/Tests/TypedElementDataStorageInterfaceTests.cpp
+++ b/Engine/Plugins/Experimental/EditorDataStorage/Source/TedsCore/Tests/TypedElementDataStorageInterfaceTests.cpp
@@ -2,8 +2,10 @@
 
 
 #if WITH_TESTS
+#include "Elements/Common/EditorDataStorageFeatures.h"
 #include "Elements/Framework/TypedElementTestColumns.h"
-#include "Elements/Framework/TypedElementRegistry.h"
+#include "Elements/Framework/TypedElementQueryBuilder.h"
+#include "Elements/Interfaces/TypedElementDataStorageInterface.h"
 #include "Misc/AutomationTest.h"
 
 namespace UE::Editor::DataStorage::Tests
@@ -126,8 +128,7 @@ namespace UE::Editor::DataStorage::Tests
 	{
 		BeforeEach([this]()
 		{
-			UTypedElementRegistry* TypedElementRegistry = UTypedElementRegistry::GetInstance();
-			TedsInterface = TypedElementRegistry->GetMutableDataStorage();
+			TedsInterface = GetMutableDataStorageFeature<IEditorDataStorageProvider>(StorageFeatureName);
 			TestTrue("", TedsInterface != nullptr);
 		});
 
diff --git a/Engine/Plugins/Experimental/EditorDataStorage/Source/TedsUI/Private/TedsUIModule.cpp b/Engine/Plugins/Experimental/EditorDataStorage/Source/TedsUI/Private/TedsUIModule.cpp
--- a/Engine/Plugins/Experimental/EditorDataStorage/Source/TedsUI/Private/TedsUIModule.cpp
+++ b/Engine/Plugins/Experimental/EditorDataStorage/Source/TedsUI/Private/TedsUIModule.cpp
@@ -2,7 +2,7 @@
 
 #include "TedsUIModule.h"
 
-#include "Elements/Framework/TypedElementRegistry.h"
+#include "Elements/Common/EditorDataStorageFeatures.h"
 #include "Elements/Interfaces/TypedElementDataStorageUiInterface.h"
 #include "HAL/IConsoleManager.h"
 #include "Misc/CoreDelegates.h"
@@ -17,8 +17,7 @@ namespace UE::Editor::DataStorage::Private
 	TEXT("Prints a list of all the known widget purposes."),
 	FConsoleCommandWithOutputDeviceDelegate::CreateLambda([](FOutputDevice& Output)
 		{
-			UTypedElementRegistry* Registry = UTypedElementRegistry::GetInstance();
-			if (IEditorDataStorageUiProvider* UiStorage = Registry->GetMutableDataStorageUi())
+			if (IEditorDataStorageUiProvider* UiStorage = GetMutableDataStorageFeature<IEditorDataStorageUiProvider>(UiFeatureName))
 			{
 				Output.Log(TEXT("The Typed Elements Data Storage has recorded the following widget purposes:"));
 				UiStorage->ListWidgetPurposes(
diff --git a/Engine/Plugins/Experimental/EditorDataStorage/Source/TedsUI/Private/Widgets/CounterWidget.cpp b/Engine/Plugins/Experimental/EditorDataStorage/Source/TedsUI/Private/Widgets/CounterWidget.cpp
--- a/Engine/Plugins/Experimental/EditorDataStorage/Source/TedsUI/Private/Widgets/CounterWidget.cpp
+++ b/Engine/Plugins/Experimental/EditorDataStorage/Source/TedsUI/Private/Widgets/CounterWidget.cpp
@@ -4,9 +4,10 @@
 
 #include "Elements/Columns/TypedElementSlateWidgetColumns.h"
 #include "Elements/Columns/TypedElementValueCacheColumns.h"
+#include "Elements/Common/EditorDataStorageFeatures.h"
 #include "Elements/Framework/TypedElementQueryBuilder.h"
-#include "Elements/Framework/TypedElementRegistry.h"
 #include "Elements/Interfaces/TypedElementDataStorageInterface.h"
+#include "Elements/Interfaces/TypedElementDataStorageUiInterface.h"
 #include "Interfaces/IMainFrameModule.h"
 #include "Layout/Margin.h"
 #include "Misc/CoreDelegates.h"
@@ -111,11 +112,9 @@ void UCounterWidgetFactory::SetupMainWindowIntegrations(TSharedPtr<SWindow> Pare
 {
 	if (!bHasBeenSetup)
 	{
-		UTypedElementRegistry* Registry = UTypedElementRegistry::GetInstance();
-		checkf(Registry, TEXT(
-			"FEditorDataStorageUiModule didn't find the UTypedElementRegistry during main window integration when it should be available."));
+		using namespace UE::Editor::DataStorage;
 
-		IEditorDataStorageUiProvider* UiInterface = Registry->GetMutableDataStorageUi();
+		IEditorDataStorageUiProvider* UiInterface = GetMutableDataStorageFeature<IEditorDataStorageUiProvider>(UiFeatureName);
 		checkf(UiInterface, TEXT(
 			"FEditorDataStorageUiModule tried to integrate with the main window before the "
 			"TEDS UI interface is available."));
diff --git a/Engine/Plugins/Experimental/EditorDataStorageFeatures/Source/TedsAssetData/Private/Debug/TedsAssetDataDebugCommands.cpp b/Engine/Plugins/Experimental/EditorDataStorageFeatures/Source/TedsAssetData/Private/Debug/TedsAssetDataDebugCommands.cpp
--- a/Engine/Plugins/Experimental/EditorDataStorageFeatures/Source/TedsAssetData/Private/Debug/TedsAssetDataDebugCommands.cpp
+++ b/Engine/Plugins/Experimental/EditorDataStorageFeatures/Source/TedsAssetData/Private/Debug/TedsAssetDataDebugCommands.cpp
@@ -4,8 +4,8 @@
 #include "AssetRegistry/IAssetRegistry.h"
 #include "Containers/UnrealString.h"
 #include "Elements/Common/TypedElementHandles.h"
+#include "Elements/Common/EditorDataStorageFeatures.h"
 #include "Elements/Framework/TypedElementIndexHasher.h"
-#include "Elements/Framework/TypedElementRegistry.h"
 #include "HAL/IConsoleManager.h"
 #include "TedsAssetDataColumns.h"
 #include "UObject/NameTypes.h"
@@ -30,7 +30,7 @@ namespace UE::Editor::DataStorage::Debug::Private
 		TEXT("Print some debug information on the specified path."),
 		FConsoleCommandWithArgsDelegate::CreateLambda([](const TArray<FString>& InArgs)
 			{
-				const IEditorDataStorageProvider* Database = UTypedElementRegistry::GetInstance()->GetDataStorage();
+				const IEditorDataStorageProvider* Database = GetMutableDataStorageFeature<IEditorDataStorageProvider>(StorageFeatureName);
 				IAssetRegistry& AssetRegistry = IAssetRegistry::GetChecked();
 
 				for (const FString& Path : InArgs)
@@ -220,7 +220,7 @@ namespace UE::Editor::DataStorage::Debug::Private
 				TRACE_CPUPROFILER_EVENT_SCOPE(TEDS.Debug.CreateDebugAssetRows);
 
 
-				IEditorDataStorageProvider* DataStorage = UTypedElementRegistry::GetInstance()->GetMutableDataStorage();
+				IEditorDataStorageProvider* DataStorage = GetMutableDataStorageFeature<IEditorDataStorageProvider>(StorageFeatureName);
 
 				if (!DataStorage)
 				{
@@ -252,7 +252,7 @@ namespace UE::Editor::DataStorage::Debug::Private
 			{
 				using namespace UE::Editor::DataStorage::Queries;
 
-				IEditorDataStorageProvider* DataStorage = UTypedElementRegistry::GetInstance()->GetMutableDataStorage();
+				IEditorDataStorageProvider* DataStorage = GetMutableDataStorageFeature<IEditorDataStorageProvider>(StorageFeatureName);
 
 				if (!DataStorage)
 				{
diff --git a/Engine/Plugins/Experimental/EditorDataStorageFeatures/Source/TedsAssetData/Private/TedsAssetDataModule.cpp b/Engine/Plugins/Experimental/EditorDataStorageFeatures/Source/TedsAssetData/Private/TedsAssetDataModule.cpp
--- a/Engine/Plugins/Experimental/EditorDataStorageFeatures/Source/TedsAssetData/Private/TedsAssetDataModule.cpp
+++ b/Engine/Plugins/Experimental/EditorDataStorageFeatures/Source/TedsAssetData/Private/TedsAssetDataModule.cpp
@@ -3,7 +3,8 @@
 #include "TedsAssetDataModule.h"
 
 #include "CB/TedsAssetDataCBDataSource.h"
-#include "Elements/Framework/TypedElementRegistry.h"
+#include "Elements/Common/EditorDataStorageFeatures.h"
+#include "Elements/Interfaces/TypedElementDataStorageInterface.h"
 #include "HAL/IConsoleManager.h"
 #include "Modules/ModuleManager.h"
 #include "TedsAssetData.h"
@@ -42,10 +43,7 @@ void FTedsAssetDataModule::StartupModule()
 
 void FTedsAssetDataModule::ShutdownModule()
 {
-	if (UTypedElementRegistry* TypedElementRegistry = UTypedElementRegistry::GetInstance())
-	{
-		TypedElementRegistry->OnDataStorageInterfacesSet().RemoveAll(this);
-	}
+	UE::Editor::DataStorage::OnEditorDataStorageFeaturesEnabled().RemoveAll(this);
 }
 
 FTedsAssetDataModule* FTedsAssetDataModule::Get()
@@ -60,17 +58,17 @@ FTedsAssetDataModule& FTedsAssetDataModule::GetChecked()
 
 void FTedsAssetDataModule::EnableTedsAssetRegistryStorage()
 {
+	using namespace UE::Editor::DataStorage;
 	if (!AssetRegistryStorage)
 	{
 		FModuleManager::Get().LoadModuleChecked(TEXT("TypedElementFramework"));
-		UTypedElementRegistry* TypedElementRegistry = UTypedElementRegistry::GetInstance();
-		if (TypedElementRegistry->GetMutableDataStorage())
+		if (AreEditorDataStorageFeaturesEnabled())
 		{
 			InitAssetRegistryStorage();
 		}
 		else
 		{
-			TypedElementRegistry->OnDataStorageInterfacesSet().AddRaw(this, &FTedsAssetDataModule::InitAssetRegistryStorage);
+			OnEditorDataStorageFeaturesEnabled().AddRaw(this, &FTedsAssetDataModule::InitAssetRegistryStorage);
 		}
 
 		if (!Private::CVarTEDSAssetDataStorage.GetValueOnGameThread())
@@ -108,7 +106,8 @@ void FTedsAssetDataModule::ProcessDependentEvents()
 
 void FTedsAssetDataModule::InitAssetRegistryStorage()
 {
-	IEditorDataStorageProvider& MutableDataStorage = *UTypedElementRegistry::GetInstance()->GetMutableDataStorage();
+	using namespace UE::Editor::DataStorage;
+	IEditorDataStorageProvider& MutableDataStorage = *GetMutableDataStorageFeature<IEditorDataStorageProvider>(StorageFeatureName);
 
 	AssetDataCBDataSource = MakeUnique<Private::FTedsAssetDataCBDataSource>(MutableDataStorage);
 	AssetRegistryStorage = MakeUnique<Private::FTedsAssetData>(MutableDataStorage);
diff --git a/Engine/Plugins/Experimental/EditorDataStorageFeatures/Source/TedsAssetData/Private/Tests/TedsAssetDataTest.cpp b/Engine/Plugins/Experimental/EditorDataStorageFeatures/Source/TedsAssetData/Private/Tests/TedsAssetDataTest.cpp
--- a/Engine/Plugins/Experimental/EditorDataStorageFeatures/Source/TedsAssetData/Private/Tests/TedsAssetDataTest.cpp
+++ b/Engine/Plugins/Experimental/EditorDataStorageFeatures/Source/TedsAssetData/Private/Tests/TedsAssetDataTest.cpp
@@ -7,6 +7,7 @@
 #include "Containers/ChunkedArray.h"
 #include "Containers/Set.h"
 #include "Containers/UnrealString.h"
+#include "Elements/Common/EditorDataStorageFeatures.h"
 #include "Elements/Common/TypedElementHandles.h"
 #include "Elements/Framework/TypedElementIndexHasher.h"
 #include "Elements/Framework/TypedElementRegistry.h"
@@ -65,7 +66,8 @@ bool FTedsAssetDataTest::RunTest(const FString& Parameters)
 
 
 	// Do a sanity check that the data from the asset registry exist in Teds
-	const IEditorDataStorageProvider* Database = UTypedElementRegistry::GetInstance()->GetDataStorage();
+	using namespace UE::Editor::DataStorage;
+	const IEditorDataStorageProvider* Database = GetMutableDataStorageFeature<IEditorDataStorageProvider>(StorageFeatureName);
 	
 	for (const TPair<FName, int32>& PathAndAssetCount : AssetRegistryPathsAndAssetCount)
 	{
diff --git a/Engine/Plugins/Experimental/EditorDataStorageFeatures/Source/TedsContentBrowser/Private/TedsContentBrowserModule.cpp b/Engine/Plugins/Experimental/EditorDataStorageFeatures/Source/TedsContentBrowser/Private/TedsContentBrowserModule.cpp
--- a/Engine/Plugins/Experimental/EditorDataStorageFeatures/Source/TedsContentBrowser/Private/TedsContentBrowserModule.cpp
+++ b/Engine/Plugins/Experimental/EditorDataStorageFeatures/Source/TedsContentBrowser/Private/TedsContentBrowserModule.cpp
@@ -6,8 +6,9 @@
 #include "ContentBrowserModule.h"
 #include "Columns/SlateDelegateColumns.h"
 #include "Elements/Common/TypedElementQueryTypes.h"
+#include "Elements/Common/EditorDataStorageFeatures.h"
 #include "Elements/Framework/TypedElementIndexHasher.h"
-#include "Elements/Framework/TypedElementRegistry.h"
+#include "Elements/Interfaces/TypedElementDataStorageInterface.h"
 #include "Experimental/ContentBrowserViewExtender.h"
 #include "Modules/ModuleManager.h"
 #include "QueryStack/FQueryStackNode_RowView.h"
@@ -126,7 +127,7 @@ namespace UE::Editor::ContentBrowser
 
 	FTedsContentBrowserViewExtender::FTedsContentBrowserViewExtender()
 	{
-		DataStorage = UTypedElementRegistry::GetInstance()->GetMutableDataStorage();
+		DataStorage = GetMutableDataStorageFeature<IEditorDataStorageProvider>(StorageFeatureName);
 		
 		RowQueryStack = MakeShared<FQueryStackNode_RowView>(&Rows);
 
diff --git a/Engine/Plugins/Experimental/EditorDataStorageFeatures/Source/TedsDebugger/Private/STedsDebugger.cpp b/Engine/Plugins/Experimental/EditorDataStorageFeatures/Source/TedsDebugger/Private/STedsDebugger.cpp
--- a/Engine/Plugins/Experimental/EditorDataStorageFeatures/Source/TedsDebugger/Private/STedsDebugger.cpp
+++ b/Engine/Plugins/Experimental/EditorDataStorageFeatures/Source/TedsDebugger/Private/STedsDebugger.cpp
@@ -7,7 +7,8 @@
 #include "Elements/Columns/TypedElementMiscColumns.h"
 #include "Elements/Columns/TypedElementSelectionColumns.h"
 #include "Elements/Columns/TypedElementTypeInfoColumns.h"
-#include "Elements/Framework/TypedElementRegistry.h"
+#include "Elements/Common/EditorDataStorageFeatures.h"
+#include "Elements/Interfaces/TypedElementDataStorageInterface.h"
 #include "Framework/Docking/LayoutService.h"
 #include "Modules/ModuleManager.h"
 #include "QueryEditor/TedsQueryEditor.h"
@@ -26,11 +27,9 @@ namespace UE::Editor::DataStorage::Debug
 
 STedsDebugger::~STedsDebugger()
 {
-	UTypedElementRegistry* Registry = UTypedElementRegistry::GetInstance();
-
-	if (Registry && Registry->AreDataStorageInterfacesSet())
+	if (AreEditorDataStorageFeaturesEnabled())
 	{
-		Registry->GetMutableDataStorage()->UnregisterQuery(TableViewerQuery);
+		GetMutableDataStorageFeature<IEditorDataStorageProvider>(StorageFeatureName)->UnregisterQuery(TableViewerQuery);
 	}
 }
 
@@ -131,11 +130,9 @@ TSharedRef<SDockTab> STedsDebugger::SpawnQueryEditorTab(const FSpawnTabArgs& Arg
 	TSharedRef<SDockTab> DockTab = SNew(SDockTab).TabRole(ETabRole::PanelTab);
 	if (!QueryEditorModel)
 	{
-		UTypedElementRegistry* Registry = UTypedElementRegistry::GetInstance();
-
-		if (Registry && Registry->AreDataStorageInterfacesSet())
+		if (AreEditorDataStorageFeaturesEnabled())
 		{
-			IEditorDataStorageProvider* DataStorageInterface = Registry->GetMutableDataStorage();
+			IEditorDataStorageProvider* DataStorageInterface = GetMutableDataStorageFeature<IEditorDataStorageProvider>(StorageFeatureName);
 			QueryEditorModel = MakeUnique<QueryEditor::FTedsQueryEditorModel>(*DataStorageInterface);
 		}
 	}
diff --git a/Engine/Plugins/Experimental/EditorDataStorageFeatures/Source/TedsDebugger/Private/Widgets/QueryEditor/TedsQueryEditorResultsView.cpp b/Engine/Plugins/Experimental/EditorDataStorageFeatures/Source/TedsDebugger/Private/Widgets/QueryEditor/TedsQueryEditorResultsView.cpp
--- a/Engine/Plugins/Experimental/EditorDataStorageFeatures/Source/TedsDebugger/Private/Widgets/QueryEditor/TedsQueryEditorResultsView.cpp
+++ b/Engine/Plugins/Experimental/EditorDataStorageFeatures/Source/TedsDebugger/Private/Widgets/QueryEditor/TedsQueryEditorResultsView.cpp
@@ -5,9 +5,10 @@
 #include "SWarningOrErrorBox.h"
 #include "TedsOutlinerModule.h"
 #include "TedsTableViewerColumn.h"
-#include "Elements/Interfaces/TypedElementDataStorageInterface.h"
+#include "Elements/Common/EditorDataStorageFeatures.h"
 #include "Elements/Framework/TypedElementQueryBuilder.h"
-#include "Elements/Framework/TypedElementRegistry.h"
+#include "Elements/Interfaces/TypedElementDataStorageInterface.h"
+#include "Elements/Interfaces/TypedElementDataStorageUiInterface.h"
 #include "Modules/ModuleManager.h"
 #include "QueryEditor/TedsQueryEditorModel.h"
 #include "QueryStack/FQueryStackNode_RowView.h"
@@ -185,9 +186,7 @@ namespace UE::Editor::DataStorage::Debug::QueryEditor
 			return false;
 		};
 	
-		UTypedElementRegistry* Registry = UTypedElementRegistry::GetInstance();
-		checkf(Registry, TEXT("SResultsView created before UTypedElementRegistry is available."));
-		IEditorDataStorageUiProvider* StorageUi = Registry->GetMutableDataStorageUi();
+		IEditorDataStorageUiProvider* StorageUi = GetMutableDataStorageFeature<IEditorDataStorageUiProvider>(UiFeatureName);
 		checkf(StorageUi, TEXT("SResultsView created before data storage interfaces were initialized."))
 
 		StorageUi->CreateWidgetConstructors(TEXT("General.Cell.RowHandle"), FMetaDataView(), AssignWidgetToColumn);
diff --git a/Engine/Plugins/Experimental/EditorDataStorageFeatures/Source/TedsDebugger/Private/Widgets/RowReferenceWidget.cpp b/Engine/Plugins/Experimental/EditorDataStorageFeatures/Source/TedsDebugger/Private/Widgets/RowReferenceWidget.cpp
--- a/Engine/Plugins/Experimental/EditorDataStorageFeatures/Source/TedsDebugger/Private/Widgets/RowReferenceWidget.cpp
+++ b/Engine/Plugins/Experimental/EditorDataStorageFeatures/Source/TedsDebugger/Private/Widgets/RowReferenceWidget.cpp
@@ -6,12 +6,12 @@
 #include "Elements/Columns/TypedElementHiearchyColumns.h"
 #include "Elements/Columns/TypedElementLabelColumns.h"
 #include "Elements/Columns/TypedElementMiscColumns.h"
+#include "Elements/Common/EditorDataStorageFeatures.h"
 #include "Elements/Interfaces/TypedElementDataStorageInterface.h"
 #include "ISceneOutliner.h"
 #include "Modules/ModuleManager.h"
 #include "TedsDebuggerModule.h"
 #include "Elements/Columns/TypedElementSlateWidgetColumns.h"
-#include "Elements/Framework/TypedElementRegistry.h"
 #include "Widgets/Input/SHyperlink.h"
 #include "Widgets/Layout/SBox.h"
 
@@ -66,7 +66,7 @@ namespace UE::Editor::DataStorage::Debug::Private
 		SBox* WidgetInstance = static_cast<SBox*>(Widget.Get());
 		WidgetInstance->SetContent(SNullWidget::NullWidget);
 
-		const IEditorDataStorageProvider* DataStorage = UTypedElementRegistry::GetInstance()->GetDataStorage();
+		const IEditorDataStorageProvider* DataStorage = GetMutableDataStorageFeature<IEditorDataStorageProvider>(StorageFeatureName);
 
 		// We only navigate to row references that have a label column
 		if (const FTypedElementLabelColumn* LabelColumn = DataStorage->GetColumn<FTypedElementLabelColumn>(TargetRow))
diff --git a/Engine/Plugins/Experimental/EditorDataStorageFeatures/Source/TedsOutliner/Private/Compatibility/SceneOutlinerRowHandleColumn.cpp b/Engine/Plugins/Experimental/EditorDataStorageFeatures/Source/TedsOutliner/Private/Compatibility/SceneOutlinerRowHandleColumn.cpp
--- a/Engine/Plugins/Experimental/EditorDataStorageFeatures/Source/TedsOutliner/Private/Compatibility/SceneOutlinerRowHandleColumn.cpp
+++ b/Engine/Plugins/Experimental/EditorDataStorageFeatures/Source/TedsOutliner/Private/Compatibility/SceneOutlinerRowHandleColumn.cpp
@@ -4,7 +4,8 @@
 
 #include "SortHelper.h"
 #include "TedsOutlinerItem.h"
-#include "Elements/Framework/TypedElementRegistry.h"
+#include "Elements/Common/EditorDataStorageFeatures.h"
+#include "Elements/Interfaces/TypedElementDataStorageUiInterface.h"
 #include "TedsTableViewerColumn.h"
 
 #define LOCTEXT_NAMESPACE "SceneOutlinerRowHandleColumn"
@@ -12,16 +13,15 @@
 FSceneOutlinerRowHandleColumn::FSceneOutlinerRowHandleColumn(ISceneOutliner& SceneOutliner)
 	: WeakSceneOutliner(StaticCastSharedRef<ISceneOutliner>(SceneOutliner.AsShared()))
 {
+	using namespace UE::Editor::DataStorage;
 	auto AssignWidgetToColumn = [this](TUniquePtr<FTypedElementWidgetConstructor> Constructor, TConstArrayView<TWeakObjectPtr<const UScriptStruct>>)
 	{
 		TSharedPtr<FTypedElementWidgetConstructor> WidgetConstructor(Constructor.Release());
-		TableViewerColumn = MakeShared<UE::Editor::DataStorage::FTedsTableViewerColumn>(TEXT("Row Handle"), WidgetConstructor);
+		TableViewerColumn = MakeShared<FTedsTableViewerColumn>(TEXT("Row Handle"), WidgetConstructor);
 		return false;
 	};
 	
-	UTypedElementRegistry* Registry = UTypedElementRegistry::GetInstance();
-	checkf(Registry, TEXT("FSceneOutlinerRowHandleColumn created before UTypedElementRegistry is available."));
-	IEditorDataStorageUiProvider* StorageUi = Registry->GetMutableDataStorageUi();
+	IEditorDataStorageUiProvider* StorageUi = GetMutableDataStorageFeature<IEditorDataStorageUiProvider>(UiFeatureName);
 	checkf(StorageUi, TEXT("FSceneOutlinerRowHandleColumn created before data storage interfaces were initialized."))
 
 	StorageUi->CreateWidgetConstructors(TEXT("General.Cell.RowHandle"), UE::Editor::DataStorage::FMetaDataView(), AssignWidgetToColumn);
diff --git a/Engine/Plugins/Experimental/EditorDataStorageFeatures/Source/TedsOutliner/Private/Compatibility/SceneOutlinerTedsBridge.cpp b/Engine/Plugins/Experimental/EditorDataStorageFeatures/Source/TedsOutliner/Private/Compatibility/SceneOutlinerTedsBridge.cpp
--- a/Engine/Plugins/Experimental/EditorDataStorageFeatures/Source/TedsOutliner/Private/Compatibility/SceneOutlinerTedsBridge.cpp
+++ b/Engine/Plugins/Experimental/EditorDataStorageFeatures/Source/TedsOutliner/Private/Compatibility/SceneOutlinerTedsBridge.cpp
@@ -8,9 +8,11 @@
 #include "Elements/Columns/TypedElementPackageColumns.h"
 #include "Elements/Columns/TypedElementTransformColumns.h"
 #include "Elements/Columns/TypedElementTypeInfoColumns.h"
-#include "Elements/Framework/TypedElementRegistry.h"
+#include "Elements/Common/EditorDataStorageFeatures.h"
 #include "Elements/Framework/TypedElementQueryBuilder.h"
+#include "Elements/Interfaces/TypedElementDataStorageInterface.h"
 #include "Elements/Interfaces/TypedElementDataStorageCompatibilityInterface.h"
+#include "Elements/Interfaces/TypedElementDataStorageUiInterface.h"
 #include "ILevelEditor.h"
 #include "ISceneOutliner.h"
 #include "ISceneOutlinerColumn.h"
@@ -38,8 +40,7 @@ FAutoConsoleCommand BindColumnsToSceneOutlinerConsoleCommand(
 			
 		    const FName WidgetPurposes[] = {TEXT("SceneOutliner.Cell"), TEXT("General.Cell")};
 
-			UTypedElementRegistry* Registry = UTypedElementRegistry::GetInstance();
-			if (IEditorDataStorageProvider* DataStorage = Registry->GetMutableDataStorage())
+			if (IEditorDataStorageProvider* DataStorage = GetMutableDataStorageFeature<IEditorDataStorageProvider>(StorageFeatureName))
 			{
 				static UE::Editor::DataStorage::QueryHandle Queries[] =
 				{
@@ -367,14 +368,10 @@ const FName FSceneOutlinerTedsQueryBinder::DefaultItemLabelCellWidgetPurpose(TEX
 
 FSceneOutlinerTedsQueryBinder::FSceneOutlinerTedsQueryBinder()
 {
-	UTypedElementRegistry* Registry = UTypedElementRegistry::GetInstance();
-	checkf(Registry, TEXT("Unable to bind a Scene Outliner to a query before the Typed Elements are available."));
-	if (Registry)
-	{
-		Storage = Registry->GetMutableDataStorage();
-		StorageUi = Registry->GetMutableDataStorageUi();
-		StorageCompatibility = Registry->GetMutableDataStorageCompatibility();
-	}
+	using namespace UE::Editor::DataStorage;
+	Storage = GetMutableDataStorageFeature<IEditorDataStorageProvider>(StorageFeatureName);
+	StorageUi = GetMutableDataStorageFeature<IEditorDataStorageUiProvider>(UiFeatureName);
+	StorageCompatibility = GetMutableDataStorageFeature<IEditorDataStorageCompatibilityProvider>(CompatibilityFeatureName);
 
 	SetupDefaultColumnMapping();
 }
diff --git a/Engine/Plugins/Experimental/EditorDataStorageFeatures/Source/TedsOutliner/Private/TedsOutlinerImpl.cpp b/Engine/Plugins/Experimental/EditorDataStorageFeatures/Source/TedsOutliner/Private/TedsOutlinerImpl.cpp
--- a/Engine/Plugins/Experimental/EditorDataStorageFeatures/Source/TedsOutliner/Private/TedsOutlinerImpl.cpp
+++ b/Engine/Plugins/Experimental/EditorDataStorageFeatures/Source/TedsOutliner/Private/TedsOutlinerImpl.cpp
@@ -7,10 +7,11 @@
 #include "Elements/Columns/TypedElementLabelColumns.h"
 #include "Elements/Columns/TypedElementMiscColumns.h"
 #include "Elements/Columns/TypedElementTypeInfoColumns.h"
+#include "Elements/Common/EditorDataStorageFeatures.h"
 #include "Elements/Interfaces/TypedElementDataStorageInterface.h"
+#include "Elements/Interfaces/TypedElementDataStorageCompatibilityInterface.h"
 #include "Elements/Interfaces/TypedElementDataStorageUiInterface.h"
 #include "Elements/Framework/TypedElementQueryBuilder.h"
-#include "Elements/Framework/TypedElementRegistry.h"
 #include "Columns/SlateDelegateColumns.h"
 #include "Compatibility/SceneOutlinerRowHandleColumn.h"
 #include "TedsOutlinerFilter.h"
@@ -88,15 +89,10 @@ FTedsOutlinerImpl::FTedsOutlinerImpl(const FTedsOutlinerParams& InParams, IScene
 	, SceneOutliner(InParams.SceneOutliner)
 {
 	// Initialize the TEDS constructs
-	UTypedElementRegistry* Registry = UTypedElementRegistry::GetInstance();
-	checkf(Registry, TEXT("Unable to initialize the Typed Elements Outliner before TEDS is initialized."));
-	if (Registry)
-	{
-		Storage = Registry->GetMutableDataStorage();
-		StorageUi = Registry->GetMutableDataStorageUi();
-		StorageCompatibility = Registry->GetMutableDataStorageCompatibility();
-	}
-
+	using namespace UE::Editor::DataStorage;
+	Storage = GetMutableDataStorageFeature<IEditorDataStorageProvider>(StorageFeatureName);
+	StorageUi = GetMutableDataStorageFeature<IEditorDataStorageUiProvider>(UiFeatureName);
+	StorageCompatibility = GetMutableDataStorageFeature<IEditorDataStorageCompatibilityProvider>(CompatibilityFeatureName);
 }
 
 void FTedsOutlinerImpl::CreateFilterQueries()
diff --git a/Engine/Plugins/Experimental/EditorDataStorageFeatures/Source/TedsOutliner/Private/TedsOutlinerModule.cpp b/Engine/Plugins/Experimental/EditorDataStorageFeatures/Source/TedsOutliner/Private/TedsOutlinerModule.cpp
--- a/Engine/Plugins/Experimental/EditorDataStorageFeatures/Source/TedsOutliner/Private/TedsOutlinerModule.cpp
+++ b/Engine/Plugins/Experimental/EditorDataStorageFeatures/Source/TedsOutliner/Private/TedsOutlinerModule.cpp
@@ -13,7 +13,8 @@
 #include "Elements/Columns/TypedElementMiscColumns.h"
 #include "Elements/Columns/TypedElementPackageColumns.h"
 #include "Elements/Columns/TypedElementTypeInfoColumns.h"
-#include "Elements/Framework/TypedElementRegistry.h"
+#include "Elements/Common/EditorDataStorageFeatures.h"
+#include "Elements/Interfaces/TypedElementDataStorageInterface.h"
 #include "Modules/ModuleManager.h"
 #include "Compatibility/SceneOutlinerTedsBridge.h"
 #include "Compatibility/SceneOutlinerRowHandleColumn.h"
@@ -65,9 +66,8 @@ FTedsOutlinerModule::FTedsOutlinerModule()
 
 TSharedRef<ISceneOutliner> FTedsOutlinerModule::CreateTedsOutliner(const FSceneOutlinerInitializationOptions& InInitOptions, const FTedsOutlinerParams& InInitTedsOptions, DataStorage::QueryHandle ColumnQuery) const
 {
-	UTypedElementRegistry* Registry = UTypedElementRegistry::GetInstance();
-
-	ensureMsgf(Registry&& Registry->AreDataStorageInterfacesSet(), TEXT("Unable to initialize the Teds-Outliner before TEDS itself is initialized."));
+	using namespace UE::Editor::DataStorage;
+	ensureMsgf(AreEditorDataStorageFeaturesEnabled(), TEXT("Unable to initialize the Teds-Outliner before TEDS itself is initialized."));
 
 	FSceneOutlinerInitializationOptions InitOptions(InInitOptions);
 	FTedsOutlinerParams InitTedsOptions(InInitTedsOptions);
@@ -116,8 +116,8 @@ void FTedsOutlinerModule::ShutdownModule()
 
 DataStorage::QueryHandle FTedsOutlinerModule::GetLevelEditorTedsOutlinerColumnQuery()
 {
-	UTypedElementRegistry* Registry = UTypedElementRegistry::GetInstance();
-	IEditorDataStorageProvider* Storage = Registry->GetMutableDataStorage();
+	using namespace UE::Editor::DataStorage;
+	IEditorDataStorageProvider* Storage = GetMutableDataStorageFeature<IEditorDataStorageProvider>(StorageFeatureName);
 		
 	using namespace DataStorage::Queries;
 
@@ -138,10 +138,7 @@ DataStorage::QueryHandle FTedsOutlinerModule::GetLevelEditorTedsOutlinerColumnQu
 
 TSharedRef<SWidget> FTedsOutlinerModule::CreateLevelEditorTedsOutliner()
 {
-	UTypedElementRegistry* Registry = UTypedElementRegistry::GetInstance();
-	checkf(Registry, TEXT("Unable to initialize the table viewer before TEDS is initialized."));
-
-	if(!Registry->AreDataStorageInterfacesSet())
+	if(!DataStorage::AreEditorDataStorageFeaturesEnabled())
 	{
 		return SNew(STextBlock)
 		.Text(LOCTEXT("TEDSPluginNotEnabledText", "You need to enable the Typed Element Data Storage plugin to see the table viewer!"));
diff --git a/Engine/Plugins/Experimental/EditorDataStorageFeatures/Source/TedsOutliner/Private/Widgets/AlertWidget.cpp b/Engine/Plugins/Experimental/EditorDataStorageFeatures/Source/TedsOutliner/Private/Widgets/AlertWidget.cpp
--- a/Engine/Plugins/Experimental/EditorDataStorageFeatures/Source/TedsOutliner/Private/Widgets/AlertWidget.cpp
+++ b/Engine/Plugins/Experimental/EditorDataStorageFeatures/Source/TedsOutliner/Private/Widgets/AlertWidget.cpp
@@ -6,7 +6,7 @@
 #include "Elements/Columns/TypedElementAlertColumns.h"
 #include "Elements/Columns/TypedElementMiscColumns.h"
 #include "Elements/Columns/TypedElementSlateWidgetColumns.h"
-#include "Elements/Framework/TypedElementRegistry.h"
+#include "Elements/Common/EditorDataStorageFeatures.h"
 #include "Widgets/Images/SImage.h"
 #include "Widgets/Input/SButton.h"
 #include "Widgets/SOverlay.h"
@@ -17,7 +17,7 @@
 namespace UE::Editor::DataStorage::Widgets::Private
 {
 	void UpdateWidget(const TSharedPtr<SWidget>& Widget, const FText& Alert, bool bIsWarning, uint16 ErrorCount, uint16 WarningCount,
-		UE::Editor::DataStorage::RowHandle RowWithAlertAction)
+		RowHandle RowWithAlertAction)
 	{
 		if (Widget)
 		{
@@ -117,14 +117,13 @@ namespace UE::Editor::DataStorage::Widgets::Private
 				}
 
 				// If there's an action to call, enable the invisible button, otherwise turn it off.
-				if (RowWithAlertAction != UE::Editor::DataStorage::InvalidRowHandle)
+				if (RowWithAlertAction != InvalidRowHandle)
 				{
 					Background.SetVisibility(EVisibility::HitTestInvisible);
 					ActionButton.SetVisibility(EVisibility::Visible);
 					ActionButton.SetOnClicked(FOnClicked::CreateLambda([RowWithAlertAction]()
 						{
-							UTypedElementRegistry* Registry = UTypedElementRegistry::GetInstance();
-							const IEditorDataStorageProvider* DataStorage = Registry->GetDataStorage();
+							const IEditorDataStorageProvider* DataStorage = GetDataStorageFeature<IEditorDataStorageProvider>(StorageFeatureName);
 							if (const FTypedElementAlertActionColumn* Action =
 								DataStorage->GetColumn<FTypedElementAlertActionColumn>(RowWithAlertAction))
 							{
diff --git a/Engine/Plugins/Experimental/EditorDataStorageFeatures/Source/TedsPropertyEditor/Private/Widgets/SPropertyMenuTedsRowPicker.cpp b/Engine/Plugins/Experimental/EditorDataStorageFeatures/Source/TedsPropertyEditor/Private/Widgets/SPropertyMenuTedsRowPicker.cpp
--- a/Engine/Plugins/Experimental/EditorDataStorageFeatures/Source/TedsPropertyEditor/Private/Widgets/SPropertyMenuTedsRowPicker.cpp
+++ b/Engine/Plugins/Experimental/EditorDataStorageFeatures/Source/TedsPropertyEditor/Private/Widgets/SPropertyMenuTedsRowPicker.cpp
@@ -2,7 +2,7 @@
 
 #include "Widgets/SPropertyMenuTedsRowPicker.h"
 
-#include "Elements/Framework/TypedElementRegistry.h"
+#include "Elements/Common/EditorDataStorageFeatures.h"
 #include "TedsRowPickingMode.h"
 #include "TedsOutlinerItem.h"
 
@@ -36,16 +36,14 @@ void SPropertyMenuTedsRowPicker::Construct(const FArguments& InArgs)
 		TSharedPtr<SWidget> MenuContent;
 		{
 			using namespace UE::Editor::Outliner;
+			using namespace UE::Editor::DataStorage;
 
 			// TEDS-Outliner TODO: Taken from private implementation of PropertyEditorAssetConstants.
 			//                     Should be centralized when TEDS is moved to core
 			static const FVector2D ContentBrowserWindowSize(300.0f, 300.0f);
 			static const FVector2D SceneOutlinerWindowSize(350.0f, 300.0f);
 
-			UTypedElementRegistry* Registry = UTypedElementRegistry::GetInstance();
-			checkf(Registry, TEXT("Unable to initialize the Typed Elements Outliner before TEDS is initialized."));
-
-			if (!Registry->AreDataStorageInterfacesSet())
+			if (!AreEditorDataStorageFeaturesEnabled())
 			{
 				MenuContent = SNew(STextBlock)
 					.Text(LOCTEXT("TEDSPluginNotEnabledText", "Typed Element Data Storage plugin required to use this property picker."));
diff --git a/Engine/Plugins/Experimental/EditorDataStorageFeatures/Source/TedsSettings/Private/TedsSettingsManager.cpp b/Engine/Plugins/Experimental/EditorDataStorageFeatures/Source/TedsSettings/Private/TedsSettingsManager.cpp
--- a/Engine/Plugins/Experimental/EditorDataStorageFeatures/Source/TedsSettings/Private/TedsSettingsManager.cpp
+++ b/Engine/Plugins/Experimental/EditorDataStorageFeatures/Source/TedsSettings/Private/TedsSettingsManager.cpp
@@ -5,9 +5,10 @@
 #include "Elements/Columns/TypedElementCompatibilityColumns.h"
 #include "Elements/Columns/TypedElementMiscColumns.h"
 #include "Elements/Columns/TypedElementSlateWidgetColumns.h"
+#include "Elements/Common/EditorDataStorageFeatures.h"
 #include "Elements/Framework/TypedElementIndexHasher.h"
 #include "Elements/Framework/TypedElementQueryBuilder.h"
-#include "Elements/Framework/TypedElementRegistry.h"
+#include "Elements/Interfaces/TypedElementDataStorageInterface.h"
 #include "ISettingsCategory.h"
 #include "ISettingsContainer.h"
 #include "ISettingsModule.h"
@@ -44,14 +45,10 @@ void FTedsSettingsManager::Initialize()
 {
 	if (!bIsInitialized)
 	{
-		FModuleManager::Get().LoadModule(TEXT("TypedElementFramework"));
-
-		UTypedElementRegistry* TypedElementRegistry = UTypedElementRegistry::GetInstance();
-		check(TypedElementRegistry);
-
-		auto OnDataStorage = [this, TypedElementRegistry]
+		using namespace UE::Editor::DataStorage;
+		auto OnDataStorage = [this]
 			{
-				IEditorDataStorageProvider* DataStorage = TypedElementRegistry->GetMutableDataStorage();
+				IEditorDataStorageProvider* DataStorage = GetMutableDataStorageFeature<IEditorDataStorageProvider>(StorageFeatureName);
 				check(DataStorage);
 
 				RegisterTables(*DataStorage);
@@ -59,13 +56,13 @@ void FTedsSettingsManager::Initialize()
 				RegisterSettings();
 			};
 
-		if (TypedElementRegistry->AreDataStorageInterfacesSet())
+		if (AreEditorDataStorageFeaturesEnabled())
 		{
 			OnDataStorage();
 		}
 		else
 		{
-			TypedElementRegistry->OnDataStorageInterfacesSet().AddSPLambda(this, OnDataStorage);
+			OnEditorDataStorageFeaturesEnabled().AddSPLambda(this, OnDataStorage);
 		}
 
 		bIsInitialized = true;
@@ -76,14 +73,12 @@ void FTedsSettingsManager::Shutdown()
 {
 	if (bIsInitialized)
 	{
-		UTypedElementRegistry* TypedElementRegistry = UTypedElementRegistry::GetInstance();
-		check(TypedElementRegistry);
-
-		TypedElementRegistry->OnDataStorageInterfacesSet().RemoveAll(this);
+		using namespace UE::Editor::DataStorage;
+		OnEditorDataStorageFeaturesEnabled().RemoveAll(this);
 
-		if (TypedElementRegistry->AreDataStorageInterfacesSet())
+		if (AreEditorDataStorageFeaturesEnabled())
 		{
-			IEditorDataStorageProvider* DataStorage = TypedElementRegistry->GetMutableDataStorage();
+			IEditorDataStorageProvider* DataStorage = GetMutableDataStorageFeature<IEditorDataStorageProvider>(StorageFeatureName);
 			check(DataStorage);
 
 			UnregisterSettings();
@@ -159,10 +154,7 @@ void FTedsSettingsManager::RegisterSettingsContainer(const FName& ContainerName)
 	ISettingsModule* SettingsModule = FModuleManager::GetModulePtr<ISettingsModule>("Settings");
 	check(SettingsModule);
 
-	UTypedElementRegistry* TypedElementRegistry = UTypedElementRegistry::GetInstance();
-	check(TypedElementRegistry);
-
-	IEditorDataStorageProvider* DataStorage = TypedElementRegistry->GetMutableDataStorage();
+	IEditorDataStorageProvider* DataStorage = GetMutableDataStorageFeature<IEditorDataStorageProvider>(StorageFeatureName);
 	check(DataStorage);
 
 	UE_LOG(LogTedsSettings, Log, TEXT("Register Settings Container : '%s'"), *ContainerName.ToString());
@@ -212,13 +204,10 @@ void FTedsSettingsManager::UnregisterSettings()
 	ISettingsModule* SettingsModule = FModuleManager::GetModulePtr<ISettingsModule>("Settings");
 	check(SettingsModule);
 
-	UTypedElementRegistry* TypedElementRegistry = UTypedElementRegistry::GetInstance();
-	check(TypedElementRegistry);
-
-	IEditorDataStorageProvider* DataStorage = TypedElementRegistry->GetMutableDataStorage();
+	IEditorDataStorageProvider* DataStorage = GetMutableDataStorageFeature<IEditorDataStorageProvider>(StorageFeatureName);
 	check(DataStorage);
 
-	IEditorDataStorageCompatibilityProvider* DataStorageCompatibility = TypedElementRegistry->GetMutableDataStorageCompatibility();
+	IEditorDataStorageCompatibilityProvider* DataStorageCompatibility = GetMutableDataStorageFeature<IEditorDataStorageCompatibilityProvider>(CompatibilityFeatureName);
 	check(DataStorageCompatibility);
 
 	SettingsModule->OnContainerAdded().RemoveAll(this);
@@ -283,13 +272,10 @@ void FTedsSettingsManager::UpdateSettingsCategory(TSharedPtr<ISettingsCategory>
 
 	TRACE_CPUPROFILER_EVENT_SCOPE(TedsSettingsManager.UpdateSettingsCategory);
 
-	UTypedElementRegistry* TypedElementRegistry = UTypedElementRegistry::GetInstance();
-	check(TypedElementRegistry);
-
-	IEditorDataStorageProvider* DataStorage = TypedElementRegistry->GetMutableDataStorage();
+	IEditorDataStorageProvider* DataStorage = GetMutableDataStorageFeature<IEditorDataStorageProvider>(StorageFeatureName);
 	check(DataStorage);
 
-	IEditorDataStorageCompatibilityProvider* DataStorageCompatibility = TypedElementRegistry->GetMutableDataStorageCompatibility();
+	IEditorDataStorageCompatibilityProvider* DataStorageCompatibility = GetMutableDataStorageFeature<IEditorDataStorageCompatibilityProvider>(CompatibilityFeatureName);
 	check(DataStorageCompatibility);
 
 	const FName& ContainerName = DataStorage->GetColumn<FNameColumn>(ContainerRow)->Name;
diff --git a/Engine/Plugins/Experimental/EditorDataStorageFeatures/Source/TedsSettings/Tests/TedsSettingsTests.cpp b/Engine/Plugins/Experimental/EditorDataStorageFeatures/Source/TedsSettings/Tests/TedsSettingsTests.cpp
--- a/Engine/Plugins/Experimental/EditorDataStorageFeatures/Source/TedsSettings/Tests/TedsSettingsTests.cpp
+++ b/Engine/Plugins/Experimental/EditorDataStorageFeatures/Source/TedsSettings/Tests/TedsSettingsTests.cpp
@@ -8,8 +8,8 @@
 #include "Editor.h"
 #include "Elements/Columns/TypedElementCompatibilityColumns.h"
 #include "Elements/Columns/TypedElementMiscColumns.h"
+#include "Elements/Common/EditorDataStorageFeatures.h"
 #include "Elements/Framework/TypedElementQueryBuilder.h"
-#include "Elements/Framework/TypedElementRegistry.h"
 #include "Elements/Interfaces/TypedElementDataStorageInterface.h"
 #include "ISettingsModule.h"
 #include "Modules/ModuleManager.h"
@@ -21,7 +21,6 @@ namespace UE::Editor::Settings::Tests
 {
 	BEGIN_DEFINE_SPEC(FTedsSettingsTestFixture, "Editor.DataStorage.Settings", EAutomationTestFlags::EditorContext | EAutomationTestFlags::EngineFilter)
 	ISettingsModule* SettingsModule = nullptr;
-	UTypedElementRegistry* TypedElementRegistry = nullptr;
 	IEditorDataStorageProvider* DataStorage = nullptr;
 	IEditorDataStorageCompatibilityProvider* DataStorageCompatibility = nullptr;
 	DataStorage::QueryHandle CountAllSettingsQuery = DataStorage::InvalidQueryHandle;
@@ -72,16 +71,14 @@ namespace UE::Editor::Settings::Tests
 
 		BeforeEach([this]()
 		{
+			using namespace UE::Editor::DataStorage;
 			SettingsModule = FModuleManager::GetModulePtr<ISettingsModule>("Settings");
 			check(SettingsModule != nullptr);
 
-			TypedElementRegistry = UTypedElementRegistry::GetInstance();
-			check(TypedElementRegistry != nullptr);
-
-			DataStorage = TypedElementRegistry->GetMutableDataStorage();
+			DataStorage = GetMutableDataStorageFeature<IEditorDataStorageProvider>(StorageFeatureName);
 			check(DataStorage != nullptr);
 
-			DataStorageCompatibility = TypedElementRegistry->GetMutableDataStorageCompatibility();
+			DataStorageCompatibility = GetMutableDataStorageFeature<IEditorDataStorageCompatibilityProvider>(CompatibilityFeatureName);
 			check(DataStorageCompatibility != nullptr);
 
 			{
@@ -105,7 +102,6 @@ namespace UE::Editor::Settings::Tests
 			TestRowHandles.Empty();
 			CountAllSettingsQuery = DataStorage::InvalidQueryHandle;
 			SettingsModule = nullptr;
-			TypedElementRegistry = nullptr;
 			DataStorage = nullptr;
 			DataStorageCompatibility = nullptr;
 		});
diff --git a/Engine/Plugins/Experimental/EditorDataStorageFeatures/Source/TedsStyling/Private/TedsStylingFactory.cpp b/Engine/Plugins/Experimental/EditorDataStorageFeatures/Source/TedsStyling/Private/TedsStylingFactory.cpp
--- a/Engine/Plugins/Experimental/EditorDataStorageFeatures/Source/TedsStyling/Private/TedsStylingFactory.cpp
+++ b/Engine/Plugins/Experimental/EditorDataStorageFeatures/Source/TedsStyling/Private/TedsStylingFactory.cpp
@@ -5,8 +5,8 @@
 #include "Columns/TedsStylingColumns.h"
 #include "Elements/Columns/TypedElementMiscColumns.h"
 #include "Elements/Columns/TypedElementSlateWidgetColumns.h"
+#include "Elements/Common/EditorDataStorageFeatures.h"
 #include "Elements/Framework/TypedElementIndexHasher.h"
-#include "Elements/Framework/TypedElementRegistry.h"
 #include "Elements/Interfaces/TypedElementDataStorageInterface.h"
 #include "HAL/IConsoleManager.h"
 #include "Styling/SlateBrush.h"
@@ -38,7 +38,8 @@ void UTedsStylingFactory::RegisterQueries(IEditorDataStorageProvider& DataStorag
 
 void UTedsStylingFactory::RegisterAllKnownStyles()
 {
-	IEditorDataStorageProvider* Interface = UTypedElementRegistry::GetInstance()->GetMutableDataStorage();
+	using namespace UE::Editor::DataStorage;
+	IEditorDataStorageProvider* Interface = GetMutableDataStorageFeature<IEditorDataStorageProvider>(StorageFeatureName);
 
 	auto RegisterBrushesForStyle = [Interface](const ISlateStyle& Style)
 	{
diff --git a/Engine/Plugins/Experimental/EditorDataStorageFeatures/Source/TedsTableViewer/Private/TedsTableViewerColumn.cpp b/Engine/Plugins/Experimental/EditorDataStorageFeatures/Source/TedsTableViewer/Private/TedsTableViewerColumn.cpp
--- a/Engine/Plugins/Experimental/EditorDataStorageFeatures/Source/TedsTableViewer/Private/TedsTableViewerColumn.cpp
+++ b/Engine/Plugins/Experimental/EditorDataStorageFeatures/Source/TedsTableViewer/Private/TedsTableViewerColumn.cpp
@@ -6,6 +6,10 @@
 #include "Elements/Columns/TypedElementMiscColumns.h"
 #include "Elements/Columns/TypedElementSlateWidgetColumns.h"
 #include "Elements/Columns/TypedElementTypeInfoColumns.h"
+#include "Elements/Common/EditorDataStorageFeatures.h"
+#include "Elements/Interfaces/TypedElementDataStorageInterface.h"
+#include "Elements/Interfaces/TypedElementDataStorageCompatibilityInterface.h"
+#include "Elements/Interfaces/TypedElementDataStorageUiInterface.h"
 #include "Elements/Framework/TypedElementDataStorageWidget.h"
 #include "Elements/Framework/TypedElementRegistry.h"
 #include "TedsTableViewerUtils.h"
@@ -24,15 +28,9 @@ namespace UE::Editor::DataStorage
 		, MatchedColumns(InMatchedColumns)
 		, WidgetMetaData(InWidgetMetaData)
 	{
-		UTypedElementRegistry* Registry = UTypedElementRegistry::GetInstance();
-		
-		checkf(Registry, TEXT("Unable to create a Table Viewer column before the Typed Element Registry is initialized."));
-		if (Registry)
-		{
-			Storage = Registry->GetMutableDataStorage();
-			StorageUi = Registry->GetMutableDataStorageUi();
-			StorageCompatibility = Registry->GetMutableDataStorageCompatibility();
-		}
+		Storage = GetMutableDataStorageFeature<IEditorDataStorageProvider>(StorageFeatureName);
+		StorageUi = GetMutableDataStorageFeature<IEditorDataStorageUiProvider>(UiFeatureName);
+		StorageCompatibility = GetMutableDataStorageFeature<IEditorDataStorageCompatibilityProvider>(CompatibilityFeatureName);
 
 		// Store the matched columns as a query condition that requires all of them (i.e AND's them)
 		for(const TWeakObjectPtr<const UScriptStruct>& Column : MatchedColumns)
diff --git a/Engine/Plugins/Experimental/EditorDataStorageFeatures/Source/TedsTableViewer/Private/TedsTableViewerModel.cpp b/Engine/Plugins/Experimental/EditorDataStorageFeatures/Source/TedsTableViewer/Private/TedsTableViewerModel.cpp
--- a/Engine/Plugins/Experimental/EditorDataStorageFeatures/Source/TedsTableViewer/Private/TedsTableViewerModel.cpp
+++ b/Engine/Plugins/Experimental/EditorDataStorageFeatures/Source/TedsTableViewer/Private/TedsTableViewerModel.cpp
@@ -4,7 +4,10 @@
 
 
 #include "Elements/Columns/TypedElementUIColumns.h"
-#include "Elements/Framework/TypedElementRegistry.h"
+#include "Elements/Common/EditorDataStorageFeatures.h"
+#include "Elements/Interfaces/TypedElementDataStorageInterface.h"
+#include "Elements/Interfaces/TypedElementDataStorageCompatibilityInterface.h"
+#include "Elements/Interfaces/TypedElementDataStorageUiInterface.h"
 #include "QueryStack/IQueryStackNode_Row.h"
 #include "TedsTableViewerColumn.h"
 #include "TedsTableViewerUtils.h"
@@ -19,15 +22,9 @@ namespace UE::Editor::DataStorage
 		, CellWidgetPurposes(InCellWidgetPurposes)
 		, IsItemVisible(InIsItemVisibleDelegate)
 	{
-		UTypedElementRegistry* Registry = UTypedElementRegistry::GetInstance();
-		
-		checkf(Registry, TEXT("Unable to create a Table Viewer before the Typed Element Registry is initialized."));
-		if (Registry)
-		{
-			Storage = Registry->GetMutableDataStorage();
-			StorageUi = Registry->GetMutableDataStorageUi();
-			StorageCompatibility = Registry->GetMutableDataStorageCompatibility();
-		}
+		Storage = GetMutableDataStorageFeature<IEditorDataStorageProvider>(StorageFeatureName);
+		StorageUi = GetMutableDataStorageFeature<IEditorDataStorageUiProvider>(UiFeatureName);
+		StorageCompatibility = GetMutableDataStorageFeature<IEditorDataStorageCompatibilityProvider>(CompatibilityFeatureName);
 		
 		TickerHandle = FTSTicker::GetCoreTicker().AddTicker(FTickerDelegate::CreateRaw(this, &FTedsTableViewerModel::Tick), 0);
 
diff --git a/Engine/Plugins/Experimental/EditorDataStorageFeatures/Source/TedsTableViewer/Private/TedsTableViewerModule.cpp b/Engine/Plugins/Experimental/EditorDataStorageFeatures/Source/TedsTableViewer/Private/TedsTableViewerModule.cpp
--- a/Engine/Plugins/Experimental/EditorDataStorageFeatures/Source/TedsTableViewer/Private/TedsTableViewerModule.cpp
+++ b/Engine/Plugins/Experimental/EditorDataStorageFeatures/Source/TedsTableViewer/Private/TedsTableViewerModule.cpp
@@ -3,7 +3,9 @@
 #include "Elements/Columns/TypedElementAlertColumns.h"
 #include "Elements/Columns/TypedElementCompatibilityColumns.h"
 #include "Elements/Columns/TypedElementLabelColumns.h"
-#include "Elements/Framework/TypedElementRegistry.h"
+#include "Elements/Common/EditorDataStorageFeatures.h"
+#include "Elements/Framework/TypedElementQueryBuilder.h"
+#include "Elements/Interfaces/TypedElementDataStorageInterface.h"
 #include "LevelEditor.h"
 #include "Modules/ModuleManager.h"
 #include "Modules/ModuleInterface.h"
@@ -62,15 +64,14 @@ public:
 	
 	TSharedRef<SWidget> CreateTableViewer()
 	{
-		UTypedElementRegistry* Registry = UTypedElementRegistry::GetInstance();
-		checkf(Registry, TEXT("Unable to initialize the table viewer before TEDS is initialized."));
+		using namespace UE::Editor::DataStorage;
 
-		if(!Registry->AreDataStorageInterfacesSet())
+		if(!AreEditorDataStorageFeaturesEnabled())
 		{
 			return SNullWidget::NullWidget;
 		}
 
-		IEditorDataStorageProvider* DataStorage = Registry->GetMutableDataStorage();
+		IEditorDataStorageProvider* DataStorage = GetMutableDataStorageFeature<IEditorDataStorageProvider>(StorageFeatureName);
 
 		using namespace UE::Editor::DataStorage::Queries;
 
@@ -84,14 +85,14 @@ public:
 
 		Rows.Empty();
 		
-		UE::Editor::DataStorage::FQueryResult QueryResult = DataStorage->RunQuery(QueryHandle,
-			CreateDirectQueryCallbackBinding([this](const IEditorDataStorageProvider::IDirectQueryContext& Context, const UE::Editor::DataStorage::RowHandle* RowHandles)
+		FQueryResult QueryResult = DataStorage->RunQuery(QueryHandle,
+			CreateDirectQueryCallbackBinding([this](const IEditorDataStorageProvider::IDirectQueryContext& Context, const RowHandle* RowHandles)
 			{
 				Rows.Append(RowHandles, Context.GetRowCount());
 			}));
 
-		return SNew(UE::Editor::DataStorage::STedsTableViewer)
-			.QueryStack(MakeShared<UE::Editor::DataStorage::FQueryStackNode_RowView>(&Rows))
+		return SNew(STedsTableViewer)
+			.QueryStack(MakeShared<FQueryStackNode_RowView>(&Rows))
 			.Columns({FTypedElementLabelColumn::StaticStruct(), FTypedElementSelectionColumn::StaticStruct(),
 				FTypedElementAlertColumn::StaticStruct(), FTypedElementChildAlertColumn::StaticStruct()});
 	}
diff --git a/Engine/Plugins/Experimental/EditorDataStorageFeatures/Source/TedsTableViewer/Private/Widgets/SRowDetails.cpp b/Engine/Plugins/Experimental/EditorDataStorageFeatures/Source/TedsTableViewer/Private/Widgets/SRowDetails.cpp
--- a/Engine/Plugins/Experimental/EditorDataStorageFeatures/Source/TedsTableViewer/Private/Widgets/SRowDetails.cpp
+++ b/Engine/Plugins/Experimental/EditorDataStorageFeatures/Source/TedsTableViewer/Private/Widgets/SRowDetails.cpp
@@ -4,7 +4,7 @@
 
 #include "Elements/Columns/TypedElementMiscColumns.h"
 #include "Elements/Columns/TypedElementTypeInfoColumns.h"
-#include "Elements/Framework/TypedElementRegistry.h"
+#include "Elements/Common/EditorDataStorageFeatures.h"
 #include "Elements/Interfaces/TypedElementDataStorageInterface.h"
 #include "Elements/Interfaces/TypedElementDataStorageUiInterface.h"
 #include "Widgets/Text/STextBlock.h"
@@ -36,12 +36,10 @@ namespace UE::Editor::DataStorage
 			WidgetPurposes = Widgets::Private::DefaultWidgetPurposes;
 		}
 
-		UTypedElementRegistry* Registry = UTypedElementRegistry::GetInstance();
-		checkf(Registry, TEXT("Unable to initialize SRowDetails before TEDS is initialized."));
-		checkf(Registry->AreDataStorageInterfacesSet(), TEXT("Unable to initialize SRowDetails without the editor data storage interfaces."));
+		checkf(AreEditorDataStorageFeaturesEnabled(), TEXT("Unable to initialize SRowDetails without the editor data storage interfaces."));
 
-		DataStorage = Registry->GetMutableDataStorage();
-		DataStorageUi = Registry->GetMutableDataStorageUi();
+		DataStorage = GetMutableDataStorageFeature<IEditorDataStorageProvider>(StorageFeatureName);
+		DataStorageUi = GetMutableDataStorageFeature<IEditorDataStorageUiProvider>(UiFeatureName);
 
 		
 		ChildSlot
diff --git a/Engine/Plugins/Experimental/InstanceDataObjectFixupTool/Source/Private/InstanceDataObjectFixupPanel.cpp b/Engine/Plugins/Experimental/InstanceDataObjectFixupTool/Source/Private/InstanceDataObjectFixupPanel.cpp
--- a/Engine/Plugins/Experimental/InstanceDataObjectFixupTool/Source/Private/InstanceDataObjectFixupPanel.cpp
+++ b/Engine/Plugins/Experimental/InstanceDataObjectFixupTool/Source/Private/InstanceDataObjectFixupPanel.cpp
@@ -13,7 +13,7 @@
 #include "UObject/PropertyBagRepository.h"
 
 #include "Elements/Columns/TypedElementAlertColumns.h"
-#include "Elements/Framework/TypedElementRegistry.h"
+#include "Elements/Common/EditorDataStorageFeatures.h"
 #include "Elements/Interfaces/TypedElementDataStorageInterface.h"
 #include "Elements/Interfaces/TypedElementDataStorageCompatibilityInterface.h"
 
@@ -250,9 +250,8 @@ FInstanceDataObjectFixupPanel::~FInstanceDataObjectFixupPanel()
 {
 	using namespace UE::Editor::DataStorage;
 
-	UTypedElementRegistry* Registry = UTypedElementRegistry::GetInstance();
-	IEditorDataStorageProvider* DataStorage = Registry ? Registry->GetMutableDataStorage() : nullptr;
-	IEditorDataStorageCompatibilityProvider* DataStorageCompatibility = Registry ? Registry->GetMutableDataStorageCompatibility() : nullptr;
+	IEditorDataStorageProvider* DataStorage = GetMutableDataStorageFeature<IEditorDataStorageProvider>(StorageFeatureName);
+	IEditorDataStorageCompatibilityProvider* DataStorageCompatibility = DataStorageCompatibility = GetMutableDataStorageFeature<IEditorDataStorageCompatibilityProvider>(CompatibilityFeatureName);
 
 	if (DataStorageCompatibility != nullptr && DataStorage != nullptr)
 	{
diff --git a/Engine/Plugins/Experimental/InstanceDataObjectFixupTool/Source/Private/InstanceDataObjectFixupToolTedsQueries.cpp b/Engine/Plugins/Experimental/InstanceDataObjectFixupTool/Source/Private/InstanceDataObjectFixupToolTedsQueries.cpp
--- a/Engine/Plugins/Experimental/InstanceDataObjectFixupTool/Source/Private/InstanceDataObjectFixupToolTedsQueries.cpp
+++ b/Engine/Plugins/Experimental/InstanceDataObjectFixupTool/Source/Private/InstanceDataObjectFixupToolTedsQueries.cpp
@@ -5,8 +5,8 @@
 #include "Elements/Columns/TypedElementAlertColumns.h"
 #include "Elements/Columns/TypedElementCompatibilityColumns.h"
 #include "Elements/Columns/TypedElementMiscColumns.h"
+#include "Elements/Common/EditorDataStorageFeatures.h"
 #include "Elements/Framework/TypedElementQueryBuilder.h"
-#include "Elements/Framework/TypedElementRegistry.h"
 #include "Elements/Interfaces/TypedElementDataStorageInterface.h"
 #include "Engine/World.h"
 #include "UObject/PropertyBagRepository.h"
@@ -61,7 +61,8 @@ void UInstanceDataObjectFixupToolTedsQueryFactory::ShowFixUpToolForLooseProperti
 
 void UInstanceDataObjectFixupToolTedsQueryFactory::ShowFixUpTool(UE::Editor::DataStorage::RowHandle Row, bool bRecurseIntoObject)
 {
-	IEditorDataStorageProvider* DataStorage = UTypedElementRegistry::GetInstance()->GetMutableDataStorage();
+	using namespace UE::Editor::DataStorage;
+	IEditorDataStorageProvider* DataStorage = GetMutableDataStorageFeature<IEditorDataStorageProvider>(StorageFeatureName);
 	if (FTypedElementUObjectColumn* ObjectColumn = DataStorage->GetColumn<FTypedElementUObjectColumn>(Row))
 	{
 		if (bRecurseIntoObject)
diff --git a/Engine/Source/Editor/UnrealEd/Private/Elements/Actor/ActorElementEditorSelectionInterface.cpp b/Engine/Source/Editor/UnrealEd/Private/Elements/Actor/ActorElementEditorSelectionInterface.cpp
--- a/Engine/Source/Editor/UnrealEd/Private/Elements/Actor/ActorElementEditorSelectionInterface.cpp
+++ b/Engine/Source/Editor/UnrealEd/Private/Elements/Actor/ActorElementEditorSelectionInterface.cpp
@@ -4,11 +4,13 @@
  
 #include "Elements/Actor/ActorElementData.h"
 #include "Elements/Columns/TypedElementSelectionColumns.h"
+#include "Elements/Common/EditorDataStorageFeatures.h"
 #include "Elements/Framework/EngineElementsLibrary.h"
 #include "Elements/Framework/TypedElementHandle.h"
 #include "Elements/Framework/TypedElementList.h"
-#include "Elements/Framework/TypedElementRegistry.h"
 #include "Elements/Interfaces/TypedElementSelectionInterface.h"
+#include "Elements/Interfaces/TypedElementDataStorageInterface.h"
+#include "Elements/Interfaces/TypedElementDataStorageCompatibilityInterface.h"
 #include "Elements/Object/ObjectElementEditorSelectionInterface.h"
 #include "GameFramework/Actor.h"
 #include "Serialization/Archive.h"
@@ -51,14 +53,13 @@ bool UActorElementEditorSelectionInterface::SelectElement(const FTypedElementHan
 	using namespace UE::Editor::DataStorage;
 
 	const AActor* Actor = ActorElementDataUtil::GetActorFromHandle(InElementHandle);
-	UTypedElementRegistry* Registry = InSelectionSet->GetRegistry();
 	// Add a selection column in TEDS
-	if (const IEditorDataStorageCompatibilityProvider* Compatibility = Registry->GetDataStorageCompatibility())
+	if (const IEditorDataStorageCompatibilityProvider* Compatibility = GetDataStorageFeature<IEditorDataStorageCompatibilityProvider>(CompatibilityFeatureName))
 	{
 		RowHandle Row = Compatibility->FindRowWithCompatibleObject(Actor);
 		if (Row != InvalidRowHandle)
 		{
-			if (IEditorDataStorageProvider* DataStorage = UTypedElementRegistry::GetInstance()->GetMutableDataStorage())
+			if (IEditorDataStorageProvider* DataStorage = GetMutableDataStorageFeature<IEditorDataStorageProvider>(StorageFeatureName))
 			{
 				DataStorage->AddColumn<FTypedElementSelectionColumn>(Row);
 
@@ -80,13 +81,12 @@ bool UActorElementEditorSelectionInterface::DeselectElement(const FTypedElementH
 	using namespace UE::Editor::DataStorage;
 
 	const AActor* Actor = ActorElementDataUtil::GetActorFromHandle(InElementHandle);
-	UTypedElementRegistry* Registry = InSelectionSet->GetRegistry();
-	if (const IEditorDataStorageCompatibilityProvider* Compatibility = Registry->GetDataStorageCompatibility())
+	if (const IEditorDataStorageCompatibilityProvider* Compatibility = GetDataStorageFeature<IEditorDataStorageCompatibilityProvider>(CompatibilityFeatureName))
 	{
 		RowHandle Row = Compatibility->FindRowWithCompatibleObject(Actor);
 		if (Row != InvalidRowHandle)
 		{
-			if (IEditorDataStorageProvider* DataStorage = UTypedElementRegistry::GetInstance()->GetMutableDataStorage())
+			if (IEditorDataStorageProvider* DataStorage = GetMutableDataStorageFeature<IEditorDataStorageProvider>(StorageFeatureName))
             {
             	DataStorage->RemoveColumn(Row, FTypedElementSelectionColumn::StaticStruct());
             }
diff --git a/Engine/Source/Editor/UnrealEd/Private/LevelEditorViewport.cpp b/Engine/Source/Editor/UnrealEd/Private/LevelEditorViewport.cpp
--- a/Engine/Source/Editor/UnrealEd/Private/LevelEditorViewport.cpp
+++ b/Engine/Source/Editor/UnrealEd/Private/LevelEditorViewport.cpp
@@ -44,12 +44,14 @@
 #include "UnrealEdGlobals.h"
 #include "Materials/MaterialExpressionTextureSample.h"
 #include "EditorSupportDelegates.h"
+#include "Elements/Common/EditorDataStorageFeatures.h"
 #include "Elements/Framework/TypedElementRegistry.h"
 #include "Elements/Framework/EngineElementsLibrary.h"
 #include "Elements/Framework/TypedElementCommonActions.h"
 #include "Elements/Framework/TypedElementListObjectUtil.h"
 #include "Elements/Framework/TypedElementViewportInteraction.h"
 #include "Elements/Interfaces/TypedElementObjectInterface.h"
+#include "Elements/Interfaces/TypedElementDataStorageCompatibilityInterface.h"
 #include "Elements/Actor/ActorElementLevelEditorViewportInteractionCustomization.h"
 #include "Elements/Component/ComponentElementLevelEditorViewportInteractionCustomization.h"
 #include "AudioDevice.h"
@@ -1659,6 +1661,7 @@ bool FLevelEditorViewportClient::UpdateDropPreviewElements(int32 MouseX, int32 M
 void FLevelEditorViewportClient::DestroyDropPreviewElements()
 {
 	using namespace LevelEditorViewportLocals;
+	using namespace UE::Editor::DataStorage;
 
 	if (!HasDropPreviewElements())
 	{
@@ -1670,8 +1673,7 @@ void FLevelEditorViewportClient::DestroyDropPreviewElements()
 	// deleting preview actors in UUnrealEdEngine::DeleteActors skips explicit handle deregistration,
 	// and although it does still happen in the immediately triggered garbage cleanup, that feels
 	// potentially brittle.
-	UTypedElementRegistry* TypedElementRegistry = UTypedElementRegistry::GetInstance();
-	if (IEditorDataStorageCompatibilityProvider* TedsCompat = TypedElementRegistry->GetMutableDataStorageCompatibility())
+	if (IEditorDataStorageCompatibilityProvider* TedsCompat = GetMutableDataStorageFeature<IEditorDataStorageCompatibilityProvider>(CompatibilityFeatureName))
 	{
 		DropPreviewElements->ForEachElement<ITypedElementObjectInterface>(
 			[this, TedsCompat](const TTypedElement<ITypedElementObjectInterface>& InElement)
diff --git a/Engine/Source/Runtime/Engine/Private/Elements/Framework/EngineElementsLibrary.cpp b/Engine/Source/Runtime/Engine/Private/Elements/Framework/EngineElementsLibrary.cpp
--- a/Engine/Source/Runtime/Engine/Private/Elements/Framework/EngineElementsLibrary.cpp
+++ b/Engine/Source/Runtime/Engine/Private/Elements/Framework/EngineElementsLibrary.cpp
@@ -2,8 +2,10 @@
 
 #include "Elements/Framework/EngineElementsLibrary.h"
 
+#include "Elements/Common/EditorDataStorageFeatures.h"
 #include "Elements/Framework/TypedElementRegistry.h"
 #include "Elements/Framework/TypedElementOwnerStore.h"
+#include "Elements/Interfaces/TypedElementDataStorageCompatibilityInterface.h"
 
 #include "Elements/Object/ObjectElementData.h"
 
@@ -379,8 +381,8 @@ void UEngineElementsLibrary::DestroyActorElement(const AActor* InActor, TTypedEl
 
 void UEngineElementsLibrary::RegisterActorElement(const AActor* InActor)
 {
-	UTypedElementRegistry* Registry = UTypedElementRegistry::GetInstance();
-	IEditorDataStorageCompatibilityProvider* Storage = Registry->GetMutableDataStorageCompatibility();
+	using namespace UE::Editor::DataStorage;
+	IEditorDataStorageCompatibilityProvider* Storage = GetMutableDataStorageFeature<IEditorDataStorageCompatibilityProvider>(CompatibilityFeatureName);
 	if (Storage)
 	{
 		Storage->AddCompatibleObject(const_cast<AActor*>(InActor));
@@ -389,8 +391,8 @@ void UEngineElementsLibrary::RegisterActorElement(const AActor* InActor)
 
 void UEngineElementsLibrary::UnregisterActorElement(const AActor* InActor)
 {
-	UTypedElementRegistry* Registry = UTypedElementRegistry::GetInstance();
-	if (IEditorDataStorageCompatibilityProvider* Storage = Registry->GetMutableDataStorageCompatibility())
+	using namespace UE::Editor::DataStorage;
+	if (IEditorDataStorageCompatibilityProvider* Storage = GetMutableDataStorageFeature<IEditorDataStorageCompatibilityProvider>(CompatibilityFeatureName))
 	{
 		Storage->RemoveCompatibleObject(const_cast<AActor*>(InActor));
 	}
diff --git a/Engine/Source/Runtime/TypedElementFramework/Private/Elements/Framework/TypedElementAttributeBinding.cpp b/Engine/Source/Runtime/TypedElementFramework/Private/Elements/Framework/TypedElementAttributeBinding.cpp
--- a/Engine/Source/Runtime/TypedElementFramework/Private/Elements/Framework/TypedElementAttributeBinding.cpp
+++ b/Engine/Source/Runtime/TypedElementFramework/Private/Elements/Framework/TypedElementAttributeBinding.cpp
@@ -2,17 +2,17 @@
 
 #include "Elements/Framework/TypedElementAttributeBinding.h"
 
-#include "Elements/Framework/TypedElementRegistry.h"
+#include "Elements/Common/EditorDataStorageFeatures.h"
 
 namespace UE::Editor::DataStorage
 {
 	FAttributeBinder::FAttributeBinder(RowHandle InTargetRow)
 		: TargetRow(InTargetRow)
 	{
-		if (ensureMsgf(UTypedElementRegistry::GetInstance()->AreDataStorageInterfacesSet(), 
+		if (ensureMsgf(AreEditorDataStorageFeaturesEnabled(),
 				TEXT("The Editor Data Storage plugin needs to be enabled to use attribute bindings.")))
 		{
-			DataStorage = UTypedElementRegistry::GetInstance()->GetMutableDataStorage();
+			DataStorage = GetMutableDataStorageFeature<IEditorDataStorageProvider>(StorageFeatureName);
 		}
 	}
 
diff --git a/Engine/Source/Runtime/TypedElementFramework/Private/Elements/Framework/TypedElementDataStorageWidget.cpp b/Engine/Source/Runtime/TypedElementFramework/Private/Elements/Framework/TypedElementDataStorageWidget.cpp
--- a/Engine/Source/Runtime/TypedElementFramework/Private/Elements/Framework/TypedElementDataStorageWidget.cpp
+++ b/Engine/Source/Runtime/TypedElementFramework/Private/Elements/Framework/TypedElementDataStorageWidget.cpp
@@ -2,7 +2,8 @@
 
 #include "Elements/Framework/TypedElementDataStorageWidget.h"
 #include "Elements/Columns/TypedElementSlateWidgetColumns.h"
-#include "Elements/Framework/TypedElementRegistry.h"
+#include "Elements/Common//EditorDataStorageFeatures.h"
+#include "Elements/Interfaces/TypedElementDataStorageInterface.h"
 
 STedsWidget::STedsWidget()
 	: UiRowHandle(UE::Editor::DataStorage::InvalidRowHandle)
@@ -75,12 +76,6 @@ UE::Editor::DataStorage::RowHandle STedsWidget::GetRowHandle() const
 
 IEditorDataStorageProvider* STedsWidget::GetStorageIfAvailable()
 {
-	UTypedElementRegistry* Registry = UTypedElementRegistry::GetInstance();
-
-	if(!Registry || !Registry->AreDataStorageInterfacesSet())
-	{
-		return nullptr;
-	}
-
-	return Registry->GetMutableDataStorage();
+	using namespace UE::Editor::DataStorage;
+	return GetMutableDataStorageFeature<IEditorDataStorageProvider>(StorageFeatureName);
 }
diff --git a/Engine/Source/Runtime/TypedElementFramework/Public/Elements/Common/EditorDataStorageFeatures.h b/Engine/Source/Runtime/TypedElementFramework/Public/Elements/Common/EditorDataStorageFeatures.h
--- a/Engine/Source/Runtime/TypedElementFramework/Public/Elements/Common/EditorDataStorageFeatures.h
+++ b/Engine/Source/Runtime/TypedElementFramework/Public/Elements/Common/EditorDataStorageFeatures.h
@@ -14,7 +14,7 @@ namespace UE::Editor::DataStorage
 	TYPEDELEMENTFRAMEWORK_API const extern FName UiFeatureName;
 
 	TYPEDELEMENTFRAMEWORK_API FSimpleMulticastDelegate& OnEditorDataStorageFeaturesEnabled();
-	bool AreEditorDataStorageFeaturesEnabled();
+	TYPEDELEMENTFRAMEWORK_API bool AreEditorDataStorageFeaturesEnabled();
 
 	template<typename T>
 	const T* GetDataStorageFeature(const FName InName)
diff --git a/Engine/Source/Runtime/TypedElementFramework/Public/Elements/Framework/TypedElementRegistry.h b/Engine/Source/Runtime/TypedElementFramework/Public/Elements/Framework/TypedElementRegistry.h
--- a/Engine/Source/Runtime/TypedElementFramework/Public/Elements/Framework/TypedElementRegistry.h
+++ b/Engine/Source/Runtime/TypedElementFramework/Public/Elements/Framework/TypedElementRegistry.h
@@ -77,25 +77,25 @@ public:
 	UFUNCTION(BlueprintPure, DisplayName="Get Default Typed Element Registry", Category = "TypedElementFramework|Registry", meta=(ScriptName="GetDefaultTypedElementRegistry"))
 	static TYPEDELEMENTFRAMEWORK_API UTypedElementRegistry* GetInstance();
 
-	//UE_DEPRECATED(5.5, "Function has moved to modular features")
+	UE_DEPRECATED(5.5, "Function has moved to modular features. Use UE::Editor::DataStorage::GetMutableDataStorageFeature instead.")
 	TYPEDELEMENTFRAMEWORK_API IEditorDataStorageProvider* GetMutableDataStorage();
-	//UE_DEPRECATED(5.5, "Function has moved to modular features")
+	UE_DEPRECATED(5.5, "Function has moved to modular features. Use UE::Editor::DataStorage::GetDataStorageFeature instead.")
 	TYPEDELEMENTFRAMEWORK_API const IEditorDataStorageProvider* GetDataStorage() const;
 
-	//UE_DEPRECATED(5.5, "Function has moved to modular features")
+	UE_DEPRECATED(5.5, "Function has moved to modular features. Use UE::Editor::DataStorage::GetMutableDataStorageFeature instead.")
 	TYPEDELEMENTFRAMEWORK_API IEditorDataStorageCompatibilityProvider* GetMutableDataStorageCompatibility();
-	//UE_DEPRECATED(5.5, "Function has moved to modular features")
+	UE_DEPRECATED(5.5, "Function has moved to modular features. Use UE::Editor::DataStorage::GetDataStorageFeature instead.")
 	TYPEDELEMENTFRAMEWORK_API const IEditorDataStorageCompatibilityProvider* GetDataStorageCompatibility() const;
 
-	//UE_DEPRECATED(5.5, "Function has moved to modular features")
+	UE_DEPRECATED(5.5, "Function has moved to modular features. Use UE::Editor::DataStorage::GetMutableDataStorageFeature instead.")
 	TYPEDELEMENTFRAMEWORK_API IEditorDataStorageUiProvider* GetMutableDataStorageUi();
-	//UE_DEPRECATED(5.5, "Function has moved to modular features")
+	UE_DEPRECATED(5.5, "Function has moved to modular features. Use UE::Editor::DataStorage::GetDataStorageFeature instead.")
 	TYPEDELEMENTFRAMEWORK_API const IEditorDataStorageUiProvider* GetDataStorageUi() const;
 
-	//UE_DEPRECATED(5.5, "Use UE::Editor::DataStorage::AreDataStorageInterfacesSet")
+	UE_DEPRECATED(5.5, "Use UE::Editor::DataStorage::AreDataStorageInterfacesSet instead.")
 	TYPEDELEMENTFRAMEWORK_API bool AreDataStorageInterfacesSet() const;
 
-	//UE_DEPRECATED(5.5, "Use UE::Editor::DataStorage::OnAllStorageFeaturesEnabled")
+	UE_DEPRECATED(5.5, "Use UE::Editor::DataStorage::OnAllStorageFeaturesEnabled instead")
 	TYPEDELEMENTFRAMEWORK_API FSimpleMulticastDelegate& OnDataStorageInterfacesSet();
 
 	/**
diff --git a/Engine/Source/Runtime/TypedElementFramework/Tests/Elements/Framework/TypedElementDataStoragePerformanceTestCommands.cpp b/Engine/Source/Runtime/TypedElementFramework/Tests/Elements/Framework/TypedElementDataStoragePerformanceTestCommands.cpp
--- a/Engine/Source/Runtime/TypedElementFramework/Tests/Elements/Framework/TypedElementDataStoragePerformanceTestCommands.cpp
+++ b/Engine/Source/Runtime/TypedElementFramework/Tests/Elements/Framework/TypedElementDataStoragePerformanceTestCommands.cpp
@@ -2,9 +2,9 @@
 
 #include "TypedElementDataStoragePerformanceTestCommands.h"
 
+#include "Elements/Common/EditorDataStorageFeatures.h"
 #include "Elements/Interfaces/TypedElementDataStorageInterface.h"
 #include "Elements/Framework/TypedElementQueryBuilder.h"
-#include "Elements/Framework/TypedElementRegistry.h"
 #include "HAL/IConsoleManager.h"
 
 namespace UE::Editor::DataStorage
@@ -24,7 +24,7 @@ namespace UE::Editor::DataStorage
 				int32 EntitiesToAdd;
 				LexFromString(EntitiesToAdd, *Args[0]);
 
-				IEditorDataStorageProvider* DataStorageInterface = UTypedElementRegistry::GetInstance()->GetMutableDataStorage();
+				IEditorDataStorageProvider* DataStorageInterface = GetMutableDataStorageFeature<IEditorDataStorageProvider>(StorageFeatureName);
 			
 				DataStorageInterface->BatchAddRow(Private::PerformanceTestCommandTable, EntitiesToAdd, [DataStorageInterface](RowHandle Row)
 					{
@@ -42,7 +42,7 @@ namespace UE::Editor::DataStorage
 		{
 			using namespace UE::Editor::DataStorage::Queries;
 
-			IEditorDataStorageProvider* DataStorageInterface = UTypedElementRegistry::GetInstance()->GetMutableDataStorage();
+			IEditorDataStorageProvider* DataStorageInterface = GetMutableDataStorageFeature<IEditorDataStorageProvider>(StorageFeatureName);
 		
 			QueryHandle Query = DataStorageInterface->RegisterQuery(
 				Select().Where().All<FTest_PingPongPrePhys>().Compile());
diff --git a/Engine/Source/Runtime/TypedElementFramework/Tests/Elements/Framework/TypedElementDataStorageProcessorTests.cpp b/Engine/Source/Runtime/TypedElementFramework/Tests/Elements/Framework/TypedElementDataStorageProcessorTests.cpp
--- a/Engine/Source/Runtime/TypedElementFramework/Tests/Elements/Framework/TypedElementDataStorageProcessorTests.cpp
+++ b/Engine/Source/Runtime/TypedElementFramework/Tests/Elements/Framework/TypedElementDataStorageProcessorTests.cpp
@@ -1,7 +1,9 @@
 // Copyright Epic Games, Inc. All Rights Reserved.
 
 #if WITH_TESTS
-#include "Elements/Framework/TypedElementRegistry.h"
+#include "Elements/Common/EditorDataStorageFeatures.h"
+#include "Elements/Framework/TypedElementQueryBuilder.h"
+#include "Elements/Interfaces/TypedElementDataStorageInterface.h"
 #include "HAL/IConsoleManager.h"
 #include "TypedElementTestColumns.h"
 
@@ -28,7 +30,7 @@ namespace UE::Editor::DataStorage::Debug::ProcessorTests
 	{
 		using namespace UE::Editor::DataStorage::Queries;
 
-		IEditorDataStorageProvider* DataStorage = UTypedElementRegistry::GetInstance()->GetMutableDataStorage();
+		IEditorDataStorageProvider* DataStorage = GetMutableDataStorageFeature<IEditorDataStorageProvider>(StorageFeatureName);
 
 		if (PrimaryTable == InvalidTableHandle)
 		{
@@ -122,7 +124,7 @@ namespace UE::Editor::DataStorage::Debug::ProcessorTests
 
 	void UnregisterProcessors()
 	{
-		IEditorDataStorageProvider* DataStorage = UTypedElementRegistry::GetInstance()->GetMutableDataStorage();
+		IEditorDataStorageProvider* DataStorage = GetMutableDataStorageFeature<IEditorDataStorageProvider>(StorageFeatureName);
 		for (const UE::Editor::DataStorage::QueryHandle& Handle : RegisteredQueries)
 		{
 			DataStorage->UnregisterQuery(Handle);
@@ -145,7 +147,7 @@ namespace UE::Editor::DataStorage::Debug::ProcessorTests
 				return;
 			}
 
-			IEditorDataStorageProvider* DataStorage = UTypedElementRegistry::GetInstance()->GetMutableDataStorage();
+			IEditorDataStorageProvider* DataStorage = GetMutableDataStorageFeature<IEditorDataStorageProvider>(StorageFeatureName);
 			DataStorage->BatchAddRow(PrimaryTable, RowsToCreate, [](UE::Editor::DataStorage::RowHandle Row)
 			{
 			});
@@ -178,4 +180,3 @@ namespace UE::Editor::DataStorage::Debug::ProcessorTests
 	};
 } // namespace UE::Editor::DataStorage::Debug::ProcessorTests
 #endif
-
diff --git a/Engine/Source/Runtime/TypedElementFramework/Tests/Elements/Framework/TypedElementQueryBuilderTests.cpp b/Engine/Source/Runtime/TypedElementFramework/Tests/Elements/Framework/TypedElementQueryBuilderTests.cpp
--- a/Engine/Source/Runtime/TypedElementFramework/Tests/Elements/Framework/TypedElementQueryBuilderTests.cpp
+++ b/Engine/Source/Runtime/TypedElementFramework/Tests/Elements/Framework/TypedElementQueryBuilderTests.cpp
@@ -1,6 +1,5 @@
 // Copyright Epic Games, Inc. All Rights Reserved.
 
-#include "Elements/Framework/TypedElementRegistry.h"
 #if WITH_TESTS
 #include "TypedElementTestColumns.h"
 
@@ -8,7 +7,9 @@
 #include "CoreTypes.h"
 #include "Containers/Array.h"
 #include "Containers/UnrealString.h"
+#include "Elements/Common/EditorDataStorageFeatures.h"
 #include "Elements/Common/TypedElementQueryConditions.h"
+#include "Elements/Interfaces/TypedElementDataStorageInterface.h"
 #include "Tests/TestHarnessAdapter.h"
 #include "UObject/Class.h"
 
@@ -31,9 +32,9 @@ static void AppendColumnName(FString& Output, TWeakObjectPtr<const UScriptStruct
 
 static bool TestMatching(FConditions& TestQuery, const TArray<FColumnBase>& RequestedColumns, bool Expected, bool Sort = false)
 {
-	IEditorDataStorageProvider* Storage = UTypedElementRegistry::GetInstance()->GetMutableDataStorage();
+	IEditorDataStorageProvider* Storage = GetMutableDataStorageFeature<IEditorDataStorageProvider>(StorageFeatureName);
 
-	TestQuery.Compile(UE::Editor::DataStorage::Queries::FEditorStorageQueryConditionCompileContext(Storage));
+	TestQuery.Compile(FEditorStorageQueryConditionCompileContext(Storage));
 	
 	if (Sort)
 	{
