[Backout] - CL36477705
Original CL Desc
-----------------------------------------------------------------
[TEDS] Deprecate accessor functions for TEDS features inside the TypedElementRegistry class

diff --git a/Engine/Plugins/Developer/PerforceSourceControl/Source/PerforceSourceControl/Private/PerforceSourceControlOperations.cpp b/Engine/Plugins/Developer/PerforceSourceControl/Source/PerforceSourceControl/Private/PerforceSourceControlOperations.cpp
--- a/Engine/Plugins/Developer/PerforceSourceControl/Source/PerforceSourceControl/Private/PerforceSourceControlOperations.cpp
+++ b/Engine/Plugins/Developer/PerforceSourceControl/Source/PerforceSourceControl/Private/PerforceSourceControlOperations.cpp
@@ -25,8 +25,8 @@
 #include "Elements/Columns/TypedElementMiscColumns.h"
 #include "Elements/Columns/TypedElementPackageColumns.h"
 #include "Elements/Columns/TypedElementRevisionControlColumns.h"
-#include "Elements/Common/EditorDataStorageFeatures.h"
 #include "Elements/Framework/TypedElementIndexHasher.h"
+#include "Elements/Framework/TypedElementRegistry.h"
 #include "Elements/Interfaces/TypedElementDataStorageInterface.h"
 #include "ProfilingDebugging/CpuProfilerTrace.h"
 
@@ -2031,10 +2031,15 @@ bool FPerforceUpdateStatusWorker::UpdateStates() const
 
 	auto UpdateDataStorage = [](const FPerforceSourceControlState& State)
 	{
-		using namespace UE::Editor::DataStorage;
 		using namespace UE::Editor::DataStorage::Queries;
 
-		IEditorDataStorageProvider* DataStorage = GetMutableDataStorageFeature<IEditorDataStorageProvider>(StorageFeatureName);
+		UTypedElementRegistry* Registry = UTypedElementRegistry::GetInstance();
+		if (!Registry)
+		{
+			return;
+		}
+
+		IEditorDataStorageProvider* DataStorage = Registry->GetMutableDataStorage();
 		if (!DataStorage)
 		{
 			return;
diff --git a/Engine/Plugins/Experimental/Animation/HierarchyTable/Source/Editor/Private/HierarchyTableEditorToolkit.cpp b/Engine/Plugins/Experimental/Animation/HierarchyTable/Source/Editor/Private/HierarchyTableEditorToolkit.cpp
--- a/Engine/Plugins/Experimental/Animation/HierarchyTable/Source/Editor/Private/HierarchyTableEditorToolkit.cpp
+++ b/Engine/Plugins/Experimental/Animation/HierarchyTable/Source/Editor/Private/HierarchyTableEditorToolkit.cpp
@@ -5,12 +5,11 @@
 #include "Modules/ModuleManager.h"
 #include "PropertyEditorModule.h"
 #include "TedsOutlinerModule.h"
+#include "Elements/Framework/TypedElementRegistry.h"
 #include "Elements/Columns/TypedElementLabelColumns.h"
 #include "Elements/Columns/TypedElementTypeInfoColumns.h"
 #include "Elements/Columns/TypedElementSelectionColumns.h"
 #include "Elements/Columns/TypedElementMiscColumns.h"
-#include "Elements/Common/EditorDataStorageFeatures.h"
-#include "Elements/Framework/TypedElementQueryBuilder.h"
 #include "Elements/Interfaces/TypedElementDataStorageInterface.h"
 #include "SceneOutlinerPublicTypes.h"
 #include "TedsOutlinerMode.h"
@@ -89,10 +88,10 @@ void FHierarchyTableEditorToolkit::RegisterTabSpawners(const TSharedRef<class FT
 
 void FHierarchyTableEditorToolkit::OnClose()
 {
-	using namespace UE::Editor::DataStorage;
-	IEditorDataStorageProvider* DSI = GetMutableDataStorageFeature<IEditorDataStorageProvider>(StorageFeatureName);
+	UTypedElementRegistry* Registry = UTypedElementRegistry::GetInstance();
+	IEditorDataStorageProvider* DSI = Registry->GetMutableDataStorage();
 
-	for (const TTuple<int32, RowHandle>& Row : EntryIndexToHandleMap)
+	for (const TTuple<int32, UE::Editor::DataStorage::RowHandle>& Row : EntryIndexToHandleMap)
 	{
 		DSI->RemoveRow(Row.Value);
 	}
@@ -107,15 +106,18 @@ void FHierarchyTableEditorToolkit::UnregisterTabSpawners(const TSharedRef<class
 
 TSharedRef<SWidget> FHierarchyTableEditorToolkit::CreateTedsOutliner()
 {
-	using namespace UE::Editor::DataStorage::Queries;
-	using namespace UE::Editor::Outliner;
+	UTypedElementRegistry* Registry = UTypedElementRegistry::GetInstance();
+	checkf(Registry, TEXT("Unable to initialize the table viewer before TEDS is initialized."));
 
-	if (!AreEditorDataStorageFeaturesEnabled())
+	if (!Registry->AreDataStorageInterfacesSet())
 	{
 		return SNew(STextBlock)
 			.Text(INVTEXT("You need to enable the Typed Element Data Storage plugin to see the table viewer!"));
 	}
 
+	using namespace UE::Editor::DataStorage::Queries;
+	using namespace UE::Editor::Outliner;
+
 	if (!ensure(HierarchyTable->TableType))
 	{
 		return SNullWidget::NullWidget;
@@ -136,7 +138,7 @@ TSharedRef<SWidget> FHierarchyTableEditorToolkit::CreateTedsOutliner()
 			.ReadOnly(HierarchyTableTypeColumns)
 		.Compile();
 
-	InitialColumnQuery = GetMutableDataStorageFeature<IEditorDataStorageProvider>(StorageFeatureName)->RegisterQuery(MoveTemp(ColumnQueryDescription));
+	InitialColumnQuery = Registry->GetMutableDataStorage()->RegisterQuery(MoveTemp(ColumnQueryDescription));
 
 	FSceneOutlinerInitializationOptions InitOptions;
 	InitOptions.bShowHeaderRow = true;
@@ -159,7 +161,7 @@ TSharedRef<SWidget> FHierarchyTableEditorToolkit::CreateTedsOutliner()
 
 	FTedsOutlinerModule& TedsOutlinerModule = FModuleManager::GetModuleChecked<FTedsOutlinerModule>("TedsOutliner");
 
-	IEditorDataStorageProvider* DSI = GetMutableDataStorageFeature<IEditorDataStorageProvider>(StorageFeatureName);
+	IEditorDataStorageProvider* DSI = Registry->GetMutableDataStorage();
 	static TableHandle Table = DSI->FindTable(FName("Editor_HierarchyTableTable"));
 	
 	TArray<UScriptStruct*> BaseHierarchyTableTypeColumns = Handler->GetColumns();
@@ -323,7 +325,8 @@ void FHierarchyTableEditorToolkit::AddEntry(const FName Identifier, const EHiera
 	}
 
 	using namespace UE::Editor::DataStorage;
-	IEditorDataStorageProvider* DSI = GetMutableDataStorageFeature<IEditorDataStorageProvider>(StorageFeatureName);
+	UTypedElementRegistry* Registry = UTypedElementRegistry::GetInstance();
+	IEditorDataStorageProvider* DSI = Registry->GetMutableDataStorage();
 	static TableHandle Table = DSI->FindTable(FName("Editor_HierarchyTableTable"));
 
 	int32 ParentIndex = 0; // root of hierarchy
diff --git a/Engine/Plugins/Experimental/EditorDataStorage/Source/TedsCore/Private/Compatibility/TedsTypedElementActorHandleFactory.cpp b/Engine/Plugins/Experimental/EditorDataStorage/Source/TedsCore/Private/Compatibility/TedsTypedElementActorHandleFactory.cpp
--- a/Engine/Plugins/Experimental/EditorDataStorage/Source/TedsCore/Private/Compatibility/TedsTypedElementActorHandleFactory.cpp
+++ b/Engine/Plugins/Experimental/EditorDataStorage/Source/TedsCore/Private/Compatibility/TedsTypedElementActorHandleFactory.cpp
@@ -5,11 +5,10 @@
 #include "Compatibility/TedsTypedElementBridge.h"
 #include "Compatibility/Columns/TypedElement.h"
 #include "Elements/Columns/TypedElementCompatibilityColumns.h"
-#include "Elements/Common/EditorDataStorageFeatures.h"
 #include "Elements/Framework/EngineElementsLibrary.h"
 #include "Elements/Framework/TypedElementList.h"
 #include "Elements/Framework/TypedElementQueryBuilder.h"
-#include "Elements/Interfaces/TypedElementDataStorageInterface.h"
+#include "Elements/Framework/TypedElementRegistry.h"
 #include "GameFramework/Actor.h"
 
 void UTypedElementActorHandleDataStorageFactory::PreRegister(IEditorDataStorageProvider& DataStorage)
@@ -76,7 +75,7 @@ void UTypedElementActorHandleDataStorageFactory::HandleBridgeEnabled(bool bEnabl
 {
 	using namespace UE::Editor::DataStorage;
 
-	IEditorDataStorageProvider* DataStorage = GetMutableDataStorageFeature<IEditorDataStorageProvider>(StorageFeatureName);
+	IEditorDataStorageProvider* DataStorage = UTypedElementRegistry::GetInstance()->GetMutableDataStorage();
 
 	if (bEnabled)
 	{
diff --git a/Engine/Plugins/Experimental/EditorDataStorage/Source/TedsCore/Private/Compatibility/TedsTypedElementBridgeQueries.cpp b/Engine/Plugins/Experimental/EditorDataStorage/Source/TedsCore/Private/Compatibility/TedsTypedElementBridgeQueries.cpp
--- a/Engine/Plugins/Experimental/EditorDataStorage/Source/TedsCore/Private/Compatibility/TedsTypedElementBridgeQueries.cpp
+++ b/Engine/Plugins/Experimental/EditorDataStorage/Source/TedsCore/Private/Compatibility/TedsTypedElementBridgeQueries.cpp
@@ -5,11 +5,10 @@
 #include "Compatibility/TedsTypedElementBridge.h"
 #include "Compatibility/Columns/TypedElement.h"
 #include "Elements/Columns/TypedElementCompatibilityColumns.h"
-#include "Elements/Common/EditorDataStorageFeatures.h"
 #include "Elements/Framework/EngineElementsLibrary.h"
 #include "Elements/Framework/TypedElementList.h"
 #include "Elements/Framework/TypedElementQueryBuilder.h"
-#include "Elements/Interfaces/TypedElementDataStorageInterface.h"
+#include "Elements/Framework/TypedElementRegistry.h"
 #include "HAL/IConsoleManager.h"
 
 namespace UE::Editor::DataStorage::Compatibility::Private
@@ -91,18 +90,17 @@ void UTypedElementBridgeDataStorageFactory::CleanupTypedElementColumns(IEditorDa
 
 void UTypedElementBridgeDataStorageFactory::HandleOnEnabled(IConsoleVariable* CVar)
 {
-	using namespace UE::Editor::DataStorage;
-	IEditorDataStorageProvider* DataStorage = GetMutableDataStorageFeature<IEditorDataStorageProvider>(StorageFeatureName);
+	IEditorDataStorageProvider* DataStorage = UTypedElementRegistry::GetInstance()->GetMutableDataStorage();
 	bool bIsEnabled = CVar->GetBool();
 
 	if (bIsEnabled)
 	{
 		RegisterQuery_NewUObject(*DataStorage);
-		Compatibility::OnTypedElementBridgeEnabled().Broadcast(bIsEnabled);
+		UE::Editor::DataStorage::Compatibility::OnTypedElementBridgeEnabled().Broadcast(bIsEnabled);
 	}
 	else
 	{
-		Compatibility::OnTypedElementBridgeEnabled().Broadcast(bIsEnabled);
+		UE::Editor::DataStorage::Compatibility::OnTypedElementBridgeEnabled().Broadcast(bIsEnabled);
 		UnregisterQuery_NewUObject(*DataStorage);
 		CleanupTypedElementColumns(*DataStorage);
 	}
diff --git a/Engine/Plugins/Experimental/EditorDataStorage/Source/TedsCore/Private/Compatibility/TypedElementActorViewportProcessors.cpp b/Engine/Plugins/Experimental/EditorDataStorage/Source/TedsCore/Private/Compatibility/TypedElementActorViewportProcessors.cpp
--- a/Engine/Plugins/Experimental/EditorDataStorage/Source/TedsCore/Private/Compatibility/TypedElementActorViewportProcessors.cpp
+++ b/Engine/Plugins/Experimental/EditorDataStorage/Source/TedsCore/Private/Compatibility/TypedElementActorViewportProcessors.cpp
@@ -6,9 +6,8 @@
 #include "Elements/Columns/TypedElementCompatibilityColumns.h"
 #include "Elements/Columns/TypedElementViewportColumns.h"
 #include "Elements/Columns/TypedElementMiscColumns.h"
-#include "Elements/Common/EditorDataStorageFeatures.h"
 #include "Elements/Framework/TypedElementQueryBuilder.h"
-#include "Elements/Interfaces/TypedElementDataStorageInterface.h"
+#include "Elements/Framework/TypedElementRegistry.h"
 #include "HAL/IConsoleManager.h"
 #include "Elements/Columns/TypedElementSelectionColumns.h"
 #include "ProfilingDebugging/CpuProfilerTrace.h"
@@ -24,7 +23,7 @@ namespace UE::Editor::DataStorage::Private
 
 				TRACE_CPUPROFILER_EVENT_SCOPE(TEDS.Debug.AddOverlayColorToSelectionCommand);
 
-				if (IEditorDataStorageProvider* DataStorage = GetMutableDataStorageFeature<IEditorDataStorageProvider>(StorageFeatureName))
+				if (IEditorDataStorageProvider* DataStorage = UTypedElementRegistry::GetInstance()->GetMutableDataStorage())
 				{
 					static QueryHandle OverlayQuery = InvalidQueryHandle;
 					if (OverlayQuery == InvalidQueryHandle)
@@ -82,7 +81,7 @@ namespace UE::Editor::DataStorage::Private
 
 				TRACE_CPUPROFILER_EVENT_SCOPE(TEDS.Debug.AddOverlayColorToSelectionCommand);
 
-				if (IEditorDataStorageProvider* DataStorage = GetMutableDataStorageFeature<IEditorDataStorageProvider>(StorageFeatureName))
+				if (IEditorDataStorageProvider* DataStorage = UTypedElementRegistry::GetInstance()->GetMutableDataStorage())
 				{
 					static QueryHandle OverlayQuery = InvalidQueryHandle;
 					if (OverlayQuery == InvalidQueryHandle)
@@ -133,7 +132,7 @@ namespace UE::Editor::DataStorage::Private
 
 				TRACE_CPUPROFILER_EVENT_SCOPE(TEDS.Debug.AddOverlayColorToSelectionCommand);
 
-				if (IEditorDataStorageProvider* DataStorage = GetMutableDataStorageFeature<IEditorDataStorageProvider>(StorageFeatureName))
+				if (IEditorDataStorageProvider* DataStorage = UTypedElementRegistry::GetInstance()->GetMutableDataStorage())
 				{
 					static QueryHandle OverlayQuery = InvalidQueryHandle;
 					if (OverlayQuery == InvalidQueryHandle)
diff --git a/Engine/Plugins/Experimental/EditorDataStorage/Source/TedsCore/Private/Debug/TypedElementDatabaseDebugCommands.cpp b/Engine/Plugins/Experimental/EditorDataStorage/Source/TedsCore/Private/Debug/TypedElementDatabaseDebugCommands.cpp
--- a/Engine/Plugins/Experimental/EditorDataStorage/Source/TedsCore/Private/Debug/TypedElementDatabaseDebugCommands.cpp
+++ b/Engine/Plugins/Experimental/EditorDataStorage/Source/TedsCore/Private/Debug/TypedElementDatabaseDebugCommands.cpp
@@ -4,9 +4,9 @@
 #include "Elements/Columns/TypedElementCompatibilityColumns.h"
 #include "Elements/Columns/TypedElementLabelColumns.h"
 #include "Elements/Columns/TypedElementTypeInfoColumns.h"
-#include "Elements/Common/EditorDataStorageFeatures.h"
 #include "Elements/Common/TypedElementDataStorageLog.h"
 #include "Elements/Framework/TypedElementQueryBuilder.h"
+#include "Elements/Framework/TypedElementRegistry.h"
 #include "Elements/Framework/TypedElementTestColumns.h"
 #include "Elements/Interfaces/TypedElementDataStorageInterface.h"
 #include "Elements/Interfaces/TypedElementDataStorageCompatibilityInterface.h"
@@ -65,7 +65,7 @@ namespace UE::Editor::DataStorage::Private
 	{
 		using namespace UE::Editor::DataStorage::Queries;
 
-		if (IEditorDataStorageProvider* DataStorage = GetMutableDataStorageFeature<IEditorDataStorageProvider>(StorageFeatureName))
+		if (IEditorDataStorageProvider* DataStorage = UTypedElementRegistry::GetInstance()->GetMutableDataStorage())
 		{
 			static QueryHandle LabelQuery = [DataStorage]
 			{
@@ -136,7 +136,7 @@ FAutoConsoleCommandWithOutputDevice PrintObjectTypeInformationConsoleCommand(
 
 			TRACE_CPUPROFILER_EVENT_SCOPE(TEDS.Debug.PrintObjectTypeInfo);
 
-			if (IEditorDataStorageProvider* DataStorage = GetMutableDataStorageFeature<IEditorDataStorageProvider>(StorageFeatureName))
+			if (IEditorDataStorageProvider* DataStorage = UTypedElementRegistry::GetInstance()->GetMutableDataStorage())
 			{
 				FString Message;
 				Output.Log(TEXT("The Typed Elements Data Storage has the types:"));
@@ -179,7 +179,6 @@ FAutoConsoleCommandWithOutputDevice ListExtensionsConsoleCommand(
 	FConsoleCommandWithOutputDeviceDelegate::CreateLambda([](FOutputDevice& Output)
 		{
 			TRACE_CPUPROFILER_EVENT_SCOPE(TEDS.Debug.ListExtensions);
-			using namespace UE::Editor::DataStorage;
 
 			FString Message;
 			auto RecordExtensions = [&Message](FName Extension)
@@ -189,17 +188,19 @@ FAutoConsoleCommandWithOutputDevice ListExtensionsConsoleCommand(
 				Message += TEXT('\n');
 			};
 
-			if (const IEditorDataStorageProvider* DataStorage = GetDataStorageFeature<IEditorDataStorageProvider>(StorageFeatureName))
+			UTypedElementRegistry* Registry = UTypedElementRegistry::GetInstance();
+
+			if (const IEditorDataStorageProvider* DataStorage = Registry->GetDataStorage())
 			{
 				Message = TEXT("Data Storage Extensions: \n");
 				DataStorage->ListExtensions(RecordExtensions);
 			}
-			if (const IEditorDataStorageCompatibilityProvider* DataStorageCompat = GetDataStorageFeature<IEditorDataStorageCompatibilityProvider>(CompatibilityFeatureName))
+			if (const IEditorDataStorageCompatibilityProvider* DataStorageCompat = Registry->GetDataStorageCompatibility())
 			{
 				Message += TEXT("Data Storage Compatibility Extensions: \n");
 				DataStorageCompat->ListExtensions(RecordExtensions);
 			}
-			if (const IEditorDataStorageUiProvider* DataStorageUi = GetDataStorageFeature<IEditorDataStorageUiProvider>(UiFeatureName))
+			if (const IEditorDataStorageUiProvider* DataStorageUi = Registry->GetDataStorageUi())
 			{
 				Message += TEXT("Data Storage UI Extensions: \n");
 				DataStorageUi->ListExtensions(RecordExtensions);
@@ -216,11 +217,10 @@ static FAutoConsoleCommand CVarCreateRow(
 	TEXT("Argument: \n"),
 	FConsoleCommandWithArgsDelegate::CreateLambda([](const TArray<FString>& Args)
 	{
-		using namespace UE::Editor::DataStorage;
-		IEditorDataStorageProvider* DataStorage = GetMutableDataStorageFeature<IEditorDataStorageProvider>(StorageFeatureName);
-		static TableHandle Table = DataStorage->RegisterTable<FTestColumnA>(FName(TEXT("Debug.CreateRow Table")));
+		IEditorDataStorageProvider* DataStorage = UTypedElementRegistry::GetInstance()->GetMutableDataStorage();
+		static UE::Editor::DataStorage::TableHandle Table = DataStorage->RegisterTable<FTestColumnA>(FName(TEXT("Debug.CreateRow Table")));
 
-		const RowHandle RowHandle = DataStorage->AddRow(Table);
+		const UE::Editor::DataStorage::RowHandle RowHandle = DataStorage->AddRow(Table);
 
 		UE_LOG(LogEditorDataStorage, Warning, TEXT("Added Row %llu"), static_cast<uint64>(RowHandle));
 	}));
@@ -230,8 +230,7 @@ static FAutoConsoleCommand CVarAddDynamicColumnTag(
 	TEXT("Argument: Row, Identifier\n"),
 	FConsoleCommandWithArgsDelegate::CreateLambda([](const TArray<FString>& Args)
 	{
-		using namespace UE::Editor::DataStorage;
-		IEditorDataStorageProvider* DataStorage = GetMutableDataStorageFeature<IEditorDataStorageProvider>(StorageFeatureName);
+		IEditorDataStorageProvider* DataStorage = UTypedElementRegistry::GetInstance()->GetMutableDataStorage();
 		
 		if (Args.Num() != 2)
 		{
@@ -240,7 +239,7 @@ static FAutoConsoleCommand CVarAddDynamicColumnTag(
 		}
 
 		const uint64 RowAsU64 = FCString::Strtoui64(*Args[0], nullptr, 10);
-		const RowHandle Row = RowAsU64;
+		const UE::Editor::DataStorage::RowHandle Row = RowAsU64;
 		const FName Identifier = FName(*Args[1]);
 
 		DataStorage->AddColumn<FTestDynamicTag>(Row, Identifier);
@@ -253,8 +252,7 @@ static FAutoConsoleCommand CVarAddDynamicColumn(
 	TEXT("Argument: Row, Identifier\n"),
 	FConsoleCommandWithArgsDelegate::CreateLambda([](const TArray<FString>& Args)
 	{
-		using namespace UE::Editor::DataStorage;
-		IEditorDataStorageProvider* DataStorage = GetMutableDataStorageFeature<IEditorDataStorageProvider>(StorageFeatureName);
+		IEditorDataStorageProvider* DataStorage = UTypedElementRegistry::GetInstance()->GetMutableDataStorage();
 		
 		if (Args.Num() != 2)
 		{
@@ -263,7 +261,7 @@ static FAutoConsoleCommand CVarAddDynamicColumn(
 		}
 
 		const uint64 RowAsU64 = FCString::Strtoui64(*Args[0], nullptr, 10);
-		const RowHandle Row = RowAsU64;
+		const UE::Editor::DataStorage::RowHandle Row = RowAsU64;
 		const FName Identifier = FName(*Args[1]);
 		
 		bool bUseDefaultApi = false;
@@ -285,8 +283,7 @@ static FAutoConsoleCommand CVarRemoveDynamicColumn(
 	TEXT("Argument: Row, Identifier\n"),
 	FConsoleCommandWithArgsDelegate::CreateLambda([](const TArray<FString>& Args)
 	{
-		using namespace UE::Editor::DataStorage;
-		IEditorDataStorageProvider* DataStorage = GetMutableDataStorageFeature<IEditorDataStorageProvider>(StorageFeatureName);
+		IEditorDataStorageProvider* DataStorage = UTypedElementRegistry::GetInstance()->GetMutableDataStorage();
 		
 		if (Args.Num() != 2)
 		{
@@ -295,7 +292,7 @@ static FAutoConsoleCommand CVarRemoveDynamicColumn(
 		}
 
 		const uint64 RowAsU64 = FCString::Strtoui64(*Args[0], nullptr, 10);
-		const RowHandle Row = RowAsU64;
+		const UE::Editor::DataStorage::RowHandle Row = RowAsU64;
 		const FName Identifier = FName(*Args[1]);
 		
 		DataStorage->RemoveColumn<FTestDynamicColumn>(Row, Identifier);
@@ -308,8 +305,7 @@ static FAutoConsoleCommand CVarAddToDynamicColumn(
 	TEXT("Argument: Row, TagId, Value, [optional] MethodId\n"),
 	FConsoleCommandWithArgsDelegate::CreateLambda([](const TArray<FString>& Args)
 	{
-		using namespace UE::Editor::DataStorage;
-		IEditorDataStorageProvider* DataStorage = GetMutableDataStorageFeature<IEditorDataStorageProvider>(StorageFeatureName);
+		IEditorDataStorageProvider* DataStorage = UTypedElementRegistry::GetInstance()->GetMutableDataStorage();
 		
 		if (Args.Num() < 3 || Args.Num() > 4)
 		{
@@ -318,7 +314,7 @@ static FAutoConsoleCommand CVarAddToDynamicColumn(
 		}
 
 		const uint64 RowAsU64 = FCString::Strtoui64(*Args[0], nullptr, 10);
-		const RowHandle Row = RowAsU64;
+		const UE::Editor::DataStorage::RowHandle Row = RowAsU64;
 		const FName TagId = FName(*Args[1]);
 		const uint64 Value = FCString::Strtoui64(*Args[2], nullptr, 10);
 		const uint64 MethodId = Args.Num() >= 4 ? FCString::Strtoui64(*Args[3], nullptr, 10) : 0;
@@ -361,8 +357,7 @@ static FAutoConsoleCommand CVarPrintDynamicColumn(
 	TEXT("Argument: Row, TagId\n"),
 	FConsoleCommandWithArgsDelegate::CreateLambda([](const TArray<FString>& Args)
 	{
-		using namespace UE::Editor::DataStorage;
-		IEditorDataStorageProvider* DataStorage = GetMutableDataStorageFeature<IEditorDataStorageProvider>(StorageFeatureName);
+		IEditorDataStorageProvider* DataStorage = UTypedElementRegistry::GetInstance()->GetMutableDataStorage();
 		
 		if (Args.Num() != 2)
 		{
@@ -370,7 +365,7 @@ static FAutoConsoleCommand CVarPrintDynamicColumn(
 			return;
 		}
 		const uint64 RowAsU64 = FCString::Strtoui64(*Args[0], nullptr, 10);
-		const RowHandle Row = RowAsU64;
+		const UE::Editor::DataStorage::RowHandle Row = RowAsU64;
 		const FName TagId = FName(*Args[1]);
 
 		FTestDynamicColumn* Column = DataStorage->GetColumn<FTestDynamicColumn>(Row, TagId);
@@ -418,7 +413,7 @@ static FAutoConsoleCommand CVarPrintDynamicColumnWithQuery(
 			return;
 		}
 		
-		IEditorDataStorageProvider* DataStorage = GetMutableDataStorageFeature<IEditorDataStorageProvider>(StorageFeatureName);
+		IEditorDataStorageProvider* DataStorage = UTypedElementRegistry::GetInstance()->GetMutableDataStorage();
 		
 		const FName Identifier(*Args[0]);
 
@@ -429,8 +424,8 @@ static FAutoConsoleCommand CVarPrintDynamicColumnWithQuery(
 			Compile());
 
 		TStringBuilder<1024> StringBuilder;
-		FQueryResult Result = DataStorage->RunQuery(Query,CreateDirectQueryCallbackBinding(
-			[Identifier, &StringBuilder](IDirectQueryContext& Context, const RowHandle* Rows)
+		UE::Editor::DataStorage::FQueryResult Result = DataStorage->RunQuery(Query,CreateDirectQueryCallbackBinding(
+			[Identifier, &StringBuilder](UE::Editor::DataStorage::IDirectQueryContext& Context, const RowHandle* Rows)
 			{
 				const TArrayView<const RowHandle> RowView = MakeConstArrayView(Rows, Context.GetRowCount());
 				// Get pointer to the start of the range of columns to process
@@ -478,7 +473,7 @@ static FAutoConsoleCommand CVarCountDynamicTagWithQuery(
 			return;
 		}
 				
-		IEditorDataStorageProvider* DataStorage = GetMutableDataStorageFeature<IEditorDataStorageProvider>(StorageFeatureName);
+		IEditorDataStorageProvider* DataStorage = UTypedElementRegistry::GetInstance()->GetMutableDataStorage();
 				
 		const FName Identifier(*Args[0]);
 
@@ -519,7 +514,7 @@ static FAutoConsoleCommand CVarRegisterListDynamicColumnQuery(
 			return;
 		}
 				
-		IEditorDataStorageProvider* DataStorage = GetMutableDataStorageFeature<IEditorDataStorageProvider>(StorageFeatureName);
+		IEditorDataStorageProvider* DataStorage = UTypedElementRegistry::GetInstance()->GetMutableDataStorage();
 				
 		const FName Identifier(*Args[0]);
 		const FName ActivationGroup(*Args[1]);
@@ -557,9 +552,8 @@ static FAutoConsoleCommand CVarActivateListDynamicColumnQuery(
 			UE_LOG(LogEditorDataStorage, Error, TEXT("Invalid number of arguments. ActivationGroup"));
 			return;
 		}
-
-		using namespace UE::Editor::DataStorage;
-		IEditorDataStorageProvider* DataStorage = GetMutableDataStorageFeature<IEditorDataStorageProvider>(StorageFeatureName);
+						
+		IEditorDataStorageProvider* DataStorage = UTypedElementRegistry::GetInstance()->GetMutableDataStorage();
 		
 		const FName ActivationGroup(*Args[0]);
 
@@ -573,7 +567,7 @@ static FAutoConsoleCommand CVarAddValueTag(
 	FConsoleCommandWithArgsDelegate::CreateLambda([](const TArray<FString>& Args)
 	{
 		using namespace UE::Editor::DataStorage;
-		IEditorDataStorageProvider* DataStorage = GetMutableDataStorageFeature<IEditorDataStorageProvider>(StorageFeatureName);
+		IEditorDataStorageProvider* DataStorage = UTypedElementRegistry::GetInstance()->GetMutableDataStorage();
 
 		if (Args.Num() != 3)
 		{
@@ -607,8 +601,7 @@ static FAutoConsoleCommand CVarRemoveValueTag(
 	TEXT("Argument: Row, Group\n"),
 	FConsoleCommandWithArgsDelegate::CreateLambda([](const TArray<FString>& Args)
 	{
-		using namespace UE::Editor::DataStorage;
-		IEditorDataStorageProvider* DataStorage = GetMutableDataStorageFeature<IEditorDataStorageProvider>(StorageFeatureName);
+		IEditorDataStorageProvider* DataStorage = UTypedElementRegistry::GetInstance()->GetMutableDataStorage();
 
 		if (Args.Num() != 2)
 		{
@@ -640,7 +633,7 @@ static FAutoConsoleCommand CVarMatchValueTag(
 	{
 		using namespace UE::Editor::DataStorage::Queries;
 		
-		IEditorDataStorageProvider* DataStorage = GetMutableDataStorageFeature<IEditorDataStorageProvider>(StorageFeatureName);
+		IEditorDataStorageProvider* DataStorage = UTypedElementRegistry::GetInstance()->GetMutableDataStorage();
 
 		if (Args.Num() < 1 || Args.Num() > 2)
 		{
@@ -694,7 +687,7 @@ static FAutoConsoleCommand CVarAddValueTagFromEnum(
 	{
 		using namespace UE::Editor::DataStorage;
 
-		IEditorDataStorageProvider* DataStorage = GetMutableDataStorageFeature<IEditorDataStorageProvider>(StorageFeatureName);
+		IEditorDataStorageProvider* DataStorage = UTypedElementRegistry::GetInstance()->GetMutableDataStorage();
 
 		if (Args.Num() < 1 || Args.Num() > 2)
 		{
@@ -732,7 +725,7 @@ static FAutoConsoleCommand CVarRemoveValueTagFromEnum(
 	FConsoleCommandWithArgsDelegate::CreateLambda([](const TArray<FString>& Args)
 	{
 		using namespace UE::Editor::DataStorage;
-		IEditorDataStorageProvider* DataStorage = GetMutableDataStorageFeature<IEditorDataStorageProvider>(StorageFeatureName);
+		IEditorDataStorageProvider* DataStorage = UTypedElementRegistry::GetInstance()->GetMutableDataStorage();
 
 		if (Args.Num() != 1)
 		{
@@ -754,7 +747,7 @@ static FAutoConsoleCommand CVarMatchValueTagFromEnum(
 	{
 		using namespace UE::Editor::DataStorage::Queries;
 		
-		IEditorDataStorageProvider* DataStorage = GetMutableDataStorageFeature<IEditorDataStorageProvider>(StorageFeatureName);
+		IEditorDataStorageProvider* DataStorage = UTypedElementRegistry::GetInstance()->GetMutableDataStorage();
 
 		if (Args.Num() > 1)
 		{
@@ -823,4 +816,4 @@ static FAutoConsoleCommand CVarMatchValueTagFromEnum(
 		UE_LOG(LogEditorDataStorage, Warning, TEXT("Processed %llu rows"), static_cast<uint64>(Count));
 		
 	}),
-	ECVF_Default);
+	ECVF_Default);
\ No newline at end of file
diff --git a/Engine/Plugins/Experimental/EditorDataStorage/Source/TedsCore/Private/Queries/TypedElementAlertQueries.cpp b/Engine/Plugins/Experimental/EditorDataStorage/Source/TedsCore/Private/Queries/TypedElementAlertQueries.cpp
--- a/Engine/Plugins/Experimental/EditorDataStorage/Source/TedsCore/Private/Queries/TypedElementAlertQueries.cpp
+++ b/Engine/Plugins/Experimental/EditorDataStorage/Source/TedsCore/Private/Queries/TypedElementAlertQueries.cpp
@@ -6,9 +6,8 @@
 #include "Elements/Columns/TypedElementHiearchyColumns.h"
 #include "Elements/Columns/TypedElementMiscColumns.h"
 #include "Elements/Columns/TypedElementSelectionColumns.h"
-#include "Elements/Common/EditorDataStorageFeatures.h"
 #include "Elements/Framework/TypedElementQueryBuilder.h"
-#include "Elements/Interfaces/TypedElementDataStorageInterface.h"
+#include "Elements/Framework/TypedElementRegistry.h"
 #include "HAL/IConsoleManager.h"
 #include "Math/UnrealMathUtility.h"
 
@@ -22,7 +21,7 @@ FAutoConsoleCommand AddRandomAlertToRowConsoleCommand(
 
 			TRACE_CPUPROFILER_EVENT_SCOPE(TEDS.Debug.AddRandomAlertToSelectedRows);
 
-			if (IEditorDataStorageProvider* DataStorage = GetMutableDataStorageFeature<IEditorDataStorageProvider>(StorageFeatureName))
+			if (IEditorDataStorageProvider* DataStorage = UTypedElementRegistry::GetInstance()->GetMutableDataStorage())
 			{
 				static QueryHandle Query = [DataStorage]
 				{
@@ -83,7 +82,7 @@ FAutoConsoleCommand ClearAllAlertsConsoleCommand(
 			static TTypedElementColumnTypeList<FTypedElementSyncBackToWorldTag> BatchAddColumns;
 			static TTypedElementColumnTypeList<FTypedElementAlertColumn, FTypedElementChildAlertColumn> BatchRemoveColumns;
 
-			if (IEditorDataStorageProvider* DataStorage = GetMutableDataStorageFeature<IEditorDataStorageProvider>(StorageFeatureName))
+			if (IEditorDataStorageProvider* DataStorage = UTypedElementRegistry::GetInstance()->GetMutableDataStorage())
 			{
 				static QueryHandle AlertInfoQuery = [DataStorage]
 				{
@@ -119,7 +118,7 @@ FAutoConsoleCommand ClearSelectedAlertsConsoleCommand(
 
 			TRACE_CPUPROFILER_EVENT_SCOPE(TEDS.Debug.ClearSelectedAlerts);
 
-			if (IEditorDataStorageProvider* DataStorage = GetMutableDataStorageFeature<IEditorDataStorageProvider>(StorageFeatureName))
+			if (IEditorDataStorageProvider* DataStorage = UTypedElementRegistry::GetInstance()->GetMutableDataStorage())
 			{
 				static QueryHandle AlertQuery = [DataStorage]
 				{
diff --git a/Engine/Plugins/Experimental/EditorDataStorage/Source/TedsCore/Private/TypedElementDatabase.cpp b/Engine/Plugins/Experimental/EditorDataStorage/Source/TedsCore/Private/TypedElementDatabase.cpp
--- a/Engine/Plugins/Experimental/EditorDataStorage/Source/TedsCore/Private/TypedElementDatabase.cpp
+++ b/Engine/Plugins/Experimental/EditorDataStorage/Source/TedsCore/Private/TypedElementDatabase.cpp
@@ -4,9 +4,8 @@
 
 #include "Editor.h"
 #include "EditorDataStorageSettings.h"
-#include "Elements/Common/EditorDataStorageFeatures.h"
 #include "Elements/Interfaces/TypedElementDataStorageFactory.h"
-#include "Elements/Interfaces/TypedElementDataStorageInterface.h"
+#include "Elements/Framework/TypedElementRegistry.h"
 #include "Engine/World.h"
 #include "GlobalLock.h"
 #include "MassEntityEditorSubsystem.h"
@@ -33,10 +32,12 @@ FAutoConsoleCommandWithOutputDevice PrintQueryCallbacksConsoleCommand(
 	TEXT("Prints out a list of all processors."),
 	FConsoleCommandWithOutputDeviceDelegate::CreateLambda([](FOutputDevice& Output)
 		{
-			using namespace UE::Editor::DataStorage;
-			if (IEditorDataStorageProvider* DataStorage = GetMutableDataStorageFeature<IEditorDataStorageProvider>(StorageFeatureName))
+			if (UTypedElementRegistry* Registry = UTypedElementRegistry::GetInstance())
 			{
-				DataStorage->DebugPrintQueryCallbacks(Output);
+				if (IEditorDataStorageProvider* DataStorage = Registry->GetMutableDataStorage())
+				{
+					DataStorage->DebugPrintQueryCallbacks(Output);
+				}
 			}
 		}));
 
diff --git a/Engine/Plugins/Experimental/EditorDataStorage/Source/TedsCore/Tests/DynamicColumnTests.cpp b/Engine/Plugins/Experimental/EditorDataStorage/Source/TedsCore/Tests/DynamicColumnTests.cpp
--- a/Engine/Plugins/Experimental/EditorDataStorage/Source/TedsCore/Tests/DynamicColumnTests.cpp
+++ b/Engine/Plugins/Experimental/EditorDataStorage/Source/TedsCore/Tests/DynamicColumnTests.cpp
@@ -1,15 +1,12 @@
 // Copyright Epic Games, Inc. All Rights Reserved.
 
+#include "Elements/Framework/TypedElementRegistry.h"
 #include "Elements/Framework/TypedElementTestColumns.h"
 #include "Elements/Interfaces/TypedElementDataStorageInterface.h"
 #include "Debug/TypedElementDatabaseDebugTypes.h"
 #include "Misc/CoreDelegates.h"
 #if WITH_TESTS
 
-#include "Elements/Common/EditorDataStorageFeatures.h"
-#include "Elements/Framework/TypedElementQueryBuilder.h"
-#include "Elements/Interfaces/TypedElementDataStorageInterface.h"
-
 #include "Misc/AutomationTest.h"
 
 namespace UE::Editor::DataStorage::Tests
@@ -101,7 +98,8 @@ namespace UE::Editor::DataStorage::Tests
 	{
 		BeforeEach([this]()
 		{
-			TedsInterface = GetMutableDataStorageFeature<IEditorDataStorageProvider>(StorageFeatureName);
+			UTypedElementRegistry* TypedElementRegistry = UTypedElementRegistry::GetInstance();
+			TedsInterface = TypedElementRegistry->GetMutableDataStorage();
 			TestTable = RegisterTestTable();
 			Identifiers = {TEXT("StaticMesh"), TEXT("Animation"), TEXT("AudioClip")};
 			TestTrue("", TedsInterface != nullptr);
diff --git a/Engine/Plugins/Experimental/EditorDataStorage/Source/TedsCore/Tests/TypedElementAttributeBindingTests.cpp b/Engine/Plugins/Experimental/EditorDataStorage/Source/TedsCore/Tests/TypedElementAttributeBindingTests.cpp
--- a/Engine/Plugins/Experimental/EditorDataStorage/Source/TedsCore/Tests/TypedElementAttributeBindingTests.cpp
+++ b/Engine/Plugins/Experimental/EditorDataStorage/Source/TedsCore/Tests/TypedElementAttributeBindingTests.cpp
@@ -3,9 +3,8 @@
 
 #include "Elements/Framework/TypedElementAttributeBinding.h"
 #if WITH_TESTS
-#include "Elements/Common/EditorDataStorageFeatures.h"
 #include "Elements/Framework/TypedElementTestColumns.h"
-#include "Elements/Interfaces/TypedElementDataStorageInterface.h"
+#include "Elements/Framework/TypedElementRegistry.h"
 #include "Misc/AutomationTest.h"
 
 namespace UE::Editor::DataStorage::Tests
@@ -51,7 +50,8 @@ namespace UE::Editor::DataStorage::Tests
 	{
 		BeforeEach([this]()
 		{
-			TedsInterface = GetMutableDataStorageFeature<IEditorDataStorageProvider>(StorageFeatureName);
+			UTypedElementRegistry* TypedElementRegistry = UTypedElementRegistry::GetInstance();
+			TedsInterface = TypedElementRegistry->GetMutableDataStorage();
 			TestTrue("", TedsInterface != nullptr);
 			
 			TestTableHandle = RegisterTestTable();
diff --git a/Engine/Plugins/Experimental/EditorDataStorage/Source/TedsCore/Tests/TypedElementDataStorageInterfaceTests.cpp b/Engine/Plugins/Experimental/EditorDataStorage/Source/TedsCore/Tests/TypedElementDataStorageInterfaceTests.cpp
--- a/Engine/Plugins/Experimental/EditorDataStorage/Source/TedsCore/Tests/TypedElementDataStorageInterfaceTests.cpp
+++ b/Engine/Plugins/Experimental/EditorDataStorage/Source/TedsCore/Tests/TypedElementDataStorageInterfaceTests.cpp
@@ -2,10 +2,8 @@
 
 
 #if WITH_TESTS
-#include "Elements/Common/EditorDataStorageFeatures.h"
 #include "Elements/Framework/TypedElementTestColumns.h"
-#include "Elements/Framework/TypedElementQueryBuilder.h"
-#include "Elements/Interfaces/TypedElementDataStorageInterface.h"
+#include "Elements/Framework/TypedElementRegistry.h"
 #include "Misc/AutomationTest.h"
 
 namespace UE::Editor::DataStorage::Tests
@@ -128,7 +126,8 @@ namespace UE::Editor::DataStorage::Tests
 	{
 		BeforeEach([this]()
 		{
-			TedsInterface = GetMutableDataStorageFeature<IEditorDataStorageProvider>(StorageFeatureName);
+			UTypedElementRegistry* TypedElementRegistry = UTypedElementRegistry::GetInstance();
+			TedsInterface = TypedElementRegistry->GetMutableDataStorage();
 			TestTrue("", TedsInterface != nullptr);
 		});
 
diff --git a/Engine/Plugins/Experimental/EditorDataStorage/Source/TedsUI/Private/TedsUIModule.cpp b/Engine/Plugins/Experimental/EditorDataStorage/Source/TedsUI/Private/TedsUIModule.cpp
--- a/Engine/Plugins/Experimental/EditorDataStorage/Source/TedsUI/Private/TedsUIModule.cpp
+++ b/Engine/Plugins/Experimental/EditorDataStorage/Source/TedsUI/Private/TedsUIModule.cpp
@@ -2,7 +2,7 @@
 
 #include "TedsUIModule.h"
 
-#include "Elements/Common/EditorDataStorageFeatures.h"
+#include "Elements/Framework/TypedElementRegistry.h"
 #include "Elements/Interfaces/TypedElementDataStorageUiInterface.h"
 #include "HAL/IConsoleManager.h"
 #include "Misc/CoreDelegates.h"
@@ -17,7 +17,8 @@ namespace UE::Editor::DataStorage::Private
 	TEXT("Prints a list of all the known widget purposes."),
 	FConsoleCommandWithOutputDeviceDelegate::CreateLambda([](FOutputDevice& Output)
 		{
-			if (IEditorDataStorageUiProvider* UiStorage = GetMutableDataStorageFeature<IEditorDataStorageUiProvider>(UiFeatureName))
+			UTypedElementRegistry* Registry = UTypedElementRegistry::GetInstance();
+			if (IEditorDataStorageUiProvider* UiStorage = Registry->GetMutableDataStorageUi())
 			{
 				Output.Log(TEXT("The Typed Elements Data Storage has recorded the following widget purposes:"));
 				UiStorage->ListWidgetPurposes(
diff --git a/Engine/Plugins/Experimental/EditorDataStorage/Source/TedsUI/Private/Widgets/CounterWidget.cpp b/Engine/Plugins/Experimental/EditorDataStorage/Source/TedsUI/Private/Widgets/CounterWidget.cpp
--- a/Engine/Plugins/Experimental/EditorDataStorage/Source/TedsUI/Private/Widgets/CounterWidget.cpp
+++ b/Engine/Plugins/Experimental/EditorDataStorage/Source/TedsUI/Private/Widgets/CounterWidget.cpp
@@ -4,10 +4,9 @@
 
 #include "Elements/Columns/TypedElementSlateWidgetColumns.h"
 #include "Elements/Columns/TypedElementValueCacheColumns.h"
-#include "Elements/Common/EditorDataStorageFeatures.h"
 #include "Elements/Framework/TypedElementQueryBuilder.h"
+#include "Elements/Framework/TypedElementRegistry.h"
 #include "Elements/Interfaces/TypedElementDataStorageInterface.h"
-#include "Elements/Interfaces/TypedElementDataStorageUiInterface.h"
 #include "Interfaces/IMainFrameModule.h"
 #include "Layout/Margin.h"
 #include "Misc/CoreDelegates.h"
@@ -112,9 +111,11 @@ void UCounterWidgetFactory::SetupMainWindowIntegrations(TSharedPtr<SWindow> Pare
 {
 	if (!bHasBeenSetup)
 	{
-		using namespace UE::Editor::DataStorage;
+		UTypedElementRegistry* Registry = UTypedElementRegistry::GetInstance();
+		checkf(Registry, TEXT(
+			"FEditorDataStorageUiModule didn't find the UTypedElementRegistry during main window integration when it should be available."));
 
-		IEditorDataStorageUiProvider* UiInterface = GetMutableDataStorageFeature<IEditorDataStorageUiProvider>(UiFeatureName);
+		IEditorDataStorageUiProvider* UiInterface = Registry->GetMutableDataStorageUi();
 		checkf(UiInterface, TEXT(
 			"FEditorDataStorageUiModule tried to integrate with the main window before the "
 			"TEDS UI interface is available."));
diff --git a/Engine/Plugins/Experimental/EditorDataStorageFeatures/Source/TedsAssetData/Private/Debug/TedsAssetDataDebugCommands.cpp b/Engine/Plugins/Experimental/EditorDataStorageFeatures/Source/TedsAssetData/Private/Debug/TedsAssetDataDebugCommands.cpp
--- a/Engine/Plugins/Experimental/EditorDataStorageFeatures/Source/TedsAssetData/Private/Debug/TedsAssetDataDebugCommands.cpp
+++ b/Engine/Plugins/Experimental/EditorDataStorageFeatures/Source/TedsAssetData/Private/Debug/TedsAssetDataDebugCommands.cpp
@@ -4,8 +4,8 @@
 #include "AssetRegistry/IAssetRegistry.h"
 #include "Containers/UnrealString.h"
 #include "Elements/Common/TypedElementHandles.h"
-#include "Elements/Common/EditorDataStorageFeatures.h"
 #include "Elements/Framework/TypedElementIndexHasher.h"
+#include "Elements/Framework/TypedElementRegistry.h"
 #include "HAL/IConsoleManager.h"
 #include "TedsAssetDataColumns.h"
 #include "UObject/NameTypes.h"
@@ -30,7 +30,7 @@ namespace UE::Editor::DataStorage::Debug::Private
 		TEXT("Print some debug information on the specified path."),
 		FConsoleCommandWithArgsDelegate::CreateLambda([](const TArray<FString>& InArgs)
 			{
-				const IEditorDataStorageProvider* Database = GetMutableDataStorageFeature<IEditorDataStorageProvider>(StorageFeatureName);
+				const IEditorDataStorageProvider* Database = UTypedElementRegistry::GetInstance()->GetDataStorage();
 				IAssetRegistry& AssetRegistry = IAssetRegistry::GetChecked();
 
 				for (const FString& Path : InArgs)
@@ -220,7 +220,7 @@ namespace UE::Editor::DataStorage::Debug::Private
 				TRACE_CPUPROFILER_EVENT_SCOPE(TEDS.Debug.CreateDebugAssetRows);
 
 
-				IEditorDataStorageProvider* DataStorage = GetMutableDataStorageFeature<IEditorDataStorageProvider>(StorageFeatureName);
+				IEditorDataStorageProvider* DataStorage = UTypedElementRegistry::GetInstance()->GetMutableDataStorage();
 
 				if (!DataStorage)
 				{
@@ -252,7 +252,7 @@ namespace UE::Editor::DataStorage::Debug::Private
 			{
 				using namespace UE::Editor::DataStorage::Queries;
 
-				IEditorDataStorageProvider* DataStorage = GetMutableDataStorageFeature<IEditorDataStorageProvider>(StorageFeatureName);
+				IEditorDataStorageProvider* DataStorage = UTypedElementRegistry::GetInstance()->GetMutableDataStorage();
 
 				if (!DataStorage)
 				{
diff --git a/Engine/Plugins/Experimental/EditorDataStorageFeatures/Source/TedsAssetData/Private/TedsAssetDataModule.cpp b/Engine/Plugins/Experimental/EditorDataStorageFeatures/Source/TedsAssetData/Private/TedsAssetDataModule.cpp
--- a/Engine/Plugins/Experimental/EditorDataStorageFeatures/Source/TedsAssetData/Private/TedsAssetDataModule.cpp
+++ b/Engine/Plugins/Experimental/EditorDataStorageFeatures/Source/TedsAssetData/Private/TedsAssetDataModule.cpp
@@ -3,8 +3,7 @@
 #include "TedsAssetDataModule.h"
 
 #include "CB/TedsAssetDataCBDataSource.h"
-#include "Elements/Common/EditorDataStorageFeatures.h"
-#include "Elements/Interfaces/TypedElementDataStorageInterface.h"
+#include "Elements/Framework/TypedElementRegistry.h"
 #include "HAL/IConsoleManager.h"
 #include "Modules/ModuleManager.h"
 #include "TedsAssetData.h"
@@ -43,7 +42,10 @@ void FTedsAssetDataModule::StartupModule()
 
 void FTedsAssetDataModule::ShutdownModule()
 {
-	UE::Editor::DataStorage::OnEditorDataStorageFeaturesEnabled().RemoveAll(this);
+	if (UTypedElementRegistry* TypedElementRegistry = UTypedElementRegistry::GetInstance())
+	{
+		TypedElementRegistry->OnDataStorageInterfacesSet().RemoveAll(this);
+	}
 }
 
 FTedsAssetDataModule* FTedsAssetDataModule::Get()
@@ -58,17 +60,17 @@ FTedsAssetDataModule& FTedsAssetDataModule::GetChecked()
 
 void FTedsAssetDataModule::EnableTedsAssetRegistryStorage()
 {
-	using namespace UE::Editor::DataStorage;
 	if (!AssetRegistryStorage)
 	{
 		FModuleManager::Get().LoadModuleChecked(TEXT("TypedElementFramework"));
-		if (AreEditorDataStorageFeaturesEnabled())
+		UTypedElementRegistry* TypedElementRegistry = UTypedElementRegistry::GetInstance();
+		if (TypedElementRegistry->GetMutableDataStorage())
 		{
 			InitAssetRegistryStorage();
 		}
 		else
 		{
-			OnEditorDataStorageFeaturesEnabled().AddRaw(this, &FTedsAssetDataModule::InitAssetRegistryStorage);
+			TypedElementRegistry->OnDataStorageInterfacesSet().AddRaw(this, &FTedsAssetDataModule::InitAssetRegistryStorage);
 		}
 
 		if (!Private::CVarTEDSAssetDataStorage.GetValueOnGameThread())
@@ -106,8 +108,7 @@ void FTedsAssetDataModule::ProcessDependentEvents()
 
 void FTedsAssetDataModule::InitAssetRegistryStorage()
 {
-	using namespace UE::Editor::DataStorage;
-	IEditorDataStorageProvider& MutableDataStorage = *GetMutableDataStorageFeature<IEditorDataStorageProvider>(StorageFeatureName);
+	IEditorDataStorageProvider& MutableDataStorage = *UTypedElementRegistry::GetInstance()->GetMutableDataStorage();
 
 	AssetDataCBDataSource = MakeUnique<Private::FTedsAssetDataCBDataSource>(MutableDataStorage);
 	AssetRegistryStorage = MakeUnique<Private::FTedsAssetData>(MutableDataStorage);
diff --git a/Engine/Plugins/Experimental/EditorDataStorageFeatures/Source/TedsAssetData/Private/Tests/TedsAssetDataTest.cpp b/Engine/Plugins/Experimental/EditorDataStorageFeatures/Source/TedsAssetData/Private/Tests/TedsAssetDataTest.cpp
--- a/Engine/Plugins/Experimental/EditorDataStorageFeatures/Source/TedsAssetData/Private/Tests/TedsAssetDataTest.cpp
+++ b/Engine/Plugins/Experimental/EditorDataStorageFeatures/Source/TedsAssetData/Private/Tests/TedsAssetDataTest.cpp
@@ -7,7 +7,6 @@
 #include "Containers/ChunkedArray.h"
 #include "Containers/Set.h"
 #include "Containers/UnrealString.h"
-#include "Elements/Common/EditorDataStorageFeatures.h"
 #include "Elements/Common/TypedElementHandles.h"
 #include "Elements/Framework/TypedElementIndexHasher.h"
 #include "Elements/Framework/TypedElementRegistry.h"
@@ -66,8 +65,7 @@ bool FTedsAssetDataTest::RunTest(const FString& Parameters)
 
 
 	// Do a sanity check that the data from the asset registry exist in Teds
-	using namespace UE::Editor::DataStorage;
-	const IEditorDataStorageProvider* Database = GetMutableDataStorageFeature<IEditorDataStorageProvider>(StorageFeatureName);
+	const IEditorDataStorageProvider* Database = UTypedElementRegistry::GetInstance()->GetDataStorage();
 	
 	for (const TPair<FName, int32>& PathAndAssetCount : AssetRegistryPathsAndAssetCount)
 	{
diff --git a/Engine/Plugins/Experimental/EditorDataStorageFeatures/Source/TedsContentBrowser/Private/TedsContentBrowserModule.cpp b/Engine/Plugins/Experimental/EditorDataStorageFeatures/Source/TedsContentBrowser/Private/TedsContentBrowserModule.cpp
--- a/Engine/Plugins/Experimental/EditorDataStorageFeatures/Source/TedsContentBrowser/Private/TedsContentBrowserModule.cpp
+++ b/Engine/Plugins/Experimental/EditorDataStorageFeatures/Source/TedsContentBrowser/Private/TedsContentBrowserModule.cpp
@@ -6,9 +6,8 @@
 #include "ContentBrowserModule.h"
 #include "Columns/SlateDelegateColumns.h"
 #include "Elements/Common/TypedElementQueryTypes.h"
-#include "Elements/Common/EditorDataStorageFeatures.h"
 #include "Elements/Framework/TypedElementIndexHasher.h"
-#include "Elements/Interfaces/TypedElementDataStorageInterface.h"
+#include "Elements/Framework/TypedElementRegistry.h"
 #include "Experimental/ContentBrowserViewExtender.h"
 #include "Modules/ModuleManager.h"
 #include "QueryStack/FQueryStackNode_RowView.h"
@@ -127,7 +126,7 @@ namespace UE::Editor::ContentBrowser
 
 	FTedsContentBrowserViewExtender::FTedsContentBrowserViewExtender()
 	{
-		DataStorage = GetMutableDataStorageFeature<IEditorDataStorageProvider>(StorageFeatureName);
+		DataStorage = UTypedElementRegistry::GetInstance()->GetMutableDataStorage();
 		
 		RowQueryStack = MakeShared<FQueryStackNode_RowView>(&Rows);
 
diff --git a/Engine/Plugins/Experimental/EditorDataStorageFeatures/Source/TedsDebugger/Private/STedsDebugger.cpp b/Engine/Plugins/Experimental/EditorDataStorageFeatures/Source/TedsDebugger/Private/STedsDebugger.cpp
--- a/Engine/Plugins/Experimental/EditorDataStorageFeatures/Source/TedsDebugger/Private/STedsDebugger.cpp
+++ b/Engine/Plugins/Experimental/EditorDataStorageFeatures/Source/TedsDebugger/Private/STedsDebugger.cpp
@@ -7,8 +7,7 @@
 #include "Elements/Columns/TypedElementMiscColumns.h"
 #include "Elements/Columns/TypedElementSelectionColumns.h"
 #include "Elements/Columns/TypedElementTypeInfoColumns.h"
-#include "Elements/Common/EditorDataStorageFeatures.h"
-#include "Elements/Interfaces/TypedElementDataStorageInterface.h"
+#include "Elements/Framework/TypedElementRegistry.h"
 #include "Framework/Docking/LayoutService.h"
 #include "Modules/ModuleManager.h"
 #include "QueryEditor/TedsQueryEditor.h"
@@ -27,9 +26,11 @@ namespace UE::Editor::DataStorage::Debug
 
 STedsDebugger::~STedsDebugger()
 {
-	if (AreEditorDataStorageFeaturesEnabled())
+	UTypedElementRegistry* Registry = UTypedElementRegistry::GetInstance();
+
+	if (Registry && Registry->AreDataStorageInterfacesSet())
 	{
-		GetMutableDataStorageFeature<IEditorDataStorageProvider>(StorageFeatureName)->UnregisterQuery(TableViewerQuery);
+		Registry->GetMutableDataStorage()->UnregisterQuery(TableViewerQuery);
 	}
 }
 
@@ -130,9 +131,11 @@ TSharedRef<SDockTab> STedsDebugger::SpawnQueryEditorTab(const FSpawnTabArgs& Arg
 	TSharedRef<SDockTab> DockTab = SNew(SDockTab).TabRole(ETabRole::PanelTab);
 	if (!QueryEditorModel)
 	{
-		if (AreEditorDataStorageFeaturesEnabled())
+		UTypedElementRegistry* Registry = UTypedElementRegistry::GetInstance();
+
+		if (Registry && Registry->AreDataStorageInterfacesSet())
 		{
-			IEditorDataStorageProvider* DataStorageInterface = GetMutableDataStorageFeature<IEditorDataStorageProvider>(StorageFeatureName);
+			IEditorDataStorageProvider* DataStorageInterface = Registry->GetMutableDataStorage();
 			QueryEditorModel = MakeUnique<QueryEditor::FTedsQueryEditorModel>(*DataStorageInterface);
 		}
 	}
diff --git a/Engine/Plugins/Experimental/EditorDataStorageFeatures/Source/TedsDebugger/Private/Widgets/QueryEditor/TedsQueryEditorResultsView.cpp b/Engine/Plugins/Experimental/EditorDataStorageFeatures/Source/TedsDebugger/Private/Widgets/QueryEditor/TedsQueryEditorResultsView.cpp
--- a/Engine/Plugins/Experimental/EditorDataStorageFeatures/Source/TedsDebugger/Private/Widgets/QueryEditor/TedsQueryEditorResultsView.cpp
+++ b/Engine/Plugins/Experimental/EditorDataStorageFeatures/Source/TedsDebugger/Private/Widgets/QueryEditor/TedsQueryEditorResultsView.cpp
@@ -5,10 +5,9 @@
 #include "SWarningOrErrorBox.h"
 #include "TedsOutlinerModule.h"
 #include "TedsTableViewerColumn.h"
-#include "Elements/Common/EditorDataStorageFeatures.h"
-#include "Elements/Framework/TypedElementQueryBuilder.h"
 #include "Elements/Interfaces/TypedElementDataStorageInterface.h"
-#include "Elements/Interfaces/TypedElementDataStorageUiInterface.h"
+#include "Elements/Framework/TypedElementQueryBuilder.h"
+#include "Elements/Framework/TypedElementRegistry.h"
 #include "Modules/ModuleManager.h"
 #include "QueryEditor/TedsQueryEditorModel.h"
 #include "QueryStack/FQueryStackNode_RowView.h"
@@ -186,7 +185,9 @@ namespace UE::Editor::DataStorage::Debug::QueryEditor
 			return false;
 		};
 	
-		IEditorDataStorageUiProvider* StorageUi = GetMutableDataStorageFeature<IEditorDataStorageUiProvider>(UiFeatureName);
+		UTypedElementRegistry* Registry = UTypedElementRegistry::GetInstance();
+		checkf(Registry, TEXT("SResultsView created before UTypedElementRegistry is available."));
+		IEditorDataStorageUiProvider* StorageUi = Registry->GetMutableDataStorageUi();
 		checkf(StorageUi, TEXT("SResultsView created before data storage interfaces were initialized."))
 
 		StorageUi->CreateWidgetConstructors(TEXT("General.Cell.RowHandle"), FMetaDataView(), AssignWidgetToColumn);
diff --git a/Engine/Plugins/Experimental/EditorDataStorageFeatures/Source/TedsDebugger/Private/Widgets/RowReferenceWidget.cpp b/Engine/Plugins/Experimental/EditorDataStorageFeatures/Source/TedsDebugger/Private/Widgets/RowReferenceWidget.cpp
--- a/Engine/Plugins/Experimental/EditorDataStorageFeatures/Source/TedsDebugger/Private/Widgets/RowReferenceWidget.cpp
+++ b/Engine/Plugins/Experimental/EditorDataStorageFeatures/Source/TedsDebugger/Private/Widgets/RowReferenceWidget.cpp
@@ -6,12 +6,12 @@
 #include "Elements/Columns/TypedElementHiearchyColumns.h"
 #include "Elements/Columns/TypedElementLabelColumns.h"
 #include "Elements/Columns/TypedElementMiscColumns.h"
-#include "Elements/Common/EditorDataStorageFeatures.h"
 #include "Elements/Interfaces/TypedElementDataStorageInterface.h"
 #include "ISceneOutliner.h"
 #include "Modules/ModuleManager.h"
 #include "TedsDebuggerModule.h"
 #include "Elements/Columns/TypedElementSlateWidgetColumns.h"
+#include "Elements/Framework/TypedElementRegistry.h"
 #include "Widgets/Input/SHyperlink.h"
 #include "Widgets/Layout/SBox.h"
 
@@ -66,7 +66,7 @@ namespace UE::Editor::DataStorage::Debug::Private
 		SBox* WidgetInstance = static_cast<SBox*>(Widget.Get());
 		WidgetInstance->SetContent(SNullWidget::NullWidget);
 
-		const IEditorDataStorageProvider* DataStorage = GetMutableDataStorageFeature<IEditorDataStorageProvider>(StorageFeatureName);
+		const IEditorDataStorageProvider* DataStorage = UTypedElementRegistry::GetInstance()->GetDataStorage();
 
 		// We only navigate to row references that have a label column
 		if (const FTypedElementLabelColumn* LabelColumn = DataStorage->GetColumn<FTypedElementLabelColumn>(TargetRow))
diff --git a/Engine/Plugins/Experimental/EditorDataStorageFeatures/Source/TedsOutliner/Private/Compatibility/SceneOutlinerRowHandleColumn.cpp b/Engine/Plugins/Experimental/EditorDataStorageFeatures/Source/TedsOutliner/Private/Compatibility/SceneOutlinerRowHandleColumn.cpp
--- a/Engine/Plugins/Experimental/EditorDataStorageFeatures/Source/TedsOutliner/Private/Compatibility/SceneOutlinerRowHandleColumn.cpp
+++ b/Engine/Plugins/Experimental/EditorDataStorageFeatures/Source/TedsOutliner/Private/Compatibility/SceneOutlinerRowHandleColumn.cpp
@@ -4,8 +4,7 @@
 
 #include "SortHelper.h"
 #include "TedsOutlinerItem.h"
-#include "Elements/Common/EditorDataStorageFeatures.h"
-#include "Elements/Interfaces/TypedElementDataStorageUiInterface.h"
+#include "Elements/Framework/TypedElementRegistry.h"
 #include "TedsTableViewerColumn.h"
 
 #define LOCTEXT_NAMESPACE "SceneOutlinerRowHandleColumn"
@@ -13,15 +12,16 @@
 FSceneOutlinerRowHandleColumn::FSceneOutlinerRowHandleColumn(ISceneOutliner& SceneOutliner)
 	: WeakSceneOutliner(StaticCastSharedRef<ISceneOutliner>(SceneOutliner.AsShared()))
 {
-	using namespace UE::Editor::DataStorage;
 	auto AssignWidgetToColumn = [this](TUniquePtr<FTypedElementWidgetConstructor> Constructor, TConstArrayView<TWeakObjectPtr<const UScriptStruct>>)
 	{
 		TSharedPtr<FTypedElementWidgetConstructor> WidgetConstructor(Constructor.Release());
-		TableViewerColumn = MakeShared<FTedsTableViewerColumn>(TEXT("Row Handle"), WidgetConstructor);
+		TableViewerColumn = MakeShared<UE::Editor::DataStorage::FTedsTableViewerColumn>(TEXT("Row Handle"), WidgetConstructor);
 		return false;
 	};
 	
-	IEditorDataStorageUiProvider* StorageUi = GetMutableDataStorageFeature<IEditorDataStorageUiProvider>(UiFeatureName);
+	UTypedElementRegistry* Registry = UTypedElementRegistry::GetInstance();
+	checkf(Registry, TEXT("FSceneOutlinerRowHandleColumn created before UTypedElementRegistry is available."));
+	IEditorDataStorageUiProvider* StorageUi = Registry->GetMutableDataStorageUi();
 	checkf(StorageUi, TEXT("FSceneOutlinerRowHandleColumn created before data storage interfaces were initialized."))
 
 	StorageUi->CreateWidgetConstructors(TEXT("General.Cell.RowHandle"), UE::Editor::DataStorage::FMetaDataView(), AssignWidgetToColumn);
diff --git a/Engine/Plugins/Experimental/EditorDataStorageFeatures/Source/TedsOutliner/Private/Compatibility/SceneOutlinerTedsBridge.cpp b/Engine/Plugins/Experimental/EditorDataStorageFeatures/Source/TedsOutliner/Private/Compatibility/SceneOutlinerTedsBridge.cpp
--- a/Engine/Plugins/Experimental/EditorDataStorageFeatures/Source/TedsOutliner/Private/Compatibility/SceneOutlinerTedsBridge.cpp
+++ b/Engine/Plugins/Experimental/EditorDataStorageFeatures/Source/TedsOutliner/Private/Compatibility/SceneOutlinerTedsBridge.cpp
@@ -8,11 +8,9 @@
 #include "Elements/Columns/TypedElementPackageColumns.h"
 #include "Elements/Columns/TypedElementTransformColumns.h"
 #include "Elements/Columns/TypedElementTypeInfoColumns.h"
-#include "Elements/Common/EditorDataStorageFeatures.h"
+#include "Elements/Framework/TypedElementRegistry.h"
 #include "Elements/Framework/TypedElementQueryBuilder.h"
-#include "Elements/Interfaces/TypedElementDataStorageInterface.h"
 #include "Elements/Interfaces/TypedElementDataStorageCompatibilityInterface.h"
-#include "Elements/Interfaces/TypedElementDataStorageUiInterface.h"
 #include "ILevelEditor.h"
 #include "ISceneOutliner.h"
 #include "ISceneOutlinerColumn.h"
@@ -40,7 +38,8 @@ FAutoConsoleCommand BindColumnsToSceneOutlinerConsoleCommand(
 			
 		    const FName WidgetPurposes[] = {TEXT("SceneOutliner.Cell"), TEXT("General.Cell")};
 
-			if (IEditorDataStorageProvider* DataStorage = GetMutableDataStorageFeature<IEditorDataStorageProvider>(StorageFeatureName))
+			UTypedElementRegistry* Registry = UTypedElementRegistry::GetInstance();
+			if (IEditorDataStorageProvider* DataStorage = Registry->GetMutableDataStorage())
 			{
 				static UE::Editor::DataStorage::QueryHandle Queries[] =
 				{
@@ -368,10 +367,14 @@ const FName FSceneOutlinerTedsQueryBinder::DefaultItemLabelCellWidgetPurpose(TEX
 
 FSceneOutlinerTedsQueryBinder::FSceneOutlinerTedsQueryBinder()
 {
-	using namespace UE::Editor::DataStorage;
-	Storage = GetMutableDataStorageFeature<IEditorDataStorageProvider>(StorageFeatureName);
-	StorageUi = GetMutableDataStorageFeature<IEditorDataStorageUiProvider>(UiFeatureName);
-	StorageCompatibility = GetMutableDataStorageFeature<IEditorDataStorageCompatibilityProvider>(CompatibilityFeatureName);
+	UTypedElementRegistry* Registry = UTypedElementRegistry::GetInstance();
+	checkf(Registry, TEXT("Unable to bind a Scene Outliner to a query before the Typed Elements are available."));
+	if (Registry)
+	{
+		Storage = Registry->GetMutableDataStorage();
+		StorageUi = Registry->GetMutableDataStorageUi();
+		StorageCompatibility = Registry->GetMutableDataStorageCompatibility();
+	}
 
 	SetupDefaultColumnMapping();
 }
diff --git a/Engine/Plugins/Experimental/EditorDataStorageFeatures/Source/TedsOutliner/Private/TedsOutlinerImpl.cpp b/Engine/Plugins/Experimental/EditorDataStorageFeatures/Source/TedsOutliner/Private/TedsOutlinerImpl.cpp
--- a/Engine/Plugins/Experimental/EditorDataStorageFeatures/Source/TedsOutliner/Private/TedsOutlinerImpl.cpp
+++ b/Engine/Plugins/Experimental/EditorDataStorageFeatures/Source/TedsOutliner/Private/TedsOutlinerImpl.cpp
@@ -7,11 +7,10 @@
 #include "Elements/Columns/TypedElementLabelColumns.h"
 #include "Elements/Columns/TypedElementMiscColumns.h"
 #include "Elements/Columns/TypedElementTypeInfoColumns.h"
-#include "Elements/Common/EditorDataStorageFeatures.h"
 #include "Elements/Interfaces/TypedElementDataStorageInterface.h"
-#include "Elements/Interfaces/TypedElementDataStorageCompatibilityInterface.h"
 #include "Elements/Interfaces/TypedElementDataStorageUiInterface.h"
 #include "Elements/Framework/TypedElementQueryBuilder.h"
+#include "Elements/Framework/TypedElementRegistry.h"
 #include "Columns/SlateDelegateColumns.h"
 #include "Compatibility/SceneOutlinerRowHandleColumn.h"
 #include "TedsOutlinerFilter.h"
@@ -89,10 +88,15 @@ FTedsOutlinerImpl::FTedsOutlinerImpl(const FTedsOutlinerParams& InParams, IScene
 	, SceneOutliner(InParams.SceneOutliner)
 {
 	// Initialize the TEDS constructs
-	using namespace UE::Editor::DataStorage;
-	Storage = GetMutableDataStorageFeature<IEditorDataStorageProvider>(StorageFeatureName);
-	StorageUi = GetMutableDataStorageFeature<IEditorDataStorageUiProvider>(UiFeatureName);
-	StorageCompatibility = GetMutableDataStorageFeature<IEditorDataStorageCompatibilityProvider>(CompatibilityFeatureName);
+	UTypedElementRegistry* Registry = UTypedElementRegistry::GetInstance();
+	checkf(Registry, TEXT("Unable to initialize the Typed Elements Outliner before TEDS is initialized."));
+	if (Registry)
+	{
+		Storage = Registry->GetMutableDataStorage();
+		StorageUi = Registry->GetMutableDataStorageUi();
+		StorageCompatibility = Registry->GetMutableDataStorageCompatibility();
+	}
+
 }
 
 void FTedsOutlinerImpl::CreateFilterQueries()
diff --git a/Engine/Plugins/Experimental/EditorDataStorageFeatures/Source/TedsOutliner/Private/TedsOutlinerModule.cpp b/Engine/Plugins/Experimental/EditorDataStorageFeatures/Source/TedsOutliner/Private/TedsOutlinerModule.cpp
--- a/Engine/Plugins/Experimental/EditorDataStorageFeatures/Source/TedsOutliner/Private/TedsOutlinerModule.cpp
+++ b/Engine/Plugins/Experimental/EditorDataStorageFeatures/Source/TedsOutliner/Private/TedsOutlinerModule.cpp
@@ -13,8 +13,7 @@
 #include "Elements/Columns/TypedElementMiscColumns.h"
 #include "Elements/Columns/TypedElementPackageColumns.h"
 #include "Elements/Columns/TypedElementTypeInfoColumns.h"
-#include "Elements/Common/EditorDataStorageFeatures.h"
-#include "Elements/Interfaces/TypedElementDataStorageInterface.h"
+#include "Elements/Framework/TypedElementRegistry.h"
 #include "Modules/ModuleManager.h"
 #include "Compatibility/SceneOutlinerTedsBridge.h"
 #include "Compatibility/SceneOutlinerRowHandleColumn.h"
@@ -66,8 +65,9 @@ FTedsOutlinerModule::FTedsOutlinerModule()
 
 TSharedRef<ISceneOutliner> FTedsOutlinerModule::CreateTedsOutliner(const FSceneOutlinerInitializationOptions& InInitOptions, const FTedsOutlinerParams& InInitTedsOptions, DataStorage::QueryHandle ColumnQuery) const
 {
-	using namespace UE::Editor::DataStorage;
-	ensureMsgf(AreEditorDataStorageFeaturesEnabled(), TEXT("Unable to initialize the Teds-Outliner before TEDS itself is initialized."));
+	UTypedElementRegistry* Registry = UTypedElementRegistry::GetInstance();
+
+	ensureMsgf(Registry&& Registry->AreDataStorageInterfacesSet(), TEXT("Unable to initialize the Teds-Outliner before TEDS itself is initialized."));
 
 	FSceneOutlinerInitializationOptions InitOptions(InInitOptions);
 	FTedsOutlinerParams InitTedsOptions(InInitTedsOptions);
@@ -116,8 +116,8 @@ void FTedsOutlinerModule::ShutdownModule()
 
 DataStorage::QueryHandle FTedsOutlinerModule::GetLevelEditorTedsOutlinerColumnQuery()
 {
-	using namespace UE::Editor::DataStorage;
-	IEditorDataStorageProvider* Storage = GetMutableDataStorageFeature<IEditorDataStorageProvider>(StorageFeatureName);
+	UTypedElementRegistry* Registry = UTypedElementRegistry::GetInstance();
+	IEditorDataStorageProvider* Storage = Registry->GetMutableDataStorage();
 		
 	using namespace DataStorage::Queries;
 
@@ -138,7 +138,10 @@ DataStorage::QueryHandle FTedsOutlinerModule::GetLevelEditorTedsOutlinerColumnQu
 
 TSharedRef<SWidget> FTedsOutlinerModule::CreateLevelEditorTedsOutliner()
 {
-	if(!DataStorage::AreEditorDataStorageFeaturesEnabled())
+	UTypedElementRegistry* Registry = UTypedElementRegistry::GetInstance();
+	checkf(Registry, TEXT("Unable to initialize the table viewer before TEDS is initialized."));
+
+	if(!Registry->AreDataStorageInterfacesSet())
 	{
 		return SNew(STextBlock)
 		.Text(LOCTEXT("TEDSPluginNotEnabledText", "You need to enable the Typed Element Data Storage plugin to see the table viewer!"));
diff --git a/Engine/Plugins/Experimental/EditorDataStorageFeatures/Source/TedsOutliner/Private/Widgets/AlertWidget.cpp b/Engine/Plugins/Experimental/EditorDataStorageFeatures/Source/TedsOutliner/Private/Widgets/AlertWidget.cpp
--- a/Engine/Plugins/Experimental/EditorDataStorageFeatures/Source/TedsOutliner/Private/Widgets/AlertWidget.cpp
+++ b/Engine/Plugins/Experimental/EditorDataStorageFeatures/Source/TedsOutliner/Private/Widgets/AlertWidget.cpp
@@ -6,7 +6,7 @@
 #include "Elements/Columns/TypedElementAlertColumns.h"
 #include "Elements/Columns/TypedElementMiscColumns.h"
 #include "Elements/Columns/TypedElementSlateWidgetColumns.h"
-#include "Elements/Common/EditorDataStorageFeatures.h"
+#include "Elements/Framework/TypedElementRegistry.h"
 #include "Widgets/Images/SImage.h"
 #include "Widgets/Input/SButton.h"
 #include "Widgets/SOverlay.h"
@@ -17,7 +17,7 @@
 namespace UE::Editor::DataStorage::Widgets::Private
 {
 	void UpdateWidget(const TSharedPtr<SWidget>& Widget, const FText& Alert, bool bIsWarning, uint16 ErrorCount, uint16 WarningCount,
-		RowHandle RowWithAlertAction)
+		UE::Editor::DataStorage::RowHandle RowWithAlertAction)
 	{
 		if (Widget)
 		{
@@ -117,13 +117,14 @@ namespace UE::Editor::DataStorage::Widgets::Private
 				}
 
 				// If there's an action to call, enable the invisible button, otherwise turn it off.
-				if (RowWithAlertAction != InvalidRowHandle)
+				if (RowWithAlertAction != UE::Editor::DataStorage::InvalidRowHandle)
 				{
 					Background.SetVisibility(EVisibility::HitTestInvisible);
 					ActionButton.SetVisibility(EVisibility::Visible);
 					ActionButton.SetOnClicked(FOnClicked::CreateLambda([RowWithAlertAction]()
 						{
-							const IEditorDataStorageProvider* DataStorage = GetDataStorageFeature<IEditorDataStorageProvider>(StorageFeatureName);
+							UTypedElementRegistry* Registry = UTypedElementRegistry::GetInstance();
+							const IEditorDataStorageProvider* DataStorage = Registry->GetDataStorage();
 							if (const FTypedElementAlertActionColumn* Action =
 								DataStorage->GetColumn<FTypedElementAlertActionColumn>(RowWithAlertAction))
 							{
diff --git a/Engine/Plugins/Experimental/EditorDataStorageFeatures/Source/TedsPropertyEditor/Private/Widgets/SPropertyMenuTedsRowPicker.cpp b/Engine/Plugins/Experimental/EditorDataStorageFeatures/Source/TedsPropertyEditor/Private/Widgets/SPropertyMenuTedsRowPicker.cpp
--- a/Engine/Plugins/Experimental/EditorDataStorageFeatures/Source/TedsPropertyEditor/Private/Widgets/SPropertyMenuTedsRowPicker.cpp
+++ b/Engine/Plugins/Experimental/EditorDataStorageFeatures/Source/TedsPropertyEditor/Private/Widgets/SPropertyMenuTedsRowPicker.cpp
@@ -2,7 +2,7 @@
 
 #include "Widgets/SPropertyMenuTedsRowPicker.h"
 
-#include "Elements/Common/EditorDataStorageFeatures.h"
+#include "Elements/Framework/TypedElementRegistry.h"
 #include "TedsRowPickingMode.h"
 #include "TedsOutlinerItem.h"
 
@@ -36,14 +36,16 @@ void SPropertyMenuTedsRowPicker::Construct(const FArguments& InArgs)
 		TSharedPtr<SWidget> MenuContent;
 		{
 			using namespace UE::Editor::Outliner;
-			using namespace UE::Editor::DataStorage;
 
 			// TEDS-Outliner TODO: Taken from private implementation of PropertyEditorAssetConstants.
 			//                     Should be centralized when TEDS is moved to core
 			static const FVector2D ContentBrowserWindowSize(300.0f, 300.0f);
 			static const FVector2D SceneOutlinerWindowSize(350.0f, 300.0f);
 
-			if (!AreEditorDataStorageFeaturesEnabled())
+			UTypedElementRegistry* Registry = UTypedElementRegistry::GetInstance();
+			checkf(Registry, TEXT("Unable to initialize the Typed Elements Outliner before TEDS is initialized."));
+
+			if (!Registry->AreDataStorageInterfacesSet())
 			{
 				MenuContent = SNew(STextBlock)
 					.Text(LOCTEXT("TEDSPluginNotEnabledText", "Typed Element Data Storage plugin required to use this property picker."));
diff --git a/Engine/Plugins/Experimental/EditorDataStorageFeatures/Source/TedsSettings/Private/TedsSettingsManager.cpp b/Engine/Plugins/Experimental/EditorDataStorageFeatures/Source/TedsSettings/Private/TedsSettingsManager.cpp
--- a/Engine/Plugins/Experimental/EditorDataStorageFeatures/Source/TedsSettings/Private/TedsSettingsManager.cpp
+++ b/Engine/Plugins/Experimental/EditorDataStorageFeatures/Source/TedsSettings/Private/TedsSettingsManager.cpp
@@ -5,10 +5,9 @@
 #include "Elements/Columns/TypedElementCompatibilityColumns.h"
 #include "Elements/Columns/TypedElementMiscColumns.h"
 #include "Elements/Columns/TypedElementSlateWidgetColumns.h"
-#include "Elements/Common/EditorDataStorageFeatures.h"
 #include "Elements/Framework/TypedElementIndexHasher.h"
 #include "Elements/Framework/TypedElementQueryBuilder.h"
-#include "Elements/Interfaces/TypedElementDataStorageInterface.h"
+#include "Elements/Framework/TypedElementRegistry.h"
 #include "ISettingsCategory.h"
 #include "ISettingsContainer.h"
 #include "ISettingsModule.h"
@@ -45,10 +44,14 @@ void FTedsSettingsManager::Initialize()
 {
 	if (!bIsInitialized)
 	{
-		using namespace UE::Editor::DataStorage;
-		auto OnDataStorage = [this]
+		FModuleManager::Get().LoadModule(TEXT("TypedElementFramework"));
+
+		UTypedElementRegistry* TypedElementRegistry = UTypedElementRegistry::GetInstance();
+		check(TypedElementRegistry);
+
+		auto OnDataStorage = [this, TypedElementRegistry]
 			{
-				IEditorDataStorageProvider* DataStorage = GetMutableDataStorageFeature<IEditorDataStorageProvider>(StorageFeatureName);
+				IEditorDataStorageProvider* DataStorage = TypedElementRegistry->GetMutableDataStorage();
 				check(DataStorage);
 
 				RegisterTables(*DataStorage);
@@ -56,13 +59,13 @@ void FTedsSettingsManager::Initialize()
 				RegisterSettings();
 			};
 
-		if (AreEditorDataStorageFeaturesEnabled())
+		if (TypedElementRegistry->AreDataStorageInterfacesSet())
 		{
 			OnDataStorage();
 		}
 		else
 		{
-			OnEditorDataStorageFeaturesEnabled().AddSPLambda(this, OnDataStorage);
+			TypedElementRegistry->OnDataStorageInterfacesSet().AddSPLambda(this, OnDataStorage);
 		}
 
 		bIsInitialized = true;
@@ -73,12 +76,14 @@ void FTedsSettingsManager::Shutdown()
 {
 	if (bIsInitialized)
 	{
-		using namespace UE::Editor::DataStorage;
-		OnEditorDataStorageFeaturesEnabled().RemoveAll(this);
+		UTypedElementRegistry* TypedElementRegistry = UTypedElementRegistry::GetInstance();
+		check(TypedElementRegistry);
+
+		TypedElementRegistry->OnDataStorageInterfacesSet().RemoveAll(this);
 
-		if (AreEditorDataStorageFeaturesEnabled())
+		if (TypedElementRegistry->AreDataStorageInterfacesSet())
 		{
-			IEditorDataStorageProvider* DataStorage = GetMutableDataStorageFeature<IEditorDataStorageProvider>(StorageFeatureName);
+			IEditorDataStorageProvider* DataStorage = TypedElementRegistry->GetMutableDataStorage();
 			check(DataStorage);
 
 			UnregisterSettings();
@@ -154,7 +159,10 @@ void FTedsSettingsManager::RegisterSettingsContainer(const FName& ContainerName)
 	ISettingsModule* SettingsModule = FModuleManager::GetModulePtr<ISettingsModule>("Settings");
 	check(SettingsModule);
 
-	IEditorDataStorageProvider* DataStorage = GetMutableDataStorageFeature<IEditorDataStorageProvider>(StorageFeatureName);
+	UTypedElementRegistry* TypedElementRegistry = UTypedElementRegistry::GetInstance();
+	check(TypedElementRegistry);
+
+	IEditorDataStorageProvider* DataStorage = TypedElementRegistry->GetMutableDataStorage();
 	check(DataStorage);
 
 	UE_LOG(LogTedsSettings, Log, TEXT("Register Settings Container : '%s'"), *ContainerName.ToString());
@@ -204,10 +212,13 @@ void FTedsSettingsManager::UnregisterSettings()
 	ISettingsModule* SettingsModule = FModuleManager::GetModulePtr<ISettingsModule>("Settings");
 	check(SettingsModule);
 
-	IEditorDataStorageProvider* DataStorage = GetMutableDataStorageFeature<IEditorDataStorageProvider>(StorageFeatureName);
+	UTypedElementRegistry* TypedElementRegistry = UTypedElementRegistry::GetInstance();
+	check(TypedElementRegistry);
+
+	IEditorDataStorageProvider* DataStorage = TypedElementRegistry->GetMutableDataStorage();
 	check(DataStorage);
 
-	IEditorDataStorageCompatibilityProvider* DataStorageCompatibility = GetMutableDataStorageFeature<IEditorDataStorageCompatibilityProvider>(CompatibilityFeatureName);
+	IEditorDataStorageCompatibilityProvider* DataStorageCompatibility = TypedElementRegistry->GetMutableDataStorageCompatibility();
 	check(DataStorageCompatibility);
 
 	SettingsModule->OnContainerAdded().RemoveAll(this);
@@ -272,10 +283,13 @@ void FTedsSettingsManager::UpdateSettingsCategory(TSharedPtr<ISettingsCategory>
 
 	TRACE_CPUPROFILER_EVENT_SCOPE(TedsSettingsManager.UpdateSettingsCategory);
 
-	IEditorDataStorageProvider* DataStorage = GetMutableDataStorageFeature<IEditorDataStorageProvider>(StorageFeatureName);
+	UTypedElementRegistry* TypedElementRegistry = UTypedElementRegistry::GetInstance();
+	check(TypedElementRegistry);
+
+	IEditorDataStorageProvider* DataStorage = TypedElementRegistry->GetMutableDataStorage();
 	check(DataStorage);
 
-	IEditorDataStorageCompatibilityProvider* DataStorageCompatibility = GetMutableDataStorageFeature<IEditorDataStorageCompatibilityProvider>(CompatibilityFeatureName);
+	IEditorDataStorageCompatibilityProvider* DataStorageCompatibility = TypedElementRegistry->GetMutableDataStorageCompatibility();
 	check(DataStorageCompatibility);
 
 	const FName& ContainerName = DataStorage->GetColumn<FNameColumn>(ContainerRow)->Name;
diff --git a/Engine/Plugins/Experimental/EditorDataStorageFeatures/Source/TedsSettings/Tests/TedsSettingsTests.cpp b/Engine/Plugins/Experimental/EditorDataStorageFeatures/Source/TedsSettings/Tests/TedsSettingsTests.cpp
--- a/Engine/Plugins/Experimental/EditorDataStorageFeatures/Source/TedsSettings/Tests/TedsSettingsTests.cpp
+++ b/Engine/Plugins/Experimental/EditorDataStorageFeatures/Source/TedsSettings/Tests/TedsSettingsTests.cpp
@@ -8,8 +8,8 @@
 #include "Editor.h"
 #include "Elements/Columns/TypedElementCompatibilityColumns.h"
 #include "Elements/Columns/TypedElementMiscColumns.h"
-#include "Elements/Common/EditorDataStorageFeatures.h"
 #include "Elements/Framework/TypedElementQueryBuilder.h"
+#include "Elements/Framework/TypedElementRegistry.h"
 #include "Elements/Interfaces/TypedElementDataStorageInterface.h"
 #include "ISettingsModule.h"
 #include "Modules/ModuleManager.h"
@@ -21,6 +21,7 @@ namespace UE::Editor::Settings::Tests
 {
 	BEGIN_DEFINE_SPEC(FTedsSettingsTestFixture, "Editor.DataStorage.Settings", EAutomationTestFlags::EditorContext | EAutomationTestFlags::EngineFilter)
 	ISettingsModule* SettingsModule = nullptr;
+	UTypedElementRegistry* TypedElementRegistry = nullptr;
 	IEditorDataStorageProvider* DataStorage = nullptr;
 	IEditorDataStorageCompatibilityProvider* DataStorageCompatibility = nullptr;
 	DataStorage::QueryHandle CountAllSettingsQuery = DataStorage::InvalidQueryHandle;
@@ -71,14 +72,16 @@ namespace UE::Editor::Settings::Tests
 
 		BeforeEach([this]()
 		{
-			using namespace UE::Editor::DataStorage;
 			SettingsModule = FModuleManager::GetModulePtr<ISettingsModule>("Settings");
 			check(SettingsModule != nullptr);
 
-			DataStorage = GetMutableDataStorageFeature<IEditorDataStorageProvider>(StorageFeatureName);
+			TypedElementRegistry = UTypedElementRegistry::GetInstance();
+			check(TypedElementRegistry != nullptr);
+
+			DataStorage = TypedElementRegistry->GetMutableDataStorage();
 			check(DataStorage != nullptr);
 
-			DataStorageCompatibility = GetMutableDataStorageFeature<IEditorDataStorageCompatibilityProvider>(CompatibilityFeatureName);
+			DataStorageCompatibility = TypedElementRegistry->GetMutableDataStorageCompatibility();
 			check(DataStorageCompatibility != nullptr);
 
 			{
@@ -102,6 +105,7 @@ namespace UE::Editor::Settings::Tests
 			TestRowHandles.Empty();
 			CountAllSettingsQuery = DataStorage::InvalidQueryHandle;
 			SettingsModule = nullptr;
+			TypedElementRegistry = nullptr;
 			DataStorage = nullptr;
 			DataStorageCompatibility = nullptr;
 		});
diff --git a/Engine/Plugins/Experimental/EditorDataStorageFeatures/Source/TedsStyling/Private/TedsStylingFactory.cpp b/Engine/Plugins/Experimental/EditorDataStorageFeatures/Source/TedsStyling/Private/TedsStylingFactory.cpp
--- a/Engine/Plugins/Experimental/EditorDataStorageFeatures/Source/TedsStyling/Private/TedsStylingFactory.cpp
+++ b/Engine/Plugins/Experimental/EditorDataStorageFeatures/Source/TedsStyling/Private/TedsStylingFactory.cpp
@@ -5,8 +5,8 @@
 #include "Columns/TedsStylingColumns.h"
 #include "Elements/Columns/TypedElementMiscColumns.h"
 #include "Elements/Columns/TypedElementSlateWidgetColumns.h"
-#include "Elements/Common/EditorDataStorageFeatures.h"
 #include "Elements/Framework/TypedElementIndexHasher.h"
+#include "Elements/Framework/TypedElementRegistry.h"
 #include "Elements/Interfaces/TypedElementDataStorageInterface.h"
 #include "HAL/IConsoleManager.h"
 #include "Styling/SlateBrush.h"
@@ -38,8 +38,7 @@ void UTedsStylingFactory::RegisterQueries(IEditorDataStorageProvider& DataStorag
 
 void UTedsStylingFactory::RegisterAllKnownStyles()
 {
-	using namespace UE::Editor::DataStorage;
-	IEditorDataStorageProvider* Interface = GetMutableDataStorageFeature<IEditorDataStorageProvider>(StorageFeatureName);
+	IEditorDataStorageProvider* Interface = UTypedElementRegistry::GetInstance()->GetMutableDataStorage();
 
 	auto RegisterBrushesForStyle = [Interface](const ISlateStyle& Style)
 	{
diff --git a/Engine/Plugins/Experimental/EditorDataStorageFeatures/Source/TedsTableViewer/Private/TedsTableViewerColumn.cpp b/Engine/Plugins/Experimental/EditorDataStorageFeatures/Source/TedsTableViewer/Private/TedsTableViewerColumn.cpp
--- a/Engine/Plugins/Experimental/EditorDataStorageFeatures/Source/TedsTableViewer/Private/TedsTableViewerColumn.cpp
+++ b/Engine/Plugins/Experimental/EditorDataStorageFeatures/Source/TedsTableViewer/Private/TedsTableViewerColumn.cpp
@@ -6,10 +6,6 @@
 #include "Elements/Columns/TypedElementMiscColumns.h"
 #include "Elements/Columns/TypedElementSlateWidgetColumns.h"
 #include "Elements/Columns/TypedElementTypeInfoColumns.h"
-#include "Elements/Common/EditorDataStorageFeatures.h"
-#include "Elements/Interfaces/TypedElementDataStorageInterface.h"
-#include "Elements/Interfaces/TypedElementDataStorageCompatibilityInterface.h"
-#include "Elements/Interfaces/TypedElementDataStorageUiInterface.h"
 #include "Elements/Framework/TypedElementDataStorageWidget.h"
 #include "Elements/Framework/TypedElementRegistry.h"
 #include "TedsTableViewerUtils.h"
@@ -28,9 +24,15 @@ namespace UE::Editor::DataStorage
 		, MatchedColumns(InMatchedColumns)
 		, WidgetMetaData(InWidgetMetaData)
 	{
-		Storage = GetMutableDataStorageFeature<IEditorDataStorageProvider>(StorageFeatureName);
-		StorageUi = GetMutableDataStorageFeature<IEditorDataStorageUiProvider>(UiFeatureName);
-		StorageCompatibility = GetMutableDataStorageFeature<IEditorDataStorageCompatibilityProvider>(CompatibilityFeatureName);
+		UTypedElementRegistry* Registry = UTypedElementRegistry::GetInstance();
+		
+		checkf(Registry, TEXT("Unable to create a Table Viewer column before the Typed Element Registry is initialized."));
+		if (Registry)
+		{
+			Storage = Registry->GetMutableDataStorage();
+			StorageUi = Registry->GetMutableDataStorageUi();
+			StorageCompatibility = Registry->GetMutableDataStorageCompatibility();
+		}
 
 		// Store the matched columns as a query condition that requires all of them (i.e AND's them)
 		for(const TWeakObjectPtr<const UScriptStruct>& Column : MatchedColumns)
diff --git a/Engine/Plugins/Experimental/EditorDataStorageFeatures/Source/TedsTableViewer/Private/TedsTableViewerModel.cpp b/Engine/Plugins/Experimental/EditorDataStorageFeatures/Source/TedsTableViewer/Private/TedsTableViewerModel.cpp
--- a/Engine/Plugins/Experimental/EditorDataStorageFeatures/Source/TedsTableViewer/Private/TedsTableViewerModel.cpp
+++ b/Engine/Plugins/Experimental/EditorDataStorageFeatures/Source/TedsTableViewer/Private/TedsTableViewerModel.cpp
@@ -4,10 +4,7 @@
 
 
 #include "Elements/Columns/TypedElementUIColumns.h"
-#include "Elements/Common/EditorDataStorageFeatures.h"
-#include "Elements/Interfaces/TypedElementDataStorageInterface.h"
-#include "Elements/Interfaces/TypedElementDataStorageCompatibilityInterface.h"
-#include "Elements/Interfaces/TypedElementDataStorageUiInterface.h"
+#include "Elements/Framework/TypedElementRegistry.h"
 #include "QueryStack/IQueryStackNode_Row.h"
 #include "TedsTableViewerColumn.h"
 #include "TedsTableViewerUtils.h"
@@ -22,9 +19,15 @@ namespace UE::Editor::DataStorage
 		, CellWidgetPurposes(InCellWidgetPurposes)
 		, IsItemVisible(InIsItemVisibleDelegate)
 	{
-		Storage = GetMutableDataStorageFeature<IEditorDataStorageProvider>(StorageFeatureName);
-		StorageUi = GetMutableDataStorageFeature<IEditorDataStorageUiProvider>(UiFeatureName);
-		StorageCompatibility = GetMutableDataStorageFeature<IEditorDataStorageCompatibilityProvider>(CompatibilityFeatureName);
+		UTypedElementRegistry* Registry = UTypedElementRegistry::GetInstance();
+		
+		checkf(Registry, TEXT("Unable to create a Table Viewer before the Typed Element Registry is initialized."));
+		if (Registry)
+		{
+			Storage = Registry->GetMutableDataStorage();
+			StorageUi = Registry->GetMutableDataStorageUi();
+			StorageCompatibility = Registry->GetMutableDataStorageCompatibility();
+		}
 		
 		TickerHandle = FTSTicker::GetCoreTicker().AddTicker(FTickerDelegate::CreateRaw(this, &FTedsTableViewerModel::Tick), 0);
 
diff --git a/Engine/Plugins/Experimental/EditorDataStorageFeatures/Source/TedsTableViewer/Private/TedsTableViewerModule.cpp b/Engine/Plugins/Experimental/EditorDataStorageFeatures/Source/TedsTableViewer/Private/TedsTableViewerModule.cpp
--- a/Engine/Plugins/Experimental/EditorDataStorageFeatures/Source/TedsTableViewer/Private/TedsTableViewerModule.cpp
+++ b/Engine/Plugins/Experimental/EditorDataStorageFeatures/Source/TedsTableViewer/Private/TedsTableViewerModule.cpp
@@ -3,9 +3,7 @@
 #include "Elements/Columns/TypedElementAlertColumns.h"
 #include "Elements/Columns/TypedElementCompatibilityColumns.h"
 #include "Elements/Columns/TypedElementLabelColumns.h"
-#include "Elements/Common/EditorDataStorageFeatures.h"
-#include "Elements/Framework/TypedElementQueryBuilder.h"
-#include "Elements/Interfaces/TypedElementDataStorageInterface.h"
+#include "Elements/Framework/TypedElementRegistry.h"
 #include "LevelEditor.h"
 #include "Modules/ModuleManager.h"
 #include "Modules/ModuleInterface.h"
@@ -64,14 +62,15 @@ public:
 	
 	TSharedRef<SWidget> CreateTableViewer()
 	{
-		using namespace UE::Editor::DataStorage;
+		UTypedElementRegistry* Registry = UTypedElementRegistry::GetInstance();
+		checkf(Registry, TEXT("Unable to initialize the table viewer before TEDS is initialized."));
 
-		if(!AreEditorDataStorageFeaturesEnabled())
+		if(!Registry->AreDataStorageInterfacesSet())
 		{
 			return SNullWidget::NullWidget;
 		}
 
-		IEditorDataStorageProvider* DataStorage = GetMutableDataStorageFeature<IEditorDataStorageProvider>(StorageFeatureName);
+		IEditorDataStorageProvider* DataStorage = Registry->GetMutableDataStorage();
 
 		using namespace UE::Editor::DataStorage::Queries;
 
@@ -85,14 +84,14 @@ public:
 
 		Rows.Empty();
 		
-		FQueryResult QueryResult = DataStorage->RunQuery(QueryHandle,
-			CreateDirectQueryCallbackBinding([this](const IEditorDataStorageProvider::IDirectQueryContext& Context, const RowHandle* RowHandles)
+		UE::Editor::DataStorage::FQueryResult QueryResult = DataStorage->RunQuery(QueryHandle,
+			CreateDirectQueryCallbackBinding([this](const IEditorDataStorageProvider::IDirectQueryContext& Context, const UE::Editor::DataStorage::RowHandle* RowHandles)
 			{
 				Rows.Append(RowHandles, Context.GetRowCount());
 			}));
 
-		return SNew(STedsTableViewer)
-			.QueryStack(MakeShared<FQueryStackNode_RowView>(&Rows))
+		return SNew(UE::Editor::DataStorage::STedsTableViewer)
+			.QueryStack(MakeShared<UE::Editor::DataStorage::FQueryStackNode_RowView>(&Rows))
 			.Columns({FTypedElementLabelColumn::StaticStruct(), FTypedElementSelectionColumn::StaticStruct(),
 				FTypedElementAlertColumn::StaticStruct(), FTypedElementChildAlertColumn::StaticStruct()});
 	}
diff --git a/Engine/Plugins/Experimental/EditorDataStorageFeatures/Source/TedsTableViewer/Private/Widgets/SRowDetails.cpp b/Engine/Plugins/Experimental/EditorDataStorageFeatures/Source/TedsTableViewer/Private/Widgets/SRowDetails.cpp
--- a/Engine/Plugins/Experimental/EditorDataStorageFeatures/Source/TedsTableViewer/Private/Widgets/SRowDetails.cpp
+++ b/Engine/Plugins/Experimental/EditorDataStorageFeatures/Source/TedsTableViewer/Private/Widgets/SRowDetails.cpp
@@ -4,7 +4,7 @@
 
 #include "Elements/Columns/TypedElementMiscColumns.h"
 #include "Elements/Columns/TypedElementTypeInfoColumns.h"
-#include "Elements/Common/EditorDataStorageFeatures.h"
+#include "Elements/Framework/TypedElementRegistry.h"
 #include "Elements/Interfaces/TypedElementDataStorageInterface.h"
 #include "Elements/Interfaces/TypedElementDataStorageUiInterface.h"
 #include "Widgets/Text/STextBlock.h"
@@ -36,10 +36,12 @@ namespace UE::Editor::DataStorage
 			WidgetPurposes = Widgets::Private::DefaultWidgetPurposes;
 		}
 
-		checkf(AreEditorDataStorageFeaturesEnabled(), TEXT("Unable to initialize SRowDetails without the editor data storage interfaces."));
+		UTypedElementRegistry* Registry = UTypedElementRegistry::GetInstance();
+		checkf(Registry, TEXT("Unable to initialize SRowDetails before TEDS is initialized."));
+		checkf(Registry->AreDataStorageInterfacesSet(), TEXT("Unable to initialize SRowDetails without the editor data storage interfaces."));
 
-		DataStorage = GetMutableDataStorageFeature<IEditorDataStorageProvider>(StorageFeatureName);
-		DataStorageUi = GetMutableDataStorageFeature<IEditorDataStorageUiProvider>(UiFeatureName);
+		DataStorage = Registry->GetMutableDataStorage();
+		DataStorageUi = Registry->GetMutableDataStorageUi();
 
 		
 		ChildSlot
diff --git a/Engine/Plugins/Experimental/InstanceDataObjectFixupTool/Source/Private/InstanceDataObjectFixupPanel.cpp b/Engine/Plugins/Experimental/InstanceDataObjectFixupTool/Source/Private/InstanceDataObjectFixupPanel.cpp
--- a/Engine/Plugins/Experimental/InstanceDataObjectFixupTool/Source/Private/InstanceDataObjectFixupPanel.cpp
+++ b/Engine/Plugins/Experimental/InstanceDataObjectFixupTool/Source/Private/InstanceDataObjectFixupPanel.cpp
@@ -13,7 +13,7 @@
 #include "UObject/PropertyBagRepository.h"
 
 #include "Elements/Columns/TypedElementAlertColumns.h"
-#include "Elements/Common/EditorDataStorageFeatures.h"
+#include "Elements/Framework/TypedElementRegistry.h"
 #include "Elements/Interfaces/TypedElementDataStorageInterface.h"
 #include "Elements/Interfaces/TypedElementDataStorageCompatibilityInterface.h"
 
@@ -250,8 +250,9 @@ FInstanceDataObjectFixupPanel::~FInstanceDataObjectFixupPanel()
 {
 	using namespace UE::Editor::DataStorage;
 
-	IEditorDataStorageProvider* DataStorage = GetMutableDataStorageFeature<IEditorDataStorageProvider>(StorageFeatureName);
-	IEditorDataStorageCompatibilityProvider* DataStorageCompatibility = DataStorageCompatibility = GetMutableDataStorageFeature<IEditorDataStorageCompatibilityProvider>(CompatibilityFeatureName);
+	UTypedElementRegistry* Registry = UTypedElementRegistry::GetInstance();
+	IEditorDataStorageProvider* DataStorage = Registry ? Registry->GetMutableDataStorage() : nullptr;
+	IEditorDataStorageCompatibilityProvider* DataStorageCompatibility = Registry ? Registry->GetMutableDataStorageCompatibility() : nullptr;
 
 	if (DataStorageCompatibility != nullptr && DataStorage != nullptr)
 	{
diff --git a/Engine/Plugins/Experimental/InstanceDataObjectFixupTool/Source/Private/InstanceDataObjectFixupToolTedsQueries.cpp b/Engine/Plugins/Experimental/InstanceDataObjectFixupTool/Source/Private/InstanceDataObjectFixupToolTedsQueries.cpp
--- a/Engine/Plugins/Experimental/InstanceDataObjectFixupTool/Source/Private/InstanceDataObjectFixupToolTedsQueries.cpp
+++ b/Engine/Plugins/Experimental/InstanceDataObjectFixupTool/Source/Private/InstanceDataObjectFixupToolTedsQueries.cpp
@@ -5,8 +5,8 @@
 #include "Elements/Columns/TypedElementAlertColumns.h"
 #include "Elements/Columns/TypedElementCompatibilityColumns.h"
 #include "Elements/Columns/TypedElementMiscColumns.h"
-#include "Elements/Common/EditorDataStorageFeatures.h"
 #include "Elements/Framework/TypedElementQueryBuilder.h"
+#include "Elements/Framework/TypedElementRegistry.h"
 #include "Elements/Interfaces/TypedElementDataStorageInterface.h"
 #include "Engine/World.h"
 #include "UObject/PropertyBagRepository.h"
@@ -61,8 +61,7 @@ void UInstanceDataObjectFixupToolTedsQueryFactory::ShowFixUpToolForLooseProperti
 
 void UInstanceDataObjectFixupToolTedsQueryFactory::ShowFixUpTool(UE::Editor::DataStorage::RowHandle Row, bool bRecurseIntoObject)
 {
-	using namespace UE::Editor::DataStorage;
-	IEditorDataStorageProvider* DataStorage = GetMutableDataStorageFeature<IEditorDataStorageProvider>(StorageFeatureName);
+	IEditorDataStorageProvider* DataStorage = UTypedElementRegistry::GetInstance()->GetMutableDataStorage();
 	if (FTypedElementUObjectColumn* ObjectColumn = DataStorage->GetColumn<FTypedElementUObjectColumn>(Row))
 	{
 		if (bRecurseIntoObject)
diff --git a/Engine/Source/Editor/UnrealEd/Private/Elements/Actor/ActorElementEditorSelectionInterface.cpp b/Engine/Source/Editor/UnrealEd/Private/Elements/Actor/ActorElementEditorSelectionInterface.cpp
--- a/Engine/Source/Editor/UnrealEd/Private/Elements/Actor/ActorElementEditorSelectionInterface.cpp
+++ b/Engine/Source/Editor/UnrealEd/Private/Elements/Actor/ActorElementEditorSelectionInterface.cpp
@@ -4,13 +4,11 @@
  
 #include "Elements/Actor/ActorElementData.h"
 #include "Elements/Columns/TypedElementSelectionColumns.h"
-#include "Elements/Common/EditorDataStorageFeatures.h"
 #include "Elements/Framework/EngineElementsLibrary.h"
 #include "Elements/Framework/TypedElementHandle.h"
 #include "Elements/Framework/TypedElementList.h"
+#include "Elements/Framework/TypedElementRegistry.h"
 #include "Elements/Interfaces/TypedElementSelectionInterface.h"
-#include "Elements/Interfaces/TypedElementDataStorageInterface.h"
-#include "Elements/Interfaces/TypedElementDataStorageCompatibilityInterface.h"
 #include "Elements/Object/ObjectElementEditorSelectionInterface.h"
 #include "GameFramework/Actor.h"
 #include "Serialization/Archive.h"
@@ -53,13 +51,14 @@ bool UActorElementEditorSelectionInterface::SelectElement(const FTypedElementHan
 	using namespace UE::Editor::DataStorage;
 
 	const AActor* Actor = ActorElementDataUtil::GetActorFromHandle(InElementHandle);
+	UTypedElementRegistry* Registry = InSelectionSet->GetRegistry();
 	// Add a selection column in TEDS
-	if (const IEditorDataStorageCompatibilityProvider* Compatibility = GetDataStorageFeature<IEditorDataStorageCompatibilityProvider>(CompatibilityFeatureName))
+	if (const IEditorDataStorageCompatibilityProvider* Compatibility = Registry->GetDataStorageCompatibility())
 	{
 		RowHandle Row = Compatibility->FindRowWithCompatibleObject(Actor);
 		if (Row != InvalidRowHandle)
 		{
-			if (IEditorDataStorageProvider* DataStorage = GetMutableDataStorageFeature<IEditorDataStorageProvider>(StorageFeatureName))
+			if (IEditorDataStorageProvider* DataStorage = UTypedElementRegistry::GetInstance()->GetMutableDataStorage())
 			{
 				DataStorage->AddColumn<FTypedElementSelectionColumn>(Row);
 
@@ -81,12 +80,13 @@ bool UActorElementEditorSelectionInterface::DeselectElement(const FTypedElementH
 	using namespace UE::Editor::DataStorage;
 
 	const AActor* Actor = ActorElementDataUtil::GetActorFromHandle(InElementHandle);
-	if (const IEditorDataStorageCompatibilityProvider* Compatibility = GetDataStorageFeature<IEditorDataStorageCompatibilityProvider>(CompatibilityFeatureName))
+	UTypedElementRegistry* Registry = InSelectionSet->GetRegistry();
+	if (const IEditorDataStorageCompatibilityProvider* Compatibility = Registry->GetDataStorageCompatibility())
 	{
 		RowHandle Row = Compatibility->FindRowWithCompatibleObject(Actor);
 		if (Row != InvalidRowHandle)
 		{
-			if (IEditorDataStorageProvider* DataStorage = GetMutableDataStorageFeature<IEditorDataStorageProvider>(StorageFeatureName))
+			if (IEditorDataStorageProvider* DataStorage = UTypedElementRegistry::GetInstance()->GetMutableDataStorage())
             {
             	DataStorage->RemoveColumn(Row, FTypedElementSelectionColumn::StaticStruct());
             }
diff --git a/Engine/Source/Editor/UnrealEd/Private/LevelEditorViewport.cpp b/Engine/Source/Editor/UnrealEd/Private/LevelEditorViewport.cpp
--- a/Engine/Source/Editor/UnrealEd/Private/LevelEditorViewport.cpp
+++ b/Engine/Source/Editor/UnrealEd/Private/LevelEditorViewport.cpp
@@ -44,14 +44,12 @@
 #include "UnrealEdGlobals.h"
 #include "Materials/MaterialExpressionTextureSample.h"
 #include "EditorSupportDelegates.h"
-#include "Elements/Common/EditorDataStorageFeatures.h"
 #include "Elements/Framework/TypedElementRegistry.h"
 #include "Elements/Framework/EngineElementsLibrary.h"
 #include "Elements/Framework/TypedElementCommonActions.h"
 #include "Elements/Framework/TypedElementListObjectUtil.h"
 #include "Elements/Framework/TypedElementViewportInteraction.h"
 #include "Elements/Interfaces/TypedElementObjectInterface.h"
-#include "Elements/Interfaces/TypedElementDataStorageCompatibilityInterface.h"
 #include "Elements/Actor/ActorElementLevelEditorViewportInteractionCustomization.h"
 #include "Elements/Component/ComponentElementLevelEditorViewportInteractionCustomization.h"
 #include "AudioDevice.h"
@@ -1661,7 +1659,6 @@ bool FLevelEditorViewportClient::UpdateDropPreviewElements(int32 MouseX, int32 M
 void FLevelEditorViewportClient::DestroyDropPreviewElements()
 {
 	using namespace LevelEditorViewportLocals;
-	using namespace UE::Editor::DataStorage;
 
 	if (!HasDropPreviewElements())
 	{
@@ -1673,7 +1670,8 @@ void FLevelEditorViewportClient::DestroyDropPreviewElements()
 	// deleting preview actors in UUnrealEdEngine::DeleteActors skips explicit handle deregistration,
 	// and although it does still happen in the immediately triggered garbage cleanup, that feels
 	// potentially brittle.
-	if (IEditorDataStorageCompatibilityProvider* TedsCompat = GetMutableDataStorageFeature<IEditorDataStorageCompatibilityProvider>(CompatibilityFeatureName))
+	UTypedElementRegistry* TypedElementRegistry = UTypedElementRegistry::GetInstance();
+	if (IEditorDataStorageCompatibilityProvider* TedsCompat = TypedElementRegistry->GetMutableDataStorageCompatibility())
 	{
 		DropPreviewElements->ForEachElement<ITypedElementObjectInterface>(
 			[this, TedsCompat](const TTypedElement<ITypedElementObjectInterface>& InElement)
diff --git a/Engine/Source/Runtime/Engine/Private/Elements/Framework/EngineElementsLibrary.cpp b/Engine/Source/Runtime/Engine/Private/Elements/Framework/EngineElementsLibrary.cpp
--- a/Engine/Source/Runtime/Engine/Private/Elements/Framework/EngineElementsLibrary.cpp
+++ b/Engine/Source/Runtime/Engine/Private/Elements/Framework/EngineElementsLibrary.cpp
@@ -2,10 +2,8 @@
 
 #include "Elements/Framework/EngineElementsLibrary.h"
 
-#include "Elements/Common/EditorDataStorageFeatures.h"
 #include "Elements/Framework/TypedElementRegistry.h"
 #include "Elements/Framework/TypedElementOwnerStore.h"
-#include "Elements/Interfaces/TypedElementDataStorageCompatibilityInterface.h"
 
 #include "Elements/Object/ObjectElementData.h"
 
@@ -381,8 +379,8 @@ void UEngineElementsLibrary::DestroyActorElement(const AActor* InActor, TTypedEl
 
 void UEngineElementsLibrary::RegisterActorElement(const AActor* InActor)
 {
-	using namespace UE::Editor::DataStorage;
-	IEditorDataStorageCompatibilityProvider* Storage = GetMutableDataStorageFeature<IEditorDataStorageCompatibilityProvider>(CompatibilityFeatureName);
+	UTypedElementRegistry* Registry = UTypedElementRegistry::GetInstance();
+	IEditorDataStorageCompatibilityProvider* Storage = Registry->GetMutableDataStorageCompatibility();
 	if (Storage)
 	{
 		Storage->AddCompatibleObject(const_cast<AActor*>(InActor));
@@ -391,8 +389,8 @@ void UEngineElementsLibrary::RegisterActorElement(const AActor* InActor)
 
 void UEngineElementsLibrary::UnregisterActorElement(const AActor* InActor)
 {
-	using namespace UE::Editor::DataStorage;
-	if (IEditorDataStorageCompatibilityProvider* Storage = GetMutableDataStorageFeature<IEditorDataStorageCompatibilityProvider>(CompatibilityFeatureName))
+	UTypedElementRegistry* Registry = UTypedElementRegistry::GetInstance();
+	if (IEditorDataStorageCompatibilityProvider* Storage = Registry->GetMutableDataStorageCompatibility())
 	{
 		Storage->RemoveCompatibleObject(const_cast<AActor*>(InActor));
 	}
diff --git a/Engine/Source/Runtime/TypedElementFramework/Private/Elements/Framework/TypedElementAttributeBinding.cpp b/Engine/Source/Runtime/TypedElementFramework/Private/Elements/Framework/TypedElementAttributeBinding.cpp
--- a/Engine/Source/Runtime/TypedElementFramework/Private/Elements/Framework/TypedElementAttributeBinding.cpp
+++ b/Engine/Source/Runtime/TypedElementFramework/Private/Elements/Framework/TypedElementAttributeBinding.cpp
@@ -2,17 +2,17 @@
 
 #include "Elements/Framework/TypedElementAttributeBinding.h"
 
-#include "Elements/Common/EditorDataStorageFeatures.h"
+#include "Elements/Framework/TypedElementRegistry.h"
 
 namespace UE::Editor::DataStorage
 {
 	FAttributeBinder::FAttributeBinder(RowHandle InTargetRow)
 		: TargetRow(InTargetRow)
 	{
-		if (ensureMsgf(AreEditorDataStorageFeaturesEnabled(),
+		if (ensureMsgf(UTypedElementRegistry::GetInstance()->AreDataStorageInterfacesSet(), 
 				TEXT("The Editor Data Storage plugin needs to be enabled to use attribute bindings.")))
 		{
-			DataStorage = GetMutableDataStorageFeature<IEditorDataStorageProvider>(StorageFeatureName);
+			DataStorage = UTypedElementRegistry::GetInstance()->GetMutableDataStorage();
 		}
 	}
 
diff --git a/Engine/Source/Runtime/TypedElementFramework/Private/Elements/Framework/TypedElementDataStorageWidget.cpp b/Engine/Source/Runtime/TypedElementFramework/Private/Elements/Framework/TypedElementDataStorageWidget.cpp
--- a/Engine/Source/Runtime/TypedElementFramework/Private/Elements/Framework/TypedElementDataStorageWidget.cpp
+++ b/Engine/Source/Runtime/TypedElementFramework/Private/Elements/Framework/TypedElementDataStorageWidget.cpp
@@ -2,8 +2,7 @@
 
 #include "Elements/Framework/TypedElementDataStorageWidget.h"
 #include "Elements/Columns/TypedElementSlateWidgetColumns.h"
-#include "Elements/Common//EditorDataStorageFeatures.h"
-#include "Elements/Interfaces/TypedElementDataStorageInterface.h"
+#include "Elements/Framework/TypedElementRegistry.h"
 
 STedsWidget::STedsWidget()
 	: UiRowHandle(UE::Editor::DataStorage::InvalidRowHandle)
@@ -76,6 +75,12 @@ UE::Editor::DataStorage::RowHandle STedsWidget::GetRowHandle() const
 
 IEditorDataStorageProvider* STedsWidget::GetStorageIfAvailable()
 {
-	using namespace UE::Editor::DataStorage;
-	return GetMutableDataStorageFeature<IEditorDataStorageProvider>(StorageFeatureName);
+	UTypedElementRegistry* Registry = UTypedElementRegistry::GetInstance();
+
+	if(!Registry || !Registry->AreDataStorageInterfacesSet())
+	{
+		return nullptr;
+	}
+
+	return Registry->GetMutableDataStorage();
 }
diff --git a/Engine/Source/Runtime/TypedElementFramework/Public/Elements/Common/EditorDataStorageFeatures.h b/Engine/Source/Runtime/TypedElementFramework/Public/Elements/Common/EditorDataStorageFeatures.h
--- a/Engine/Source/Runtime/TypedElementFramework/Public/Elements/Common/EditorDataStorageFeatures.h
+++ b/Engine/Source/Runtime/TypedElementFramework/Public/Elements/Common/EditorDataStorageFeatures.h
@@ -14,7 +14,7 @@ namespace UE::Editor::DataStorage
 	TYPEDELEMENTFRAMEWORK_API const extern FName UiFeatureName;
 
 	TYPEDELEMENTFRAMEWORK_API FSimpleMulticastDelegate& OnEditorDataStorageFeaturesEnabled();
-	TYPEDELEMENTFRAMEWORK_API bool AreEditorDataStorageFeaturesEnabled();
+	bool AreEditorDataStorageFeaturesEnabled();
 
 	template<typename T>
 	const T* GetDataStorageFeature(const FName InName)
diff --git a/Engine/Source/Runtime/TypedElementFramework/Public/Elements/Framework/TypedElementRegistry.h b/Engine/Source/Runtime/TypedElementFramework/Public/Elements/Framework/TypedElementRegistry.h
--- a/Engine/Source/Runtime/TypedElementFramework/Public/Elements/Framework/TypedElementRegistry.h
+++ b/Engine/Source/Runtime/TypedElementFramework/Public/Elements/Framework/TypedElementRegistry.h
@@ -77,25 +77,25 @@ public:
 	UFUNCTION(BlueprintPure, DisplayName="Get Default Typed Element Registry", Category = "TypedElementFramework|Registry", meta=(ScriptName="GetDefaultTypedElementRegistry"))
 	static TYPEDELEMENTFRAMEWORK_API UTypedElementRegistry* GetInstance();
 
-	UE_DEPRECATED(5.5, "Function has moved to modular features. Use UE::Editor::DataStorage::GetMutableDataStorageFeature instead.")
+	//UE_DEPRECATED(5.5, "Function has moved to modular features")
 	TYPEDELEMENTFRAMEWORK_API IEditorDataStorageProvider* GetMutableDataStorage();
-	UE_DEPRECATED(5.5, "Function has moved to modular features. Use UE::Editor::DataStorage::GetDataStorageFeature instead.")
+	//UE_DEPRECATED(5.5, "Function has moved to modular features")
 	TYPEDELEMENTFRAMEWORK_API const IEditorDataStorageProvider* GetDataStorage() const;
 
-	UE_DEPRECATED(5.5, "Function has moved to modular features. Use UE::Editor::DataStorage::GetMutableDataStorageFeature instead.")
+	//UE_DEPRECATED(5.5, "Function has moved to modular features")
 	TYPEDELEMENTFRAMEWORK_API IEditorDataStorageCompatibilityProvider* GetMutableDataStorageCompatibility();
-	UE_DEPRECATED(5.5, "Function has moved to modular features. Use UE::Editor::DataStorage::GetDataStorageFeature instead.")
+	//UE_DEPRECATED(5.5, "Function has moved to modular features")
 	TYPEDELEMENTFRAMEWORK_API const IEditorDataStorageCompatibilityProvider* GetDataStorageCompatibility() const;
 
-	UE_DEPRECATED(5.5, "Function has moved to modular features. Use UE::Editor::DataStorage::GetMutableDataStorageFeature instead.")
+	//UE_DEPRECATED(5.5, "Function has moved to modular features")
 	TYPEDELEMENTFRAMEWORK_API IEditorDataStorageUiProvider* GetMutableDataStorageUi();
-	UE_DEPRECATED(5.5, "Function has moved to modular features. Use UE::Editor::DataStorage::GetDataStorageFeature instead.")
+	//UE_DEPRECATED(5.5, "Function has moved to modular features")
 	TYPEDELEMENTFRAMEWORK_API const IEditorDataStorageUiProvider* GetDataStorageUi() const;
 
-	UE_DEPRECATED(5.5, "Use UE::Editor::DataStorage::AreDataStorageInterfacesSet instead.")
+	//UE_DEPRECATED(5.5, "Use UE::Editor::DataStorage::AreDataStorageInterfacesSet")
 	TYPEDELEMENTFRAMEWORK_API bool AreDataStorageInterfacesSet() const;
 
-	UE_DEPRECATED(5.5, "Use UE::Editor::DataStorage::OnAllStorageFeaturesEnabled instead")
+	//UE_DEPRECATED(5.5, "Use UE::Editor::DataStorage::OnAllStorageFeaturesEnabled")
 	TYPEDELEMENTFRAMEWORK_API FSimpleMulticastDelegate& OnDataStorageInterfacesSet();
 
 	/**
diff --git a/Engine/Source/Runtime/TypedElementFramework/Tests/Elements/Framework/TypedElementDataStoragePerformanceTestCommands.cpp b/Engine/Source/Runtime/TypedElementFramework/Tests/Elements/Framework/TypedElementDataStoragePerformanceTestCommands.cpp
--- a/Engine/Source/Runtime/TypedElementFramework/Tests/Elements/Framework/TypedElementDataStoragePerformanceTestCommands.cpp
+++ b/Engine/Source/Runtime/TypedElementFramework/Tests/Elements/Framework/TypedElementDataStoragePerformanceTestCommands.cpp
@@ -2,9 +2,9 @@
 
 #include "TypedElementDataStoragePerformanceTestCommands.h"
 
-#include "Elements/Common/EditorDataStorageFeatures.h"
 #include "Elements/Interfaces/TypedElementDataStorageInterface.h"
 #include "Elements/Framework/TypedElementQueryBuilder.h"
+#include "Elements/Framework/TypedElementRegistry.h"
 #include "HAL/IConsoleManager.h"
 
 namespace UE::Editor::DataStorage
@@ -24,7 +24,7 @@ namespace UE::Editor::DataStorage
 				int32 EntitiesToAdd;
 				LexFromString(EntitiesToAdd, *Args[0]);
 
-				IEditorDataStorageProvider* DataStorageInterface = GetMutableDataStorageFeature<IEditorDataStorageProvider>(StorageFeatureName);
+				IEditorDataStorageProvider* DataStorageInterface = UTypedElementRegistry::GetInstance()->GetMutableDataStorage();
 			
 				DataStorageInterface->BatchAddRow(Private::PerformanceTestCommandTable, EntitiesToAdd, [DataStorageInterface](RowHandle Row)
 					{
@@ -42,7 +42,7 @@ namespace UE::Editor::DataStorage
 		{
 			using namespace UE::Editor::DataStorage::Queries;
 
-			IEditorDataStorageProvider* DataStorageInterface = GetMutableDataStorageFeature<IEditorDataStorageProvider>(StorageFeatureName);
+			IEditorDataStorageProvider* DataStorageInterface = UTypedElementRegistry::GetInstance()->GetMutableDataStorage();
 		
 			QueryHandle Query = DataStorageInterface->RegisterQuery(
 				Select().Where().All<FTest_PingPongPrePhys>().Compile());
diff --git a/Engine/Source/Runtime/TypedElementFramework/Tests/Elements/Framework/TypedElementDataStorageProcessorTests.cpp b/Engine/Source/Runtime/TypedElementFramework/Tests/Elements/Framework/TypedElementDataStorageProcessorTests.cpp
--- a/Engine/Source/Runtime/TypedElementFramework/Tests/Elements/Framework/TypedElementDataStorageProcessorTests.cpp
+++ b/Engine/Source/Runtime/TypedElementFramework/Tests/Elements/Framework/TypedElementDataStorageProcessorTests.cpp
@@ -1,9 +1,7 @@
 // Copyright Epic Games, Inc. All Rights Reserved.
 
 #if WITH_TESTS
-#include "Elements/Common/EditorDataStorageFeatures.h"
-#include "Elements/Framework/TypedElementQueryBuilder.h"
-#include "Elements/Interfaces/TypedElementDataStorageInterface.h"
+#include "Elements/Framework/TypedElementRegistry.h"
 #include "HAL/IConsoleManager.h"
 #include "TypedElementTestColumns.h"
 
@@ -30,7 +28,7 @@ namespace UE::Editor::DataStorage::Debug::ProcessorTests
 	{
 		using namespace UE::Editor::DataStorage::Queries;
 
-		IEditorDataStorageProvider* DataStorage = GetMutableDataStorageFeature<IEditorDataStorageProvider>(StorageFeatureName);
+		IEditorDataStorageProvider* DataStorage = UTypedElementRegistry::GetInstance()->GetMutableDataStorage();
 
 		if (PrimaryTable == InvalidTableHandle)
 		{
@@ -124,7 +122,7 @@ namespace UE::Editor::DataStorage::Debug::ProcessorTests
 
 	void UnregisterProcessors()
 	{
-		IEditorDataStorageProvider* DataStorage = GetMutableDataStorageFeature<IEditorDataStorageProvider>(StorageFeatureName);
+		IEditorDataStorageProvider* DataStorage = UTypedElementRegistry::GetInstance()->GetMutableDataStorage();
 		for (const UE::Editor::DataStorage::QueryHandle& Handle : RegisteredQueries)
 		{
 			DataStorage->UnregisterQuery(Handle);
@@ -147,7 +145,7 @@ namespace UE::Editor::DataStorage::Debug::ProcessorTests
 				return;
 			}
 
-			IEditorDataStorageProvider* DataStorage = GetMutableDataStorageFeature<IEditorDataStorageProvider>(StorageFeatureName);
+			IEditorDataStorageProvider* DataStorage = UTypedElementRegistry::GetInstance()->GetMutableDataStorage();
 			DataStorage->BatchAddRow(PrimaryTable, RowsToCreate, [](UE::Editor::DataStorage::RowHandle Row)
 			{
 			});
@@ -180,3 +178,4 @@ namespace UE::Editor::DataStorage::Debug::ProcessorTests
 	};
 } // namespace UE::Editor::DataStorage::Debug::ProcessorTests
 #endif
+
diff --git a/Engine/Source/Runtime/TypedElementFramework/Tests/Elements/Framework/TypedElementQueryBuilderTests.cpp b/Engine/Source/Runtime/TypedElementFramework/Tests/Elements/Framework/TypedElementQueryBuilderTests.cpp
--- a/Engine/Source/Runtime/TypedElementFramework/Tests/Elements/Framework/TypedElementQueryBuilderTests.cpp
+++ b/Engine/Source/Runtime/TypedElementFramework/Tests/Elements/Framework/TypedElementQueryBuilderTests.cpp
@@ -1,5 +1,6 @@
 // Copyright Epic Games, Inc. All Rights Reserved.
 
+#include "Elements/Framework/TypedElementRegistry.h"
 #if WITH_TESTS
 #include "TypedElementTestColumns.h"
 
@@ -7,9 +8,7 @@
 #include "CoreTypes.h"
 #include "Containers/Array.h"
 #include "Containers/UnrealString.h"
-#include "Elements/Common/EditorDataStorageFeatures.h"
 #include "Elements/Common/TypedElementQueryConditions.h"
-#include "Elements/Interfaces/TypedElementDataStorageInterface.h"
 #include "Tests/TestHarnessAdapter.h"
 #include "UObject/Class.h"
 
@@ -32,9 +31,9 @@ static void AppendColumnName(FString& Output, TWeakObjectPtr<const UScriptStruct
 
 static bool TestMatching(FConditions& TestQuery, const TArray<FColumnBase>& RequestedColumns, bool Expected, bool Sort = false)
 {
-	IEditorDataStorageProvider* Storage = GetMutableDataStorageFeature<IEditorDataStorageProvider>(StorageFeatureName);
+	IEditorDataStorageProvider* Storage = UTypedElementRegistry::GetInstance()->GetMutableDataStorage();
 
-	TestQuery.Compile(FEditorStorageQueryConditionCompileContext(Storage));
+	TestQuery.Compile(UE::Editor::DataStorage::Queries::FEditorStorageQueryConditionCompileContext(Storage));
 	
 	if (Sort)
 	{
