Fix crash that can occur when renaming actors
* Crash happens because the side-effects of FActorLabelUtilities::RenameExistingActor can be numerous such as deleting other actors.
* These side-effects can end up calling back into TEDS via IEditorDataStorage interfaces using the direct API and cause memory corruption or assertions.
* The fix is to defer the rename to outside of the TEDS processing phases but before anything else on the gamethread has had a chance to tick.
* To facilitate this (and other potential cases that might crop up in the future), created a command queue that can be used in situations where the side-effects of the function that is being used cannot be predicted to be safe in all situations.  The command queue is flushed at the end of processing all phases in TEDS.

diff --git a/Engine/Plugins/Experimental/EditorDataStorage/Source/TedsCore/Private/Compatibility/TypedElementActorLabelQueries.cpp b/Engine/Plugins/Experimental/EditorDataStorage/Source/TedsCore/Private/Compatibility/TypedElementActorLabelQueries.cpp
--- a/Engine/Plugins/Experimental/EditorDataStorage/Source/TedsCore/Private/Compatibility/TypedElementActorLabelQueries.cpp
+++ b/Engine/Plugins/Experimental/EditorDataStorage/Source/TedsCore/Private/Compatibility/TypedElementActorLabelQueries.cpp
@@ -50,13 +50,29 @@ void UActorLabelDataStorageFactory::RegisterActorLabelToColumnQuery(IEditorDataS
 void UActorLabelDataStorageFactory::RegisterLabelColumnToActorQuery(IEditorDataStorageProvider& DataStorage) const
 {
 	using namespace UE::Editor::DataStorage::Queries;
+
+	struct FRenameActorCommand
+	{
+		void operator()()
+		{
+			TStrongObjectPtr<AActor> PinnedActor = Actor.Pin();
+			if (PinnedActor)
+			{
+				const FScopedTransaction Transaction(LOCTEXT("RenameActorTransaction", "Rename Actor"));
+				FActorLabelUtilities::RenameExistingActor(PinnedActor.Get(), NewLabel);
+			}
+			
+		}
+		TWeakObjectPtr<AActor> Actor;
+		FString NewLabel;
+	};
 	
 	DataStorage.RegisterQuery(
 		Select(
 			TEXT("Sync label column to actor"),
 			FProcessor(EQueryTickPhase::FrameEnd, DataStorage.GetQueryTickGroupName(EQueryTickGroups::SyncDataStorageToExternal))
 				.SetExecutionMode(EExecutionMode::GameThread),
-			[](FTypedElementUObjectColumn& Actor, const FTypedElementLabelColumn& Label, const FTypedElementLabelHashColumn& LabelHash)
+			[](IQueryContext& Context, FTypedElementUObjectColumn& Actor, const FTypedElementLabelColumn& Label, const FTypedElementLabelHashColumn& LabelHash)
 			{
 				if (AActor* ActorInstance = Cast<AActor>(Actor.Object); ActorInstance != nullptr)
 				{
@@ -64,8 +80,11 @@ void UActorLabelDataStorageFactory::RegisterLabelColumnToActorQuery(IEditorDataS
 					uint64 ActorLabelHash = CityHash64(reinterpret_cast<const char*>(*ActorLabel), ActorLabel.Len() * sizeof(**ActorLabel));
 					if (LabelHash.LabelHash != ActorLabelHash)
 					{
-						const FScopedTransaction Transaction(LOCTEXT("RenameActorTransaction", "Rename Actor"));
-						FActorLabelUtilities::RenameExistingActor(ActorInstance, Label.Label);
+						Context.PushCommand(FRenameActorCommand
+							{
+								.Actor = ActorInstance,
+								.NewLabel = Label.Label
+							});
 					}
 				}
 			}
diff --git a/Engine/Plugins/Experimental/EditorDataStorage/Source/TedsCore/Private/Processors/TypedElementProcessorAdaptors.cpp b/Engine/Plugins/Experimental/EditorDataStorage/Source/TedsCore/Private/Processors/TypedElementProcessorAdaptors.cpp
--- a/Engine/Plugins/Experimental/EditorDataStorage/Source/TedsCore/Private/Processors/TypedElementProcessorAdaptors.cpp
+++ b/Engine/Plugins/Experimental/EditorDataStorage/Source/TedsCore/Private/Processors/TypedElementProcessorAdaptors.cpp
@@ -612,6 +612,51 @@ namespace UE::Editor::DataStorage
 				this->Context.Defer().DestroyEntities(
 					TConstArrayView<FMassEntityHandle>(reinterpret_cast<const FMassEntityHandle*>(Rows.begin()), Rows.Num()));
 			}
+
+			void PushCommand(void (*CommandFunction)(void*), void* InCommandData)
+			{
+				if (!ensure(CommandFunction))
+				{
+					return;
+				}
+				const FEnvironment::FEnvironmentCommand Command
+				{
+					.CommandFunction = CommandFunction,
+					.CommandData = InCommandData
+				};
+				Environment.PushCommands(MakeConstArrayView(&Command, 1));
+			}
+
+			void* EmplaceObjectInScratch(size_t ObjectSize, size_t Alignment, void(* Construct)(void*, void*), void(* Destroy)(void*), void* SourceCommandContext)
+			{
+				FScratchBuffer& ScratchBuffer = Environment.GetScratchBuffer();
+				void* ObjectMemory = ScratchBuffer.Allocate(ObjectSize, Alignment);
+				Construct(ObjectMemory, SourceCommandContext);
+				// The presence of a Destroy function implies that the objects that was just added to the scratch buffer
+				// is not trivially destructable, hence needs its destructor called.
+				// The API for the scratch buffer's internal memory allocator needs us to emplace a non-trivially destructable object
+				// of some type.  FDestructor is used to fulfil that role to destroy the object that was just constructed.
+				if (Destroy)
+				{
+					struct FDestructor
+					{
+						using DestroyFnType = decltype(Destroy);
+						using ObjectPtrType = decltype(ObjectMemory);
+						FDestructor(DestroyFnType InDestroyFn, ObjectPtrType InObjectPtr)
+							: DestroyFn(InDestroyFn)
+							, ObjectPtr(InObjectPtr)
+						{}
+						~FDestructor()
+						{
+							DestroyFn(ObjectPtr);
+						}
+						DestroyFnType DestroyFn;
+						ObjectPtrType ObjectPtr;
+					};
+					ScratchBuffer.Emplace<FDestructor>(Destroy, ObjectMemory);
+				}
+				return ObjectMemory;
+			}
 		};
 
 		struct FMassDirectContextForwarder final : public IEditorDataStorageProvider::IDirectQueryContext
@@ -665,7 +710,9 @@ namespace UE::Editor::DataStorage
 			void RemoveColumns(RowHandle Row, TConstArrayView<const UScriptStruct*> ColumnTypes) override { Implementation.RemoveColumns(Row, ColumnTypes); }
 			void RemoveColumns(TConstArrayView<RowHandle> Rows, TConstArrayView<const UScriptStruct*> ColumnTypes) override { Implementation.RemoveColumns(Rows, ColumnTypes); }
 			const UScriptStruct* FindDynamicColumnType(const UE::Editor::DataStorage::FDynamicColumnDescription& Description) const override { return Implementation.FindDynamicColumnType(Description); }
-			
+			void PushCommand(void (*CommandFunction)(void*), void* CommandData) override { return Implementation.PushCommand(CommandFunction, CommandData); }
+		protected:
+			void* EmplaceObjectInScratch(const FEmplaceObjectParams& Params) override { return Implementation.EmplaceObjectInScratch(Params.ObjectSize, Params.Alignment, Params.Construct, Params.Destroy, Params.SourceObject); }
 
 			FMassWithEnvironmentContextCommon Implementation;
 		};
@@ -811,6 +858,7 @@ namespace UE::Editor::DataStorage
 			void RemoveColumns(RowHandle Row, TConstArrayView<const UScriptStruct*> ColumnTypes) override { Implementation.RemoveColumns(Row, ColumnTypes); }
 			void RemoveColumns(TConstArrayView<RowHandle> Rows, TConstArrayView<const UScriptStruct*> ColumnTypes) override { Implementation.RemoveColumns(Rows, ColumnTypes); }
 			const UScriptStruct* FindDynamicColumnType(const UE::Editor::DataStorage::FDynamicColumnDescription& Description) const override { return Implementation.FindDynamicColumnType(Description); }
+			void PushCommand(void (*CommandFunction)(void*), void* Context) override { return Implementation.PushCommand(CommandFunction, Context); }
 
 			const UObject* GetDependency(const UClass* DependencyClass) override { return Implementation.GetDependency(DependencyClass); }
 			UObject* GetMutableDependency(const UClass* DependencyClass) override { return Implementation.GetMutableDependency(DependencyClass); }
@@ -821,6 +869,9 @@ namespace UE::Editor::DataStorage
 			FQueryResult RunSubquery(int32 SubqueryIndex, UE::Editor::DataStorage::SubqueryCallbackRef Callback) override { return Implementation.RunSubquery(SubqueryIndex, Callback); }
 			FQueryResult RunSubquery(int32 SubqueryIndex, RowHandle Row, UE::Editor::DataStorage::SubqueryCallbackRef Callback) override { return Implementation.RunSubquery(SubqueryIndex, Row, Callback); }
 	
+protected:
+			void* EmplaceObjectInScratch(const FEmplaceObjectParams& Params) override { return Implementation.EmplaceObjectInScratch(Params.ObjectSize, Params.Alignment, Params.Construct, Params.Destroy, Params.SourceObject); }
+			
 			FMassQueryContextImplementation Implementation;
 		};
 	} // namespace Processors::Private
diff --git a/Engine/Plugins/Experimental/EditorDataStorage/Source/TedsCore/Private/TypedElementDatabaseEnvironment.cpp b/Engine/Plugins/Experimental/EditorDataStorage/Source/TedsCore/Private/TypedElementDatabaseEnvironment.cpp
--- a/Engine/Plugins/Experimental/EditorDataStorage/Source/TedsCore/Private/TypedElementDatabaseEnvironment.cpp
+++ b/Engine/Plugins/Experimental/EditorDataStorage/Source/TedsCore/Private/TypedElementDatabaseEnvironment.cpp
@@ -114,6 +114,7 @@ namespace UE::Editor::DataStorage
 	void FEnvironment::NextUpdateCycle()
 	{
 		Queries.UpdateActivatableQueries();
+		FlushCommands();
 		ScratchBuffer.BatchDelete();
 		UpdateCycleId++;
 	}
@@ -122,4 +123,20 @@ namespace UE::Editor::DataStorage
 	{
 		return UpdateCycleId;
 	}
+
+	void FEnvironment::PushCommands(TConstArrayView<const FEnvironmentCommand> Commands)
+	{
+		TUniqueLock<FMutex> Lock(CommandQueueMutex);
+
+		CommandQueue.Append(Commands);
+	}
+
+	void FEnvironment::FlushCommands()
+	{
+		for (FEnvironmentCommand& Command : CommandQueue)
+		{
+			Command.CommandFunction(Command.CommandData);
+		}
+		CommandQueue.SetNum(0, EAllowShrinking::No);
+	}
 } // namespace UE::Editor::DataStorage
diff --git a/Engine/Plugins/Experimental/EditorDataStorage/Source/TedsCore/Private/TypedElementDatabaseEnvironment.h b/Engine/Plugins/Experimental/EditorDataStorage/Source/TedsCore/Private/TypedElementDatabaseEnvironment.h
--- a/Engine/Plugins/Experimental/EditorDataStorage/Source/TedsCore/Private/TypedElementDatabaseEnvironment.h
+++ b/Engine/Plugins/Experimental/EditorDataStorage/Source/TedsCore/Private/TypedElementDatabaseEnvironment.h
@@ -2,6 +2,7 @@
 
 #pragma once
 
+#include "Async/Mutex.h"
 #include "DynamicColumnGenerator.h"
 #include "MassEntityManager.h"
 #include "MassProcessingPhaseManager.h"
@@ -60,7 +61,19 @@ namespace UE::Editor::DataStorage
 		void NextUpdateCycle();
 		uint64 GetUpdateCycleId() const;
 
+		struct FEnvironmentCommand
+		{
+			void (*CommandFunction)(void*);
+			// If this is not static data or null, it should be a pointer into the scratch buffer
+			void* CommandData = nullptr;
+		};
+
+		// Commands are flushed on NextUpdateCycle
+		void PushCommands(TConstArrayView<const FEnvironmentCommand> Commands);
+
 	private:
+		void FlushCommands();
+		
 		UEditorDataStorage& DataStorage;
 		Legacy::FCommandBuffer DirectDeferredCommands;
 		FIndexTable IndexTable;
@@ -70,6 +83,9 @@ namespace UE::Editor::DataStorage
 		FDynamicColumnGenerator DynamicColumnGenerator;
 		FValueTagManager ValueTagManager;
 
+		FMutex CommandQueueMutex;
+		TArray<FEnvironmentCommand> CommandQueue;
+
 		FMassEntityManager& MassEntityManager;
 		FMassProcessingPhaseManager& MassPhaseManager;
 
diff --git a/Engine/Source/Runtime/TypedElementFramework/Public/Elements/Framework/TypedElementQueryBuilder.h b/Engine/Source/Runtime/TypedElementFramework/Public/Elements/Framework/TypedElementQueryBuilder.h
--- a/Engine/Source/Runtime/TypedElementFramework/Public/Elements/Framework/TypedElementQueryBuilder.h
+++ b/Engine/Source/Runtime/TypedElementFramework/Public/Elements/Framework/TypedElementQueryBuilder.h
@@ -330,6 +330,8 @@ namespace UE::Editor::DataStorage::Queries
 		inline void RemoveColumns(RowHandle Row, TConstArrayView<const UScriptStruct*> ColumnTypes) override;
 		inline void RemoveColumns(TConstArrayView<RowHandle> Rows, TConstArrayView<const UScriptStruct*> ColumnTypes) override;
 
+		inline void PushCommand(void(* CommandFunction)(void*), void* CommandData) override;
+
 		inline FQueryResult RunQuery(QueryHandle Query) override;
 		inline FQueryResult RunSubquery(int32 SubqueryIndex) override;
 		inline FQueryResult RunSubquery(int32 SubqueryIndex, SubqueryCallbackRef Callback) override;
diff --git a/Engine/Source/Runtime/TypedElementFramework/Public/Elements/Framework/TypedElementQueryBuilder.inl b/Engine/Source/Runtime/TypedElementFramework/Public/Elements/Framework/TypedElementQueryBuilder.inl
--- a/Engine/Source/Runtime/TypedElementFramework/Public/Elements/Framework/TypedElementQueryBuilder.inl
+++ b/Engine/Source/Runtime/TypedElementFramework/Public/Elements/Framework/TypedElementQueryBuilder.inl
@@ -227,6 +227,11 @@ namespace UE::Editor::DataStorage::Queries
 		return ParentContext.RunSubquery(SubqueryIndex, Row, Callback);
 	}
 
+	void FQueryContextForwarder::PushCommand(void(* CommandFunction)(void*), void* CommandData)
+	{
+    		return ParentContext.PushCommand(CommandFunction, CommandData);
+	}
+
 	//
 	// FCachedQueryContext
 	//
diff --git a/Engine/Source/Runtime/TypedElementFramework/Public/Elements/Interfaces/TypedElementQueryStorageInterfaces.h b/Engine/Source/Runtime/TypedElementFramework/Public/Elements/Interfaces/TypedElementQueryStorageInterfaces.h
--- a/Engine/Source/Runtime/TypedElementFramework/Public/Elements/Interfaces/TypedElementQueryStorageInterfaces.h
+++ b/Engine/Source/Runtime/TypedElementFramework/Public/Elements/Interfaces/TypedElementQueryStorageInterfaces.h
@@ -274,6 +274,38 @@ namespace UE::Editor::DataStorage
 		 * tick group.
 		 */
 		virtual void RemoveColumns(TConstArrayView<RowHandle> Rows, TConstArrayView<const UScriptStruct*> ColumnTypes) = 0;
+		
+		/**
+		 * Creates a command that will run immediately after all processors have completed.
+		 * Intended to be used for cases where the query callback calls into something that cannot run while
+		 * the TEDS processors are running.
+		 * Note that commands will be executed in order with respect to the thread that pushed them and will be executed
+		 * on the game thread.
+		 *
+		 * Usage:
+		 *  Define a command struct with a mutable operator() overload
+		 *  struct FMyCommand
+		 *  {
+		 *      void operator()() { DoSomethingWithSideEffects(MyActor) };
+		 *      TWeakObjectPtr<AActor> MyActor;
+		 *  };
+		 *
+		 *  Context.PushCommand(FMyCommand{ .MyActor = Actor });
+		 */
+		template<typename T>
+		void PushCommand(T CommandContext);
+		
+		virtual void PushCommand(void (*CommandFunction)(void* /*CommandData*/), void* InCommandData) = 0;
+	protected:
+		struct FEmplaceObjectParams
+		{
+			size_t ObjectSize;
+			size_t Alignment;
+			void (*Construct)(void*, void*);
+			void(*Destroy)(void*);
+			void* SourceObject;
+		};
+		virtual void* EmplaceObjectInScratch(const FEmplaceObjectParams& Params) = 0;
 	};
 
 	enum class EDirectQueryExecutionFlags : uint32
@@ -591,4 +623,51 @@ namespace UE::Editor::DataStorage
 	{
 		RemoveColumns(Rows, { Columns::StaticStruct()... });
 	}
+
+	template <typename T>
+	void ICommonQueryWithEnvironmentContext::PushCommand(T CommandContext)
+	{
+		void (*CommandFunction)(void* /*Context*/);
+
+		// If a member operator() is defined
+		CommandFunction = [](void* InInstanceOfT)
+		{
+			T* Instance = static_cast<T*>(InInstanceOfT);
+			Instance->operator()();
+		};
+			
+		if (std::is_empty_v<T>)
+		{
+			PushCommand(CommandFunction, nullptr);
+		}
+		else
+		{
+			FEmplaceObjectParams Params;
+
+			Params.ObjectSize = sizeof(T);
+			Params.Alignment = alignof(T);
+			Params.Construct = [](void* Destination, void* SourceCommandContext)
+			{
+				T& SourceCommand = *static_cast<T*>(SourceCommandContext);
+				new (Destination) T(MoveTemp(SourceCommand));
+			};
+			if (std::is_trivially_destructible_v<T>)
+			{
+				Params.Destroy = nullptr;
+			}
+			else
+			{
+				Params.Destroy = [](void* EmplacedObject)
+				{
+					static_cast<T*>(EmplacedObject)->~T();
+				};
+			}
+				
+			Params.SourceObject = &CommandContext;
+				
+			void* EmplacedCommandContext = EmplaceObjectInScratch(Params);
+			PushCommand(CommandFunction, EmplacedCommandContext);
+		}
+	}
+
 } // namespace UE::Editor::DataStorage
