[TEDS] Fix crash when URC processors disable overlays with TEDS match function override
* Crash occurs because the match function override expects that a query handle will never be invalid since removing the processor would remove the call into the match function.  The issue is that if the processor's subqueries are removed but not the processor itself, then Mass will call the Match function using an invalid query handle.
* This contains another change noticed while investigating where if a query fails to register, we still register the match function.  This can be problematic if we've set up a Processor already and then subsequenctly remove the Query from the list of queries (thus invalidating the handle).  There is more work that can be done here to unregister the processor and generally clean up the state of the system.

diff --git a/Engine/Plugins/Experimental/EditorDataStorage/Source/TedsCore/Private/Queries/TypedElementExtendedQueryStore.cpp b/Engine/Plugins/Experimental/EditorDataStorage/Source/TedsCore/Private/Queries/TypedElementExtendedQueryStore.cpp
--- a/Engine/Plugins/Experimental/EditorDataStorage/Source/TedsCore/Private/Queries/TypedElementExtendedQueryStore.cpp
+++ b/Engine/Plugins/Experimental/EditorDataStorage/Source/TedsCore/Private/Queries/TypedElementExtendedQueryStore.cpp
@@ -6,6 +6,7 @@
 #include "MassProcessor.h"
 #include "TypedElementDatabaseEnvironment.h"
 #include "Commands/EditorDataStorageCommandBuffer.h"
+#include "Elements/Common/TypedElementDataStorageLog.h"
 #include "Misc/OutputDevice.h"
 #include "Processors/TypedElementProcessorAdaptors.h"
 
@@ -49,10 +50,52 @@ namespace UE::Editor::DataStorage
 		FExtendedQueryStore::Handle QueryHandle;
 		FEnvironment& Environment;
 	};
+
+	static void FORCENOINLINE PrintQueryCompositionToLog(const FMassArchetypeCompositionDescriptor& CompositionDescriptor)
+	{
+		TStringBuilder<512> StringBuilder;
+		StringBuilder.Append( TEXT("Composition Descriptor:\nColumns: {"));
+		{
+			int32 FragmentsProcessed = 0;
+			const int32 FragmentCount = CompositionDescriptor.Fragments.CountStoredTypes();
+			for(auto IndexIterator = CompositionDescriptor.Fragments.GetIndexIterator(); IndexIterator; ++IndexIterator)
+			{
+				StringBuilder.Append(CompositionDescriptor.Fragments.GetTypeAtIndex(*IndexIterator)->GetName());
+				if (FragmentsProcessed != FragmentCount - 1)
+				{
+					StringBuilder.Append(TEXT(", "));
+				}
+				++FragmentsProcessed;
+			}
+			StringBuilder.Append(TEXT("}\n"));
+		}
+		{
+			int32 TagsProcessed = 0;
+			const int32 TagsCount = CompositionDescriptor.Tags.CountStoredTypes();
+			StringBuilder.Append(TEXT("Tags: {"));
+			for(auto IndexIterator = CompositionDescriptor.Tags.GetIndexIterator(); IndexIterator; ++IndexIterator)
+			{
+				StringBuilder.Append(CompositionDescriptor.Tags.GetTypeAtIndex(*IndexIterator)->GetName());
+				if (TagsProcessed != TagsCount - 1)
+				{
+					StringBuilder.Append(TEXT(", "));
+				}
+				++TagsProcessed;
+			}
+			StringBuilder.Append(TEXT("}"));
+		}
+		UE_LOG(LogEditorDataStorage, Error, TEXT("%s"), StringBuilder.ToString());
+	}
 	
 	bool FTedsQueryArchetypeMatchOverride::Match(const FMassArchetypeCompositionDescriptor& CompositionDescriptor) const
 	{
 		FExtendedQuery* Query = Environment.GetQueryStore().Get(QueryHandle);
+		if (UNLIKELY(!ensureAlwaysMsgf(Query != nullptr, TEXT("Invalid query discovered in Match function"))))
+		{
+			PrintQueryCompositionToLog(CompositionDescriptor);
+			// Failing this is likely a result of subqueries of a processor being unregistered without the processor being unregistered.
+			return false;
+		}
 		if (Query->Description.Conditions)
 		{
 			TArray<TWeakObjectPtr<const UScriptStruct>, TInlineAllocator<64>> ColumnTypes;
@@ -109,17 +152,6 @@ namespace UE::Editor::DataStorage
 		StoredQuery.Description = MoveTemp(Query);
 
 		FMassEntityQuery& NativeQuery = SetupNativeQuery(StoredQuery.Description, StoredQuery);
-
-		if (Private::bUseTedsMatchFunction)
-		{
-			const FTedsQueryArchetypeMatchOverride Context
-			{
-				.QueryHandle = Result,
-				.Environment = Environment
-			};
-
-			NativeQuery.SetArchetypeMatchOverride(Context);
-		}
 		
 		bool bContinueSetup =               SetupDynamicColumns(StoredQuery.Description, Environment);
 		bContinueSetup = bContinueSetup &&	SetupSelectedColumns(StoredQuery.Description, NativeQuery);
@@ -135,6 +167,19 @@ namespace UE::Editor::DataStorage
 			// This will also make the handle invalid.
 			Queries.Remove(Result);
 		}
+		else
+		{
+			if (Private::bUseTedsMatchFunction)
+			{
+				const FTedsQueryArchetypeMatchOverride Context
+				{
+					.QueryHandle = Result,
+					.Environment = Environment
+				};
+
+				NativeQuery.SetArchetypeMatchOverride(Context);
+			}
+		}
 
 		return Result;
 	}
diff --git a/Engine/Plugins/Experimental/EditorDataStorageFeatures/Source/TedsRevisionControl/Private/Processors/RevisionControlProcessors.cpp b/Engine/Plugins/Experimental/EditorDataStorageFeatures/Source/TedsRevisionControl/Private/Processors/RevisionControlProcessors.cpp
--- a/Engine/Plugins/Experimental/EditorDataStorageFeatures/Source/TedsRevisionControl/Private/Processors/RevisionControlProcessors.cpp
+++ b/Engine/Plugins/Experimental/EditorDataStorageFeatures/Source/TedsRevisionControl/Private/Processors/RevisionControlProcessors.cpp
@@ -290,6 +290,9 @@ void URevisionControlDataStorageFactory::RegisterQueries(IEditorDataStorageProvi
 				DataStorage.UnregisterQuery(SelectionRemoved);
 				SelectionRemoved = InvalidQueryHandle;
 
+				DataStorage.UnregisterQuery(FlushPackageUpdates);
+				FlushPackageUpdates = InvalidQueryHandle;
+
 				RegisterRemoveOverlays(DataStorage);
 			}
 		}
