[Teds] Enable core plugin by default

#rnx

diff --git a/Engine/Plugins/Experimental/EditorDataStorage/EditorDataStorage.uplugin b/Engine/Plugins/Experimental/EditorDataStorage/EditorDataStorage.uplugin
--- a/Engine/Plugins/Experimental/EditorDataStorage/EditorDataStorage.uplugin
+++ b/Engine/Plugins/Experimental/EditorDataStorage/EditorDataStorage.uplugin
@@ -14,6 +14,7 @@
 	"IsBetaVersion": false,
 	"IsExperimentalVersion": true,
 	"Installed": false,
+	"EnabledByDefault": true,
 	"Modules": [
 		{
 			"Name": "TedsCore",
