Fix startup crash when cooking certain games, LoadDefaultEditorModules is skipped for commandlets

diff --git a/Engine/Plugins/Experimental/EditorDataStorage/Source/TedsCore/Private/TypedElementsDataStorage.cpp b/Engine/Plugins/Experimental/EditorDataStorage/Source/TedsCore/Private/TypedElementsDataStorage.cpp
--- a/Engine/Plugins/Experimental/EditorDataStorage/Source/TedsCore/Private/TypedElementsDataStorage.cpp
+++ b/Engine/Plugins/Experimental/EditorDataStorage/Source/TedsCore/Private/TypedElementsDataStorage.cpp
@@ -70,6 +70,9 @@ void FEditorDataStorageModule::StartupModule()
 	// Load the dependent TypedElementFramework module (holding TypedElementRegistry) here so that it is guaranteed to be available in Shutdown
 	// and it is shutdown AFTER FEditorDataStorageModule
 	FModuleManager::Get().LoadModule(TEXT("TypedElementFramework"));
+
+	// Make sure this is loaded in case it got missed due to commandlets
+	FModuleManager::Get().LoadModule(TEXT("MassEntityEditor"));
 	
 	ImpersonateMassTagsAndFragments();
 
