DynamicColumnGenerator: Remove destructor as it causes a double delete in shipping builds when WITH_RELOAD is false

diff --git a/Engine/Plugins/Experimental/EditorDataStorage/Source/TedsCore/Private/DynamicColumnGenerator.cpp b/Engine/Plugins/Experimental/EditorDataStorage/Source/TedsCore/Private/DynamicColumnGenerator.cpp
--- a/Engine/Plugins/Experimental/EditorDataStorage/Source/TedsCore/Private/DynamicColumnGenerator.cpp
+++ b/Engine/Plugins/Experimental/EditorDataStorage/Source/TedsCore/Private/DynamicColumnGenerator.cpp
@@ -9,18 +9,6 @@
 
 namespace UE::Editor::DataStorage
 {
-	FDynamicColumnGenerator::~FDynamicColumnGenerator()
-	{
-		FReload ReloadContext(EActiveReloadType::Reinstancing, TEXT("REINSTANCING"), *GLog);
-		ReloadContext.SetEnableReinstancing(true);
-		ReloadContext.SetSendReloadCompleteNotification(false);
-		
-		for (FGeneratedColumnRecord& GeneratedColumnRecord : GeneratedColumnData)
-		{
-			GeneratedColumnRecord.Type->DeferCppStructOps(GeneratedColumnRecord.AssetPath, nullptr);
-		}
-	}
-
 	FDynamicColumnGeneratorInfo FDynamicColumnGenerator::GenerateColumn(const UScriptStruct& Template, const FName& Identifier)
 	{
 		const FGeneratedColumnKey Key
diff --git a/Engine/Plugins/Experimental/EditorDataStorage/Source/TedsCore/Private/DynamicColumnGenerator.h b/Engine/Plugins/Experimental/EditorDataStorage/Source/TedsCore/Private/DynamicColumnGenerator.h
--- a/Engine/Plugins/Experimental/EditorDataStorage/Source/TedsCore/Private/DynamicColumnGenerator.h
+++ b/Engine/Plugins/Experimental/EditorDataStorage/Source/TedsCore/Private/DynamicColumnGenerator.h
@@ -46,7 +46,6 @@ namespace UE::Editor::DataStorage
 	class FDynamicColumnGenerator
 	{
 	public:
-		~FDynamicColumnGenerator();
 		/**
 		 * Generates a dynamic TEDS column type based on a Template type (if it hasn't been generated before)
 		 */
