DynamicColumnGenerator: Remove DeferCppStructOps call to fix a crash on shutdown

diff --git a/Engine/Plugins/Experimental/EditorDataStorage/Source/TedsCore/Private/DynamicColumnGenerator.cpp b/Engine/Plugins/Experimental/EditorDataStorage/Source/TedsCore/Private/DynamicColumnGenerator.cpp
--- a/Engine/Plugins/Experimental/EditorDataStorage/Source/TedsCore/Private/DynamicColumnGenerator.cpp
+++ b/Engine/Plugins/Experimental/EditorDataStorage/Source/TedsCore/Private/DynamicColumnGenerator.cpp
@@ -67,7 +67,6 @@ namespace UE::Editor::DataStorage
 			// for copy/move.
 			NewScriptStruct->SetSuperStruct(&const_cast<UScriptStruct&>(Template));
 			
-			NewScriptStruct->DeferCppStructOps(AssetPath, Template.GetCppStructOps());
 			NewScriptStruct->Bind();
 			NewScriptStruct->PrepareCppStructOps();
 			NewScriptStruct->StaticLink(true);			
