Teds Alerts:

Alert queries: Change the alert update queires to run PostPhysics instead of FrameEnd so they run before the widget update
Alert Widget: Remove the row reference column from the additional columns list since it is guaraneed in the table viewer and virtualization removes additional columns from the widget row when the alert column is removed from the target row

diff --git a/Engine/Plugins/Experimental/EditorDataStorage/Source/TedsCore/Private/Queries/TypedElementAlertQueries.cpp b/Engine/Plugins/Experimental/EditorDataStorage/Source/TedsCore/Private/Queries/TypedElementAlertQueries.cpp
--- a/Engine/Plugins/Experimental/EditorDataStorage/Source/TedsCore/Private/Queries/TypedElementAlertQueries.cpp
+++ b/Engine/Plugins/Experimental/EditorDataStorage/Source/TedsCore/Private/Queries/TypedElementAlertQueries.cpp
@@ -177,7 +177,7 @@ void UTypedElementAlertQueriesFactory::RegisterParentUpdatesQueries(IEditorDataS
 	DataStorage.RegisterQuery(
 		Select(
 			TEXT("Trigger alert update if alert's parent changed"),
-			FProcessor(EQueryTickPhase::FrameEnd, DataStorage.GetQueryTickGroupName(EQueryTickGroups::Default)),
+			FProcessor(EQueryTickPhase::PostPhysics, DataStorage.GetQueryTickGroupName(EQueryTickGroups::Default)),
 			[](IQueryContext& Context, FTypedElementAlertColumn& Alert, const FTableRowParentColumn& Parent)
 			{
 				if (Alert.CachedParent != Parent.Parent)
@@ -193,7 +193,7 @@ void UTypedElementAlertQueriesFactory::RegisterParentUpdatesQueries(IEditorDataS
 	DataStorage.RegisterQuery(
 		Select(
 			TEXT("Trigger alert update if child alert's parent changed"),
-			FProcessor(EQueryTickPhase::FrameEnd, DataStorage.GetQueryTickGroupName(EQueryTickGroups::Default)),
+			FProcessor(EQueryTickPhase::PostPhysics, DataStorage.GetQueryTickGroupName(EQueryTickGroups::Default)),
 			[](IQueryContext& Context, FTypedElementChildAlertColumn& ChildAlert, const FTableRowParentColumn& Parent)
 			{
 				if (ChildAlert.CachedParent != Parent.Parent)
@@ -214,7 +214,7 @@ void UTypedElementAlertQueriesFactory::RegisterChildAlertUpdatesQueries(IEditorD
 	DataStorage.RegisterQuery(
 		Select(
 			TEXT("Add missing child alerts"),
-			FPhaseAmble(FPhaseAmble::ELocation::Preamble, EQueryTickPhase::FrameEnd)
+			FPhaseAmble(FPhaseAmble::ELocation::Preamble, EQueryTickPhase::PostPhysics)
 				.MakeActivatable(AlertConditionName),
 			[](IQueryContext& Context, RowHandle Row, FTypedElementAlertColumn& Alert, const FTableRowParentColumn& Parent)
 			{
@@ -227,7 +227,7 @@ void UTypedElementAlertQueriesFactory::RegisterChildAlertUpdatesQueries(IEditorD
 	DataStorage.RegisterQuery(
 		Select(
 			TEXT("Clear child alerts"),
-			FProcessor(EQueryTickPhase::FrameEnd, DataStorage.GetQueryTickGroupName(EQueryTickGroups::PreUpdate))
+			FProcessor(EQueryTickPhase::PostPhysics, DataStorage.GetQueryTickGroupName(EQueryTickGroups::PreUpdate))
 				.MakeActivatable(AlertConditionName),
 			[](IQueryContext& Context, RowHandle Row, FTypedElementChildAlertColumn& ChildAlert)
 			{
@@ -238,7 +238,7 @@ void UTypedElementAlertQueriesFactory::RegisterChildAlertUpdatesQueries(IEditorD
 	DataStorage.RegisterQuery(
 		Select(
 			TEXT("Increment child alerts"),
-			FProcessor(EQueryTickPhase::FrameEnd, DataStorage.GetQueryTickGroupName(EQueryTickGroups::Update))
+			FProcessor(EQueryTickPhase::PostPhysics, DataStorage.GetQueryTickGroupName(EQueryTickGroups::Update))
 				.MakeActivatable(AlertConditionName),
 			[](IQueryContext& Context, RowHandle Row, FTypedElementAlertColumn& Alert)
 			{
@@ -251,7 +251,7 @@ void UTypedElementAlertQueriesFactory::RegisterChildAlertUpdatesQueries(IEditorD
 	DataStorage.RegisterQuery(
 		Select(
 			TEXT("Remove unused child alerts"),
-			FProcessor(EQueryTickPhase::FrameEnd, DataStorage.GetQueryTickGroupName(EQueryTickGroups::PostUpdate))
+			FProcessor(EQueryTickPhase::PostPhysics, DataStorage.GetQueryTickGroupName(EQueryTickGroups::PostUpdate))
 				.MakeActivatable(AlertConditionName),
 			[](IQueryContext& Context, RowHandle Row, FTypedElementChildAlertColumn& ChildAlert)
 			{
@@ -263,6 +263,7 @@ void UTypedElementAlertQueriesFactory::RegisterChildAlertUpdatesQueries(IEditorD
 					}
 				}
 				Context.RemoveColumns<FTypedElementChildAlertColumn>(Row);
+				Context.AddColumns<FTypedElementSyncBackToWorldTag>(Row);
 			})
 		.Compile());
 }
diff --git a/Engine/Plugins/Experimental/EditorDataStorageFeatures/Source/TedsOutliner/Private/Widgets/AlertWidget.cpp b/Engine/Plugins/Experimental/EditorDataStorageFeatures/Source/TedsOutliner/Private/Widgets/AlertWidget.cpp
--- a/Engine/Plugins/Experimental/EditorDataStorageFeatures/Source/TedsOutliner/Private/Widgets/AlertWidget.cpp
+++ b/Engine/Plugins/Experimental/EditorDataStorageFeatures/Source/TedsOutliner/Private/Widgets/AlertWidget.cpp
@@ -357,7 +357,7 @@ TSharedPtr<SWidget> FAlertWidgetConstructor::CreateWidget(
 
 TConstArrayView<const UScriptStruct*> FAlertWidgetConstructor::GetAdditionalColumnsList() const
 {
-	static TTypedElementColumnTypeList<FTypedElementRowReferenceColumn, FAlertWidgetTag> Columns;
+	static TTypedElementColumnTypeList<FAlertWidgetTag> Columns;
 	return Columns;
 }
 
@@ -366,7 +366,14 @@ bool FAlertWidgetConstructor::FinalizeWidget(IEditorDataStorageProvider* DataSto
 {
 	using namespace UE::Editor::DataStorage;
 
-	RowHandle TargetRow = DataStorage->GetColumn<FTypedElementRowReferenceColumn>(Row)->Row;
+	const FTypedElementRowReferenceColumn* RowReferenceColumn = DataStorage->GetColumn<FTypedElementRowReferenceColumn>(Row);
+	
+	if (!ensureMsgf(RowReferenceColumn, TEXT("The alert widget requires your widget to have a target row")))
+	{
+		return false;
+	}
+	
+	RowHandle TargetRow = RowReferenceColumn->Row;
 	const FTypedElementAlertColumn* Alert = DataStorage->GetColumn<FTypedElementAlertColumn>(TargetRow);
 	const FTypedElementChildAlertColumn* ChildAlert = DataStorage->GetColumn<FTypedElementChildAlertColumn>(TargetRow);
 	
