TEDS - Fixed incorrect uses of Context.GetRowHandles() with the singular version of the query callback

diff --git a/Engine/Plugins/Experimental/EditorDataStorage/Source/TedsCore/Private/Compatibility/TedsTypedElementBridgeQueries.cpp b/Engine/Plugins/Experimental/EditorDataStorage/Source/TedsCore/Private/Compatibility/TedsTypedElementBridgeQueries.cpp
--- a/Engine/Plugins/Experimental/EditorDataStorage/Source/TedsCore/Private/Compatibility/TedsTypedElementBridgeQueries.cpp
+++ b/Engine/Plugins/Experimental/EditorDataStorage/Source/TedsCore/Private/Compatibility/TedsTypedElementBridgeQueries.cpp
@@ -80,7 +80,7 @@ void UTypedElementBridgeDataStorageFactory::CleanupTypedElementColumns(IEditorDa
 		DataStorage.RunQuery(
 			RemoveTypedElementRowHandleQuery,
 			CreateDirectQueryCallbackBinding(
-				[&Handles](IDirectQueryContext& Context)
+				[&Handles](IDirectQueryContext& Context, const RowHandle*)
 				{
 					Handles.Append(Context.GetRowHandles());
 				}));
diff --git a/Engine/Plugins/Experimental/EditorDataStorage/Source/TedsCore/Private/Compatibility/TypedElementActorViewportProcessors.cpp b/Engine/Plugins/Experimental/EditorDataStorage/Source/TedsCore/Private/Compatibility/TypedElementActorViewportProcessors.cpp
--- a/Engine/Plugins/Experimental/EditorDataStorage/Source/TedsCore/Private/Compatibility/TypedElementActorViewportProcessors.cpp
+++ b/Engine/Plugins/Experimental/EditorDataStorage/Source/TedsCore/Private/Compatibility/TypedElementActorViewportProcessors.cpp
@@ -59,10 +59,11 @@ namespace UE::Editor::DataStorage::Private
 
 					TArray<RowHandle> RowHandles;
 				
-					DataStorage->RunQuery(OverlayQuery, [ColorIndex, &RowHandles](const FQueryDescription&, IDirectQueryContext& Context)
+					DataStorage->RunQuery(OverlayQuery, CreateDirectQueryCallbackBinding(
+						[ColorIndex, &RowHandles](IDirectQueryContext& Context, const RowHandle*)
 					{
-						RowHandles = Context.GetRowHandles();
-					});
+						RowHandles.Append(Context.GetRowHandles());
+					}));
 
 					for (RowHandle Row : RowHandles)
 					{
@@ -111,10 +112,11 @@ namespace UE::Editor::DataStorage::Private
 
 					TArray<RowHandle> RowHandles;
 				
-					DataStorage->RunQuery(OverlayQuery, [Color, &RowHandles](const FQueryDescription&, IDirectQueryContext& Context)
+					DataStorage->RunQuery(OverlayQuery, CreateDirectQueryCallbackBinding(
+						[Color, &RowHandles](IDirectQueryContext& Context, const RowHandle*)
 					{
-						RowHandles = Context.GetRowHandles();
-					});
+						RowHandles.Append(Context.GetRowHandles());
+					}));
 
 					for (RowHandle Row : RowHandles)
 					{
@@ -152,10 +154,11 @@ namespace UE::Editor::DataStorage::Private
 
 					TArray<RowHandle> RowHandles;
 				
-					DataStorage->RunQuery(OverlayQuery, [&RowHandles](const FQueryDescription&, IDirectQueryContext& Context)
+					DataStorage->RunQuery(OverlayQuery, CreateDirectQueryCallbackBinding(
+							[&RowHandles](IDirectQueryContext& Context, const RowHandle*)
 					{
-						RowHandles = Context.GetRowHandles();
-					});
+						RowHandles.Append(Context.GetRowHandles());
+					}));
 
 					for (RowHandle Row : RowHandles)
 					{
diff --git a/Engine/Plugins/Experimental/EditorDataStorageFeatures/Source/TedsAssetData/Private/Debug/TedsAssetDataDebugCommands.cpp b/Engine/Plugins/Experimental/EditorDataStorageFeatures/Source/TedsAssetData/Private/Debug/TedsAssetDataDebugCommands.cpp
--- a/Engine/Plugins/Experimental/EditorDataStorageFeatures/Source/TedsAssetData/Private/Debug/TedsAssetDataDebugCommands.cpp
+++ b/Engine/Plugins/Experimental/EditorDataStorageFeatures/Source/TedsAssetData/Private/Debug/TedsAssetDataDebugCommands.cpp
@@ -265,7 +265,7 @@ namespace UE::Editor::DataStorage::Debug::Private
 				TArray<RowHandle> Rows;
 				
 				DataStorage->RunQuery(AssetQueryHandle,
-					CreateDirectQueryCallbackBinding([&Rows](const IEditorDataStorageProvider::IDirectQueryContext& Context)
+					CreateDirectQueryCallbackBinding([&Rows](const IEditorDataStorageProvider::IDirectQueryContext& Context, const RowHandle*)
 					{
 						Rows.Append(Context.GetRowHandles());
 					}));
diff --git a/Engine/Plugins/Experimental/EditorDataStorageFeatures/Source/TedsOutliner/Private/TedsOutlinerImpl.cpp b/Engine/Plugins/Experimental/EditorDataStorageFeatures/Source/TedsOutliner/Private/TedsOutlinerImpl.cpp
--- a/Engine/Plugins/Experimental/EditorDataStorageFeatures/Source/TedsOutliner/Private/TedsOutlinerImpl.cpp
+++ b/Engine/Plugins/Experimental/EditorDataStorageFeatures/Source/TedsOutliner/Private/TedsOutlinerImpl.cpp
@@ -351,7 +351,7 @@ void FTedsOutlinerImpl::CreateItemsFromQuery(TArray<FSceneOutlinerTreeItemPtr>&
 	TArray<RowHandle> Rows;
 	
 	DirectQueryCallback RowCollector = CreateDirectQueryCallbackBinding(
-		[&Rows](IDirectQueryContext& Context)
+		[&Rows](IDirectQueryContext& Context, const RowHandle*)
 		{
 			TConstArrayView<RowHandle> ContextRows = Context.GetRowHandles();
 			Rows.Append(ContextRows);
@@ -408,11 +408,11 @@ void FTedsOutlinerImpl::CreateChildren(const FSceneOutlinerTreeItemPtr& Item, TA
 	}
 
 
-	TSet<RowHandle> MatchedRowsWithParentColumn;
+	TArray<RowHandle> MatchedRowsWithParentColumn;
 	
 	// Collect all entities that are owned by our entity
 	DirectQueryCallback ChildRowCollector = CreateDirectQueryCallbackBinding(
-	[&MatchedRowsWithParentColumn] (const IDirectQueryContext& Context)
+	[&MatchedRowsWithParentColumn] (const IDirectQueryContext& Context, const RowHandle*)
 	{
 		MatchedRowsWithParentColumn.Append(Context.GetRowHandles());
 	});
