Allowing the TEDS core module the Chaos Visual Debugger depends on to work in UE programs. This is necessary for the standalone version of the Chaos Visual Debugger.

This does not change the fact that TEDS is still editor only. It just allows UE programs that compile against editor code to use it

diff --git a/Engine/Plugins/ChaosVD/ChaosVD.uplugin b/Engine/Plugins/ChaosVD/ChaosVD.uplugin
--- a/Engine/Plugins/ChaosVD/ChaosVD.uplugin
+++ b/Engine/Plugins/ChaosVD/ChaosVD.uplugin
@@ -33,6 +33,10 @@
 		}
 	],
 	"Plugins": [
+		{
+			"Name": "EditorDataStorage",
+			"Enabled": true
+		},
 		{
 			"Name": "EditorDataStorageFeatures",
 			"Enabled": true
diff --git a/Engine/Plugins/Experimental/EditorDataStorage/EditorDataStorage.uplugin b/Engine/Plugins/Experimental/EditorDataStorage/EditorDataStorage.uplugin
--- a/Engine/Plugins/Experimental/EditorDataStorage/EditorDataStorage.uplugin
+++ b/Engine/Plugins/Experimental/EditorDataStorage/EditorDataStorage.uplugin
@@ -18,7 +18,7 @@
 	"Modules": [
 		{
 			"Name": "TedsCore",
-			"Type": "Editor"
+			"Type": "EditorAndProgram"
 		},
 		{
 			"Name": "TedsUI",
@@ -26,5 +26,6 @@
 		}
 	],
 	"Plugins": [
-	]
+	],
+	"SupportedPrograms": ["ChaosVisualDebugger"]
 }
\ No newline at end of file
