Added option to batch delete rows from TEDS and made the Index Table use garbage collection.

Options to batch delete rows from TEDS have been added. The first option deletes all in the provided list of rows. The second version removes all rows that match a set of provided column for better performance when large tables are removed. It does this be avoiding the overhead of first compiling a list of entities to remove, but the downside is that these were needed by the Index Table to clean itself properly. To avoid adding this cost, the Index Table has been updated to not immediately check to see if a deleted row needs to be removed from itself, but instead uses a garbage collection like approach by first going over all its stored rows to see if any have been removed and then clears out the ones that have been removed. This process is spread out over multiple threads and multiple frames and targets no more than 2ms per frame by selfbalancing job counts per frame and rows per batch. This also only happens after at least one row has been deleted so if a users isn't interacting with the editor typically no work is done. The extra benefit is that the Index Table now doesn't add any cost during deletion, where previously even rows that didn't have an index would pay the cost for checking and updating the Index Table's internal bookkeeping. The downside is that while there's pending clean up work, resolving a hash to a row handle is more expensive and periodically 2ms per frame is used up for the cleanup cycle.

Additional changes made:
- Reindexing a row in the TEDS Index Table no longer needs the row itself.
- A utils header was introduced. It currently provides functions that quickly convert to/from Mass entity (arrays).
- Unsafe versions of IsRowAvailable and IsRowAssigned have been added that omit checking if the database has been initialized. This is internal to the TEDS plugin and can be used in most places that get created during initialization to safe on a bit of performance cost checking the state.
- Marked UEditorDataStorage as "final" to give the compiler a few more optimization opportunities whenever a pointer of reference to UEditorDataStorage is directly called.

[RN] Added option to batch delete rows from TEDS and made the Index Table use garbage collection.

diff --git a/Engine/Plugins/Experimental/EditorDataStorage/Source/TedsCore/Private/Processors/TypedElementProcessorAdaptors.cpp b/Engine/Plugins/Experimental/EditorDataStorage/Source/TedsCore/Private/Processors/TypedElementProcessorAdaptors.cpp
--- a/Engine/Plugins/Experimental/EditorDataStorage/Source/TedsCore/Private/Processors/TypedElementProcessorAdaptors.cpp
+++ b/Engine/Plugins/Experimental/EditorDataStorage/Source/TedsCore/Private/Processors/TypedElementProcessorAdaptors.cpp
@@ -3,13 +3,14 @@
 #include "Processors/TypedElementProcessorAdaptors.h"
 
 #include <utility>
+#include "Algo/Copy.h"
 #include "Elements/Common/TypedElementQueryTypes.h"
 #include "MassEntityView.h"
 #include "MassExecutionContext.h"
 #include "Queries/TypedElementExtendedQueryStore.h"
 #include "TypedElementDatabase.h"
 #include "TypedElementDatabaseEnvironment.h"
-#include "Algo/Copy.h"
+#include "TypedElementUtils.h"
 
 namespace UE::Editor::DataStorage
 {
@@ -32,11 +33,7 @@ namespace UE::Editor::DataStorage
 
 			TConstArrayView<RowHandle> GetRowHandles() const
 			{
-				static_assert(
-					sizeof(RowHandle) == sizeof(FMassEntityHandle) && alignof(RowHandle) == alignof(FMassEntityHandle),
-					"RowHandle and FMassEntityHandle must be layout compatible.");
-				TConstArrayView<FMassEntityHandle> Entities = Context.GetEntities();
-				return TConstArrayView<RowHandle>(reinterpret_cast<const RowHandle*>(Entities.GetData()), Entities.Num());
+				return MassEntitiesToRowsConversion(Context.GetEntities());
 			}
 	
 			const void* GetColumn(const UScriptStruct* ColumnType) const
@@ -601,16 +598,7 @@ namespace UE::Editor::DataStorage
 
 			void RemoveRows(TConstArrayView<RowHandle> Rows) 
 			{
-				// Row handles and entities map 1:1 for data, so a reintpret_cast can be safely done to avoid
-				// having to allocate memory and iterating over the rows.
-
-				static_assert(sizeof(FMassEntityHandle) == sizeof(RowHandle), 
-					"Sizes of mass entity and data storage row have gone out of sync.");
-				static_assert(alignof(FMassEntityHandle) == alignof(RowHandle),
-					"Alignment of mass entity and data storage row have gone out of sync.");
-
-				this->Context.Defer().DestroyEntities(
-					TConstArrayView<FMassEntityHandle>(reinterpret_cast<const FMassEntityHandle*>(Rows.begin()), Rows.Num()));
+				this->Context.Defer().DestroyEntities(RowsToMassEntitiesConversion(Rows));
 			}
 
 			void PushCommand(void (*CommandFunction)(void*), void* InCommandData)
diff --git a/Engine/Plugins/Experimental/EditorDataStorage/Source/TedsCore/Private/TypedElementDatabase.cpp b/Engine/Plugins/Experimental/EditorDataStorage/Source/TedsCore/Private/TypedElementDatabase.cpp
--- a/Engine/Plugins/Experimental/EditorDataStorage/Source/TedsCore/Private/TypedElementDatabase.cpp
+++ b/Engine/Plugins/Experimental/EditorDataStorage/Source/TedsCore/Private/TypedElementDatabase.cpp
@@ -18,6 +18,7 @@
 #include "Stats/Stats.h"
 #include "TickTaskManagerInterface.h"
 #include "TypedElementDatabaseEnvironment.h"
+#include "TypedElementUtils.h"
 #include "UObject/UObjectIterator.h"
 
 const FName UEditorDataStorage::TickGroupName_Default(TEXT("Default"));
@@ -335,15 +336,11 @@ void UEditorDataStorage::BatchReserveRows(int32 Count, TFunctionRef<void(RowHand
 
 void UEditorDataStorage::BatchReserveRows(TArrayView<RowHandle> ReservedRows)
 {
+	using namespace UE::Editor::DataStorage;
+
 	if (ActiveEditorEntityManager)
 	{
-		// Depend on the fact that a row handle is an alias for an entity within the Mass powered backend. This
-		// avoids the need for copying to a temporary array;
-		static_assert(sizeof(RowHandle) == sizeof(FMassEntityHandle),
-			"BatchReserveRows in TEDS requires the row handle and the Mass entity handle to be the same size.");
-
-		TArrayView<FMassEntityHandle>* MassHandles = reinterpret_cast<TArrayView<FMassEntityHandle>*>(&ReservedRows);
-		ActiveEditorEntityManager->BatchReserveEntities(*MassHandles);
+		ActiveEditorEntityManager->BatchReserveEntities(RowsToMassEntitiesConversion(ReservedRows));
 	}
 }
 
@@ -437,20 +434,15 @@ bool UEditorDataStorage::BatchAddRow(TableHandle Table,
 	TConstArrayView<RowHandle> ReservedHandles, RowCreationCallbackRef OnCreated)
 {
 	using namespace UE::Editor;
+	using namespace UE::Editor::DataStorage;
+
 	OnCreated.CheckCallable();
 	if (ActiveEditorEntityManager)
 	{
 		checkf(Table < Tables.Num(), TEXT("Attempting to add multiple rows to a non-existing table."));
 	
-		// Depend on the fact that a row handle is an alias for an entity within the Mass powered backend. This
-		// avoids the need for copying to a temporary array;
-		static_assert(sizeof(RowHandle) == sizeof(FMassEntityHandle), 
-			"BatchAddRow in TEDS requires the row handle and the Mass entity handle to be the same size.");
-		
-		TConstArrayView<FMassEntityHandle> Entities(
-			reinterpret_cast<const FMassEntityHandle*>(ReservedHandles.GetData()), ReservedHandles.Num());
 		TSharedRef<FMassEntityManager::FEntityCreationContext> Context =
-			ActiveEditorEntityManager->BatchCreateReservedEntities(Tables[Table], Entities);
+			ActiveEditorEntityManager->BatchCreateReservedEntities(Tables[Table], RowsToMassEntitiesConversion(ReservedHandles));
 
 		for (RowHandle Entity : ReservedHandles)
 		{
@@ -468,7 +460,6 @@ void UEditorDataStorage::RemoveRow(RowHandle Row)
 	FMassEntityHandle Entity = FMassEntityHandle::FromNumber(Row);
 	if (ActiveEditorEntityManager && ActiveEditorEntityManager->IsEntityValid(Entity))
 	{
-		Environment->GetIndexTable().RemoveRow(UE::Editor::DataStorage::EGlobalLockScope::Public, Row);
 		if (ActiveEditorEntityManager->IsEntityBuilt(FMassEntityHandle::FromNumber(Row)))
 		{
 			ActiveEditorEntityManager->DestroyEntity(FMassEntityHandle::FromNumber(Row));
@@ -478,6 +469,59 @@ void UEditorDataStorage::RemoveRow(RowHandle Row)
 			Environment->GetDirectDeferredCommands().Clear(Row);
 			ActiveEditorEntityManager->ReleaseReservedEntity(FMassEntityHandle::FromNumber(Row));
 		}
+		Environment->GetIndexTable().MarkDirty();
+	}
+}
+
+void UEditorDataStorage::BatchRemoveRows(TConstArrayView<RowHandle> Rows)
+{
+	using namespace UE::Editor::DataStorage;
+
+	if (ActiveEditorEntityManager)
+	{
+		ActiveEditorEntityManager->BatchDestroyEntities(RowsToMassEntitiesConversion(Rows));
+		Environment->GetIndexTable().MarkDirty();
+	}
+}
+
+void UEditorDataStorage::RemoveAllRowsWithColumns(TConstArrayView<UScriptStruct*> Columns)
+{
+	if (ActiveEditorEntityManager)
+	{
+		FMassFragmentRequirements Requirements;
+		Requirements.Initialize(ActiveEditorEntityManager.ToSharedRef());
+		for (const UScriptStruct* Column : Columns)
+		{
+			if (Column)
+			{
+				if (Column->IsChildOf(FEditorDataStorageTag::StaticStruct()))
+				{
+					Requirements.AddTagRequirement(*Column, EMassFragmentPresence::All);
+				}
+				else
+				{
+					Requirements.AddRequirement(Column, EMassFragmentAccess::None, EMassFragmentPresence::All);
+				}
+			}
+		}
+
+		TArray<FMassArchetypeHandle> MatchingArchetypes;
+		ActiveEditorEntityManager->GetMatchingArchetypes(Requirements, MatchingArchetypes);
+		
+		if (!MatchingArchetypes.IsEmpty())
+		{
+			TArray<FMassArchetypeEntityCollection> Collections;
+			Collections.Reserve(MatchingArchetypes.Num());
+			for (FMassArchetypeHandle& Archetype : MatchingArchetypes)
+			{
+				FMassArchetypeEntityCollection Collection(Archetype);
+				Collections.Add(MoveTemp(Collection));
+			}
+
+			ActiveEditorEntityManager->BatchDestroyEntityChunks(Collections);
+
+			Environment->GetIndexTable().MarkDirty();
+		}
 	}
 }
 
@@ -491,6 +535,16 @@ bool UEditorDataStorage::IsRowAssigned(RowHandle Row) const
 	return ActiveEditorEntityManager ? UE::Editor::DataStorage::Legacy::FCommandBuffer::Execute_IsRowAssigned(*ActiveEditorEntityManager, Row) : false;
 }
 
+bool UEditorDataStorage::IsRowAvailableUnsafe(RowHandle Row) const
+{
+	return UE::Editor::DataStorage::Legacy::FCommandBuffer::Execute_IsRowAvailable(*ActiveEditorEntityManager, Row);
+}
+
+bool UEditorDataStorage::IsRowAssignedUnsafe(RowHandle Row) const
+{
+	return UE::Editor::DataStorage::Legacy::FCommandBuffer::Execute_IsRowAssigned(*ActiveEditorEntityManager, Row);
+}
+
 void UEditorDataStorage::AddColumn(RowHandle Row, const UScriptStruct* ColumnType)
 {
 	if (ColumnType && ActiveEditorEntityManager)
@@ -1102,10 +1156,9 @@ void UEditorDataStorage::BatchIndexRows(
 	Environment->GetIndexTable().BatchIndexRows(UE::Editor::DataStorage::EGlobalLockScope::Public, IndexRowPairs);
 }
 
-void UEditorDataStorage::ReindexRow(UE::Editor::DataStorage::IndexHash OriginalIndex, UE::Editor::DataStorage::IndexHash NewIndex, 
-	RowHandle RowHandle)
+void UEditorDataStorage::Reindex(UE::Editor::DataStorage::IndexHash OriginalIndex, UE::Editor::DataStorage::IndexHash NewIndex)
 {
-	Environment->GetIndexTable().ReindexRow(UE::Editor::DataStorage::EGlobalLockScope::Public, OriginalIndex, NewIndex, RowHandle);
+	Environment->GetIndexTable().Reindex(UE::Editor::DataStorage::EGlobalLockScope::Public, OriginalIndex, NewIndex);
 }
 
 void UEditorDataStorage::RemoveIndex(UE::Editor::DataStorage::IndexHash Index)
diff --git a/Engine/Plugins/Experimental/EditorDataStorage/Source/TedsCore/Private/TypedElementDatabase.h b/Engine/Plugins/Experimental/EditorDataStorage/Source/TedsCore/Private/TypedElementDatabase.h
--- a/Engine/Plugins/Experimental/EditorDataStorage/Source/TedsCore/Private/TypedElementDatabase.h
+++ b/Engine/Plugins/Experimental/EditorDataStorage/Source/TedsCore/Private/TypedElementDatabase.h
@@ -25,7 +25,7 @@ namespace UE::Editor::DataStorage
 } // namespace UE::Editor::DataStorage
 
 UCLASS()
-class TEDSCORE_API UEditorDataStorage
+class TEDSCORE_API UEditorDataStorage final
 	: public UObject
 	, public IEditorDataStorageProvider
 {
@@ -112,8 +112,14 @@ public:
 	virtual bool BatchAddRow(TableHandle Table, TConstArrayView<RowHandle> ReservedHandles,
 		RowCreationCallbackRef OnCreated) override;
 	virtual void RemoveRow(RowHandle Row) override;
+	virtual void BatchRemoveRows(TConstArrayView<RowHandle> Rows) override;
+	virtual void RemoveAllRowsWithColumns(TConstArrayView<UScriptStruct*> Columns) override;
 	virtual bool IsRowAvailable(RowHandle Row) const override;
 	virtual bool IsRowAssigned(RowHandle Row) const override;
+	/** Same as IsRowAvailable, but doesn't check if the data storage has been initialized. */
+	bool IsRowAvailableUnsafe(RowHandle Row) const;
+	/** Same as IsRowAssigned, but doesn't check if the data storage has been initialized. */
+	bool IsRowAssignedUnsafe(RowHandle Row) const;
 
 	virtual void AddColumn(RowHandle Row, const UScriptStruct* ColumnType) override;
 	virtual void AddColumn(RowHandle Row, const UE::Editor::DataStorage::FValueTag& Tag, const FName& InValue) override;
@@ -156,10 +162,9 @@ public:
 	virtual void IndexRow(UE::Editor::DataStorage::IndexHash Index, RowHandle Row) override;
 	virtual void BatchIndexRows(
 		TConstArrayView<TPair<UE::Editor::DataStorage::IndexHash, RowHandle>> IndexRowPairs) override;
-	virtual void ReindexRow(
+	virtual void Reindex(
 		UE::Editor::DataStorage::IndexHash OriginalIndex, 
-		UE::Editor::DataStorage::IndexHash NewIndex, 
-		RowHandle Row) override;
+		UE::Editor::DataStorage::IndexHash NewIndex) override;
 	virtual void RemoveIndex(UE::Editor::DataStorage::IndexHash Index) override;
 
 	virtual FTypedElementOnDataStorageUpdate& OnUpdate() override;
diff --git a/Engine/Plugins/Experimental/EditorDataStorage/Source/TedsCore/Private/TypedElementDatabaseEnvironment.cpp b/Engine/Plugins/Experimental/EditorDataStorage/Source/TedsCore/Private/TypedElementDatabaseEnvironment.cpp
--- a/Engine/Plugins/Experimental/EditorDataStorage/Source/TedsCore/Private/TypedElementDatabaseEnvironment.cpp
+++ b/Engine/Plugins/Experimental/EditorDataStorage/Source/TedsCore/Private/TypedElementDatabaseEnvironment.cpp
@@ -9,6 +9,7 @@ namespace UE::Editor::DataStorage
 		FMassEntityManager& InMassEntityManager, FMassProcessingPhaseManager& InMassPhaseManager)
 		: DataStorage(InDataStorage)
 		, DirectDeferredCommands(*this)
+		, IndexTable(InDataStorage)
 		, MementoSystem(InDataStorage)
 		, ValueTagManager(DynamicColumnGenerator)
 		, MassEntityManager(InMassEntityManager)
@@ -116,6 +117,8 @@ namespace UE::Editor::DataStorage
 		Queries.UpdateActivatableQueries();
 		FlushCommands();
 		ScratchBuffer.BatchDelete();
+		IndexTable.RemoveInvalidRows();
+
 		UpdateCycleId++;
 	}
 
diff --git a/Engine/Plugins/Experimental/EditorDataStorage/Source/TedsCore/Private/TypedElementDatabaseIndexTable.cpp b/Engine/Plugins/Experimental/EditorDataStorage/Source/TedsCore/Private/TypedElementDatabaseIndexTable.cpp
--- a/Engine/Plugins/Experimental/EditorDataStorage/Source/TedsCore/Private/TypedElementDatabaseIndexTable.cpp
+++ b/Engine/Plugins/Experimental/EditorDataStorage/Source/TedsCore/Private/TypedElementDatabaseIndexTable.cpp
@@ -2,15 +2,41 @@
 
 #include "TypedElementDatabaseIndexTable.h"
 
+#include "Async/ParallelFor.h"
+#include "Elements/Common/TypedElementDataStorageLog.h"
+#include "GenericPlatform/GenericPlatformTime.h"
+#include "TypedElementDatabase.h"
+#include "TypedElementDataStorageProfilingMacros.h"
+
 namespace UE::Editor::DataStorage
 {
+	const FTimespan FIndexTable::FCleanUpInfo::TargetDuration = FTimespan::FromMicroseconds(2000.0);
+	const FTimespan FIndexTable::FCleanUpInfo::JobShrinkThreshold = FTimespan::FromMicroseconds(100);
+	const FTimespan FIndexTable::FCleanUpInfo::JobGrowthThreshold = FTimespan::FromMicroseconds(500);
+
+	const FTimespan FIndexTable::FCleanUpInfo::MaxBatchDuration = FTimespan::FromMicroseconds(500);
+	const FTimespan FIndexTable::FCleanUpInfo::MinBatchDuration = FTimespan::FromMicroseconds(350);
+
+	FIndexTable::FIndexTable(UEditorDataStorage& DataStorage)
+		: DataStorage(DataStorage)
+	{
+	}
+
 	RowHandle FIndexTable::FindIndexedRow(
 		EGlobalLockScope LockScope, IndexHash Index) const
 	{	
 		FScopedSharedLock Lock(LockScope);
 	
-		const RowHandle* Result = IndexLookupMap.Find(Index);
-		return Result ? *Result : InvalidRowHandle;
+		RowHandle Result = InvalidRowHandle;
+		if (const int32* RowIndex = IndexLookupMap.Find(Index))
+		{
+			Result = Rows[*RowIndex];
+			if (IsDirty() && !DataStorage.IsRowAvailableUnsafe(Result))
+			{
+				Result = InvalidRowHandle;
+			}
+		}
+		return Result;
 	}
 
 	void FIndexTable::BatchIndexRows(EGlobalLockScope LockScope,
@@ -19,7 +45,12 @@ namespace UE::Editor::DataStorage
 		FScopedExclusiveLock Lock(LockScope);
 	
 		IndexLookupMap.Reserve(IndexLookupMap.Num() + IndexRowPairs.Num());
-		ReverseIndexLookupMap.Reserve(ReverseIndexLookupMap.Num() + IndexRowPairs.Num());
+		
+		checkf(Rows.Num() >= FreeList.Num(), 
+			TEXT("There can't be less rows than there are rows stored in the free list as free list is a subset of rows."));
+		int32 ArraySize = Rows.Num() - FreeList.Num() + IndexRowPairs.Num();
+		Rows.Reserve(ArraySize);
+		Hashes.Reserve(ArraySize);
 
 		for (const TPair<IndexHash, RowHandle>& IndexAndRow : IndexRowPairs)
 		{
@@ -35,50 +66,230 @@ namespace UE::Editor::DataStorage
 		IndexRowUnguarded(Index, Row);
 	}
 
-	void FIndexTable::ReindexRow(EGlobalLockScope LockScope,
-		IndexHash OriginalIndex, IndexHash NewIndex, 
-		RowHandle Row)
+	void FIndexTable::Reindex(EGlobalLockScope LockScope,
+		IndexHash OriginalIndex, IndexHash NewIndex)
 	{
 		FScopedExclusiveLock Lock(LockScope);
 	
-		RemoveIndexUnguarded(OriginalIndex);
-		IndexRowUnguarded(NewIndex, Row);
+		int32 RowIndex;
+		if (IndexLookupMap.RemoveAndCopyValue(OriginalIndex, RowIndex))
+		{
+			IndexLookupMap.Add(NewIndex, RowIndex);
+		}
 	}
 
 	void FIndexTable::RemoveIndex(EGlobalLockScope LockScope, IndexHash Index)
 	{
 		FScopedExclusiveLock Lock(LockScope);
-		RemoveIndexUnguarded(Index);
+		
+		int32 RowIndex;
+		if (IndexLookupMap.RemoveAndCopyValue(Index, RowIndex))
+		{
+			FreeList.PushFirst(RowIndex);
+		}
 	}
 
-	void FIndexTable::RemoveRow(EGlobalLockScope LockScope, RowHandle Row)
+	void FIndexTable::MarkDirty()
 	{
-		FScopedExclusiveLock Lock(LockScope);
-	
-		if (TMultiMap<RowHandle, IndexHash>::TKeyIterator It = ReverseIndexLookupMap.CreateKeyIterator(Row); It)
+		DirtyDueToRemoval = true;
+	}
+
+	void FIndexTable::RemoveInvalidRows()
+	{
+		TEDS_EVENT_SCOPE("Index Table clean up");
+
+		FScopedExclusiveLock Lock(EGlobalLockScope::Public);
+
+		FTimespan RemainingTime = FCleanUpInfo::TargetDuration;
+		uint64 StartTime = FPlatformTime::Cycles64();
+
+		// If there's no work from a previous pass, then start a new pass if needed.
+		if (CleanUpInfo.RemainingJobs == 0 && CleanUpInfo.DeletionQueue.IsEmpty())
+		{
+			if (DirtyDueToRemoval.load())
+			{
+				DirtyDueToRemoval = false;
+				
+				CleanUpInfo.JobCount = (Rows.Num() / CleanUpInfo.BatchSize) + 1;
+				CleanUpInfo.RemainingJobs = CleanUpInfo.JobCount;
+			}
+			else
+			{
+				return;
+			}
+		}
+
+		// If there are still batches left from the previous frame or a new pass was just started, then collect invalid rows.
+		if (CleanUpInfo.RemainingJobs > 0)
 		{
-			do
+			int32 NumBatches = FMath::Min(CleanUpInfo.RemainingJobs, CleanUpInfo.MaxNumJobs);
+			ParallelForTemplate(NumBatches, [this](int32 Block)
+				{
+					InspectRowBlockForCleanUp(Block);
+				}, EParallelForFlags::Unbalanced);
+
+			CleanUpInfo.RemainingJobs -= NumBatches;
+
+			// Follow up with any adjustments needed.
+			RemainingTime -= FTimespan(FPlatformTime::Cycles64() - StartTime);
+			if (NumBatches == CleanUpInfo.MaxNumJobs) // Only adjust if a full set of batches was used to avoid skewing.
 			{
-				IndexLookupMap.Remove(It.Value());
-				++It;
-			} while (It);
-			ReverseIndexLookupMap.Remove(Row);
+				AdjustJobCount(RemainingTime);
+			}
+			if (CleanUpInfo.RemainingJobs == 0) // Don't adjust the batches when there are more to be run.
+			{
+				AdjustBatchSize(CleanUpInfo.JobCount);
+			}
+		}
+
+		// If there's time left and there are rows to remove, start removing rows.
+		if (CleanUpInfo.RemainingJobs == 0 && RemainingTime > FTimespan::Zero() && !CleanUpInfo.DeletionQueue.IsEmpty())
+		{
+			DrainDeletionQueue(RemainingTime);
 		}
+		
+		UE_LOG(LogEditorDataStorage, Verbose,
+			TEXT("TEDS Index Table cleanup - %7.2fms - Has%sremaining rows, Batch size: %i, Job count: %i, Remaining jobs: %i, Max batches: %i"),
+			FTimespan(FPlatformTime::Cycles64() - StartTime).GetTotalMilliseconds(),
+			CleanUpInfo.DeletionQueue.IsEmpty() ? TEXT(" no ") : TEXT(" "),
+			CleanUpInfo.BatchSize,
+			CleanUpInfo.JobCount,
+			CleanUpInfo.RemainingJobs,
+			CleanUpInfo.MaxNumJobs);
+	}
+
+	bool FIndexTable::IsDirty() const
+	{
+		return CleanUpInfo.RemainingJobs > 0 || DirtyDueToRemoval || !CleanUpInfo.DeletionQueue.IsEmpty();
 	}
 
 	void FIndexTable::IndexRowUnguarded(IndexHash Index, RowHandle Row)
 	{
-		IndexLookupMap.Add(Index, Row);
-		ReverseIndexLookupMap.Add(Row, Index);
+		int32 RowIndex;
+		if (FreeList.IsEmpty())
+		{
+			RowIndex = Rows.Add(Row);
+			Hashes.Add(Index);
+		}
+		else
+		{
+			RowIndex = FreeList.Last();
+			FreeList.PopLast();
+			Rows[RowIndex] = Row;
+			Hashes[RowIndex] = Index;
+		}
+		
+		IndexLookupMap.Add(Index, RowIndex);
 	}
 
+	void FIndexTable::InspectRowBlockForCleanUp(int32 Block)
+	{
+		TEDS_EVENT_SCOPE("Index Table inspect rows");
 
-	void FIndexTable::RemoveIndexUnguarded(IndexHash Index)
+		int32 Index = ((CleanUpInfo.JobCount - CleanUpInfo.RemainingJobs + Block) * CleanUpInfo.BatchSize);
+
+		const RowHandle* Begin = Rows.GetData();
+		const RowHandle* Front = Rows.GetData() + Index;
+		const RowHandle* End = FMath::Min(Begin + Rows.Num(), Front + CleanUpInfo.BatchSize);
+		bool bIsFullBatch = (End - Front) == CleanUpInfo.BatchSize;
+
+		uint64 StartTime = FPlatformTime::Cycles64();
+		
+		FCleanUpInfo::InvalidRowContainer Result;
+		for (; Front != End; ++Front)
+		{
+			// Skip rows with an invalid handle as those are already freed.
+			if (*Front != InvalidRowHandle && !DataStorage.IsRowAvailableUnsafe(*Front))
+			{
+				Result.Add(Front - Begin);
+			}
+		}
+
+		if (!Result.IsEmpty())
+		{
+			CleanUpInfo.DeletionQueue.Enqueue(MoveTemp(Result));
+		}
+
+		// Only adjust if this is a full batch, otherwise partial batches skew performance stats.
+		if (bIsFullBatch)
+		{
+			FTimespan Duration(FPlatformTime::Cycles64() - StartTime);
+			if (Duration >= FCleanUpInfo::MaxBatchDuration)
+			{
+				++CleanUpInfo.BatchWentOverTime;
+			}
+			else if (Duration <= FCleanUpInfo::MinBatchDuration)
+			{
+				++CleanUpInfo.BatchWentUnderTime;
+			}
+		}
+	}
+
+	void FIndexTable::DrainDeletionQueue(FTimespan RemainingFrameTime)
 	{
-		if (const RowHandle* Row = IndexLookupMap.Find(Index))
+		TEDS_EVENT_SCOPE("Index Table drain deletion queue");
+
+		uint64 StartTime = FPlatformTime::Cycles64();
+
+		while (!CleanUpInfo.DeletionQueue.IsEmpty() && FTimespan(FPlatformTime::Cycles64() - StartTime) < RemainingFrameTime)
+		{
+			TOptional<FCleanUpInfo::InvalidRowContainer> OptionalContainer = CleanUpInfo.DeletionQueue.Dequeue();
+			checkf(OptionalContainer.IsSet(),
+				TEXT("Retrieved an invalid row container from a non-empty deletion queue, but no result was returned"));
+			
+			FCleanUpInfo::InvalidRowContainer& Container = OptionalContainer.GetValue();
+			for (int32 RowIndex : Container)
+			{
+				ClearRow(RowIndex);
+			}
+		}
+	}
+
+	void FIndexTable::ClearRow(int32 Index)
+	{
+		IndexHash Hash = Hashes[Index];
+
+		int32 StoredIndex;
+		bool Result = IndexLookupMap.RemoveAndCopyValue(Hash, StoredIndex);
+		checkf(Result && StoredIndex == Index, TEXT("Attempting to remove an invalid row but it doesn't exist in the lookup map."));
+
+		Rows[Index] = InvalidRowHandle;
+		Hashes[Index] = 0;
+
+		FreeList.PushFirst(Index);
+	}
+
+	void FIndexTable::AdjustJobCount(FTimespan RemainingFrameTime)
+	{
+		if (RemainingFrameTime >= FCleanUpInfo::JobGrowthThreshold)
+		{
+			CleanUpInfo.MaxNumJobs *= FCleanUpInfo::JobGrowthFactor;
+		}
+		else if (RemainingFrameTime <= FCleanUpInfo::JobShrinkThreshold)
+		{
+			CleanUpInfo.MaxNumJobs *= FCleanUpInfo::JobShrinkFactor;
+		}
+
+		CleanUpInfo.MaxNumJobs = FMath::Clamp(CleanUpInfo.MaxNumJobs, FCleanUpInfo::MinJobCount, FCleanUpInfo::MaxJobCount);
+	}
+
+	void FIndexTable::AdjustBatchSize(int32 JobCount)
+	{
+		if (CleanUpInfo.BatchWentOverTime.load() >= FCleanUpInfo::BatchShrinkThreshold)
+		{
+			// The batch size was too big to complete in the allotted time so take a sizable chunk off.
+			CleanUpInfo.BatchSize *= FCleanUpInfo::BatchShrinkFactor;
+		}
+		// Check if more than the minimum percentage of jobs took less time than currently allocated.
+		else if (CleanUpInfo.BatchWentUnderTime.load() >= 
+			static_cast<uint32>(FMath::RoundToInt32(JobCount * FCleanUpInfo::BatchIncreaseThreshold)))
 		{
-			IndexLookupMap.Remove(Index);
-			ReverseIndexLookupMap.Remove(*Row, Index);
+			// The batch size fitted within the allotted time so slightly increase it.
+			CleanUpInfo.BatchSize *= FCleanUpInfo::BatchGrowthFactor;
 		}
+		// Clamp the job batch size within reasonable sizes to avoid extremes.
+		CleanUpInfo.BatchSize = FMath::Clamp(CleanUpInfo.BatchSize, FCleanUpInfo::MinBatchSize, FCleanUpInfo::MaxBatchSize);
+		CleanUpInfo.BatchWentOverTime = 0;
+		CleanUpInfo.BatchWentUnderTime = 0;
 	}
 } // namespace UE::Editor::DataStorage
diff --git a/Engine/Plugins/Experimental/EditorDataStorage/Source/TedsCore/Private/TypedElementDatabaseIndexTable.h b/Engine/Plugins/Experimental/EditorDataStorage/Source/TedsCore/Private/TypedElementDatabaseIndexTable.h
--- a/Engine/Plugins/Experimental/EditorDataStorage/Source/TedsCore/Private/TypedElementDatabaseIndexTable.h
+++ b/Engine/Plugins/Experimental/EditorDataStorage/Source/TedsCore/Private/TypedElementDatabaseIndexTable.h
@@ -2,11 +2,17 @@
 
 #pragma once
 
+#include <atomic>
+#include "Containers/Array.h"
 #include "Containers/Map.h"
+#include "Containers/Deque.h"
+#include "Containers/MpscQueue.h"
 #include "Elements/Common/TypedElementHandles.h"
 #include "Elements/Interfaces/TypedElementQueryStorageInterfaces.h"
 #include "GlobalLock.h"
+#include "Misc/Timespan.h"
 
+class UEditorDataStorage;
 
 namespace UE::Editor::DataStorage
 {
@@ -18,33 +24,96 @@ namespace UE::Editor::DataStorage
 	class FIndexTable final
 	{
 	public:
-		RowHandle FindIndexedRow(
-			EGlobalLockScope LockScope,
-			IndexHash Index) const;
-		void IndexRow(
-			EGlobalLockScope LockScope,
-			IndexHash Index, 
-			RowHandle Row);
-		void BatchIndexRows(
-			EGlobalLockScope LockScope,
-			TConstArrayView<TPair<IndexHash, RowHandle>> IndexRowPairs);
-		void ReindexRow(
-			EGlobalLockScope LockScope,
-			IndexHash OriginalIndex,
-			IndexHash NewIndex,
-			RowHandle Row);
-		void RemoveIndex(
-			EGlobalLockScope LockScope,
-			IndexHash Index);
-		void RemoveRow(
-			EGlobalLockScope LockScope,
-			RowHandle Row);
+		explicit FIndexTable(UEditorDataStorage& DataStorage);
+
+		RowHandle FindIndexedRow(EGlobalLockScope LockScope, IndexHash Index) const;
+		void IndexRow(EGlobalLockScope LockScope, IndexHash Index, RowHandle Row);
+		void BatchIndexRows(EGlobalLockScope LockScope, TConstArrayView<TPair<IndexHash, RowHandle>> IndexRowPairs);
+		void Reindex(EGlobalLockScope LockScope, IndexHash OriginalIndex, IndexHash NewIndex);
+		void RemoveIndex(EGlobalLockScope LockScope, IndexHash Index);
+		
+		void MarkDirty();
+		void RemoveInvalidRows();
 
 	private:
-		TMap<IndexHash, RowHandle> IndexLookupMap;
-		TMultiMap<RowHandle, IndexHash> ReverseIndexLookupMap;
-	
+		struct FCleanUpInfo final
+		{
+			// The number of row indices to store on the stack before allocating memory on the heap.
+			static constexpr uint32 StackReservationCount = 1024;
+
+			using InvalidRowContainer = TArray<int32, TInlineAllocator<StackReservationCount>>;
+
+			// The target amount of time that should be spend per frame on cleaning up old rows.
+			static const FTimespan TargetDuration;
+
+			// If less than this time is available after running jobs in a frame, reduce the number of jobs.
+			static const FTimespan JobShrinkThreshold;
+			// If more than this time is left per frame, increase the number of jobs.
+			static const FTimespan JobGrowthThreshold;
+			// The rate at which to grow the number of jobs run per frame.
+			static constexpr float JobShrinkFactor = 0.7f;
+			// The rate at which to shrink the number of jobs run per frame.
+			static constexpr float JobGrowthFactor = 1.2f;
+			// The minimum required number of jobs.
+			static constexpr int32 MinJobCount = 4;
+			// The maximum number of allowed jobs.
+			static constexpr int32 MaxJobCount = 1024;
+			
+			// If enough batches exceed this limit, the batch size for jobs is reduced. Aim for 4 batches per thread. Increasing this number
+			// increases the spikes that can happen when batches are too big. Decreasing this number increases the overhead of scheduling jobs
+			// and makes it more expensive overall.
+			static const FTimespan MaxBatchDuration;
+			// If enough batches are under this limit, the batch size for jobs is increased.
+			static const FTimespan MinBatchDuration;
+			// The absolute minimum number of rows to check for validity per batch.
+			static constexpr int32 MinBatchSize = 1000;
+			// The absolute maximum number of rows to check for validity per batch.
+			static constexpr int32 MaxBatchSize = 100'000;
+			// The number of jobs that need to take longer than MaxBatchDuration before the batch size gets shrunken down. This is
+			// used to mitigate the occasional spike due to for instance the OS handing off work to other programs.
+			static constexpr uint32 BatchShrinkThreshold = 2;
+			// The portion of jobs that need to be below the MinBatchDuration threshold before the batch size gets increased.
+			static constexpr float BatchIncreaseThreshold = 0.9f;
+			// The rate at which the batch size gets reduced if enough batches went over time.
+			static constexpr float BatchShrinkFactor = 0.8f;
+			// The rate at which the batch size gets increased if there's time left.
+			static constexpr float BatchGrowthFactor = 1.1f;
+
+			// Queue containing arrays with removed row indices.
+			TMpscQueue<InvalidRowContainer> DeletionQueue;
+			
+			// The maximum number of rows that a single job will process.
+			int32 BatchSize = 4096;
+			// The number of jobs that had too many batches to process within MaxBatchDuration.
+			std::atomic<uint32> BatchWentOverTime = 0;
+			// The number of jobs that processed their batches faster than MinBatchDuration.
+			std::atomic<uint32> BatchWentUnderTime = 0;
+			
+			// The maximum number of jobs to run per frame.
+			int32 MaxNumJobs = 8;
+			// The total number of jobs needed to inspect all indexed rows.
+			int32 JobCount = 0;
+			// The number of jobs that were not completed in the previous frame.
+			int32 RemainingJobs = 0;
+		};
+		
+		TMap<IndexHash, int32> IndexLookupMap;
+		TArray<RowHandle> Rows;
+		TArray<IndexHash> Hashes;
+		TDeque<int32> FreeList;
+		
+		FCleanUpInfo CleanUpInfo;
+		
+		UEditorDataStorage& DataStorage;
+		std::atomic<bool> DirtyDueToRemoval = false;
+
+		bool IsDirty() const;
 		void IndexRowUnguarded(IndexHash Index, RowHandle Row);
-		void RemoveIndexUnguarded(IndexHash Index);
+		
+		void InspectRowBlockForCleanUp(int32 Block);
+		void DrainDeletionQueue(FTimespan RemainingFrameTime);
+		void ClearRow(int32 Index);
+		void AdjustJobCount(FTimespan RemainingFrameTime);
+		void AdjustBatchSize(int32 JobCount);
 	};
 } // namespace UE::Editor::DataStorage
diff --git a/Engine/Plugins/Experimental/EditorDataStorage/Source/TedsCore/Private/TypedElementUtils.h b/Engine/Plugins/Experimental/EditorDataStorage/Source/TedsCore/Private/TypedElementUtils.h
new file mode 100644
--- /dev/null
+++ b/Engine/Plugins/Experimental/EditorDataStorage/Source/TedsCore/Private/TypedElementUtils.h
@@ -0,0 +1,59 @@
+// Copyright Epic Games, Inc. All Rights Reserved.
+
+#pragma once
+
+#include "Containers/ContainersFwd.h"
+#include "Elements/Common/TypedElementHandles.h"
+#include "MassEntityHandle.h"
+
+namespace UE::Editor::DataStorage
+{
+	inline FMassEntityHandle RowToMassEntityConversion(RowHandle Row);
+	inline RowHandle MassEntityToRowConversion(FMassEntityHandle Entity);
+
+	inline TConstArrayView<FMassEntityHandle> RowsToMassEntitiesConversion(TConstArrayView<RowHandle> Rows);
+	inline TArrayView<FMassEntityHandle> RowsToMassEntitiesConversion(TArrayView<RowHandle> Rows);
+	inline TConstArrayView<RowHandle> MassEntitiesToRowsConversion(TConstArrayView<FMassEntityHandle> Entities);
+	inline TArrayView<RowHandle> MassEntitiesToRowsConversion(TArrayView<FMassEntityHandle> Entities);
+} // namespace UE::Editor::DataStorage
+
+// Implementations;
+namespace UE::Editor::DataStorage
+{
+	FMassEntityHandle RowToMassEntityConversion(RowHandle Row)
+	{
+		return FMassEntityHandle::FromNumber(Row);
+	}
+
+	RowHandle MassEntityToRowConversion(FMassEntityHandle Entity)
+	{
+		return Entity.AsNumber();
+	}
+
+	TConstArrayView<FMassEntityHandle> RowsToMassEntitiesConversion(TConstArrayView<RowHandle> Rows)
+	{
+		return *reinterpret_cast<TConstArrayView<FMassEntityHandle>*>(&Rows);
+	}
+
+	TArrayView<FMassEntityHandle> RowsToMassEntitiesConversion(TArrayView<RowHandle> Rows)
+	{
+		return *reinterpret_cast<TArrayView<FMassEntityHandle>*>(&Rows);
+	}
+
+	TConstArrayView<RowHandle> MassEntitiesToRowsConversion(TConstArrayView<FMassEntityHandle> Entities)
+	{
+		return *reinterpret_cast<TConstArrayView<RowHandle>*>(&Entities);
+	}
+
+	TArrayView<RowHandle> MassEntitiesToRowsConversion(TArrayView<FMassEntityHandle> Entities)
+	{
+		return *reinterpret_cast<TArrayView<RowHandle>*>(&Entities);
+	}
+} // namespace UE::Editor::DataStorage
+
+// Several of the conversions between TEDS rows and Mass entities depend on the fact that they're both 64 bit integers so a cheap
+// type reinterpretation can be done instead of having copy lists. Especially for large lists this can become quite costly.
+static_assert(sizeof(FMassEntityHandle) == sizeof(UE::Editor::DataStorage::RowHandle),
+	"Size of Mass entity and data storage row have gone out of sync.");
+static_assert(alignof(FMassEntityHandle) == alignof(UE::Editor::DataStorage::RowHandle), 
+	"Alignment of Mass entity and data storage row have gone out of sync.");
\ No newline at end of file
diff --git a/Engine/Plugins/Experimental/EditorDataStorageFeatures/Source/TedsAssetData/Private/TedsAssetData.cpp b/Engine/Plugins/Experimental/EditorDataStorageFeatures/Source/TedsAssetData/Private/TedsAssetData.cpp
--- a/Engine/Plugins/Experimental/EditorDataStorageFeatures/Source/TedsAssetData/Private/TedsAssetData.cpp
+++ b/Engine/Plugins/Experimental/EditorDataStorageFeatures/Source/TedsAssetData/Private/TedsAssetData.cpp
@@ -622,7 +622,7 @@ void FTedsAssetData::OnAssetRenamed(const FAssetData& InAsset, const FString& In
 		}
 
 		Database.AddColumn<FUpdatedPathTag>(Row);
-		Database.ReindexRow(OldAssetHash, NewAssetHash, Row);
+		Database.Reindex(OldAssetHash, NewAssetHash);
 	}
 }
 
diff --git a/Engine/Source/Runtime/TypedElementFramework/Public/Elements/Interfaces/TypedElementDataStorageInterface.h b/Engine/Source/Runtime/TypedElementFramework/Public/Elements/Interfaces/TypedElementDataStorageInterface.h
--- a/Engine/Source/Runtime/TypedElementFramework/Public/Elements/Interfaces/TypedElementDataStorageInterface.h
+++ b/Engine/Source/Runtime/TypedElementFramework/Public/Elements/Interfaces/TypedElementDataStorageInterface.h
@@ -153,12 +153,29 @@ public:
 	/** Removes a previously reserved or added row. If the row handle is invalid or already removed, nothing happens */
 	virtual void RemoveRow(RowHandle Row) = 0;
 
+	/** Removes multiple rows at one. If any of the rows are invalid or already removed, nothing happens. */
+	virtual void BatchRemoveRows(TConstArrayView<RowHandle> Rows) = 0;
+
+	/** 
+	 * Removes all rows that have at least the provided columns.
+	 * This can be used to for instance remove all rows with a specific type tag, such as removing all rows with
+	 * an entity tag to remove all entities from the data storage.
+	 */
+	virtual void RemoveAllRowsWithColumns(TConstArrayView<UScriptStruct*> Columns) = 0;
+
+	/**
+	 * Removes all rows that have at least the provided columns as template arguments.
+	 * This can be used to for instance remove all rows with a specific type tag, such as removing all rows with
+	 * an entity tag to remove all entities from the data storage.
+	 */
+	template<UE::Editor::DataStorage::TColumnType... Columns>
+	void RemoveAllRowsWith();
+
 	/** Checks whether or not a row is in use. This is true even if the row has only been reserved. */
 	virtual bool IsRowAvailable(RowHandle Row) const = 0;
 	/** Checks whether or not a row has been reserved but not yet assigned to a table. */
 	virtual bool IsRowAssigned(RowHandle Row) const = 0;
 
-	
 	/**
 	 * @section Column management
 	 */
@@ -411,8 +428,7 @@ public:
 	virtual void BatchIndexRows(
 		TConstArrayView<TPair<UE::Editor::DataStorage::IndexHash, RowHandle>> IndexRowPairs) = 0;
 	/** Updates the index of a row to a new value. Effectively this is the same as removing an index and adding a new one. */
-	virtual void ReindexRow(
-		UE::Editor::DataStorage::IndexHash OriginalIndex, UE::Editor::DataStorage::IndexHash NewIndex, RowHandle Row) = 0;
+	virtual void Reindex(UE::Editor::DataStorage::IndexHash OriginalIndex, UE::Editor::DataStorage::IndexHash NewIndex) = 0;
 	/** Removes a previously registered index hash from the index lookup table or does nothing if the hash no longer exists. */
 	virtual void RemoveIndex(UE::Editor::DataStorage::IndexHash Index) = 0;
 
@@ -447,6 +463,13 @@ public:
 	virtual bool SupportsExtension(FName Extension) const = 0;
 	/** Provides a list of all extensions that are enabled. */
 	virtual void ListExtensions(TFunctionRef<void(FName)> Callback) const = 0;
+
+
+	/**
+	 * @section Deprecated
+	 */
+	UE_DEPRECATED(5.6, "Use 'Reindex' instead. The last argument 'Row' can be left out.")
+	inline void ReindexRow(UE::Editor::DataStorage::IndexHash OriginalIndex, UE::Editor::DataStorage::IndexHash NewIndex, RowHandle Row);
 };
 
 // Implementations
@@ -524,6 +547,12 @@ void IEditorDataStorageProvider::AddColumn(RowHandle Row, EnumT Value)
 	}
 }
 
+template<UE::Editor::DataStorage::TColumnType... Columns>
+void IEditorDataStorageProvider::RemoveAllRowsWith()
+{
+	RemoveAllRowsWithColumns({ Columns::StaticStruct()... });
+}
+
 template<auto Value, UE::Editor::DataStorage::TEnumType EnumT>
 void IEditorDataStorageProvider::AddColumn(RowHandle Row)
 {
@@ -677,3 +706,9 @@ SystemType* IEditorDataStorageProvider::GetExternalSystem()
 {
 	return reinterpret_cast<SystemType*>(GetExternalSystemAddress(SystemType::StaticClass()));
 }
+
+void IEditorDataStorageProvider::ReindexRow(
+	UE::Editor::DataStorage::IndexHash OriginalIndex, UE::Editor::DataStorage::IndexHash NewIndex, RowHandle Row)
+{
+	Reindex(OriginalIndex, NewIndex);
+}
\ No newline at end of file
