Fixed static analysis warning.

#rnx

diff --git a/Engine/Plugins/Experimental/EditorDataStorage/Source/TedsCore/Private/TypedElementDatabase.h b/Engine/Plugins/Experimental/EditorDataStorage/Source/TedsCore/Private/TypedElementDatabase.h
--- a/Engine/Plugins/Experimental/EditorDataStorage/Source/TedsCore/Private/TypedElementDatabase.h
+++ b/Engine/Plugins/Experimental/EditorDataStorage/Source/TedsCore/Private/TypedElementDatabase.h
@@ -138,12 +138,12 @@ public:
 		TConstArrayView<const UScriptStruct*> ColumnsToRemove) override;
 	virtual bool HasColumns(RowHandle Row, TConstArrayView<const UScriptStruct*> ColumnTypes) const override;
 	virtual bool HasColumns(RowHandle Row, TConstArrayView<TWeakObjectPtr<const UScriptStruct>> ColumnTypes) const override;
-	virtual void ListColumns(RowHandle Row, ColumnListCallbackRef Callback) const;
-	virtual void ListColumns(RowHandle Row, ColumnListWithDataCallbackRef Callback);
+	virtual void ListColumns(RowHandle Row, ColumnListCallbackRef Callback) const override;
+	virtual void ListColumns(RowHandle Row, ColumnListWithDataCallbackRef Callback) override;
 	virtual bool MatchesColumns(RowHandle Row, const UE::Editor::DataStorage::Queries::FConditions& Conditions) const override;
 
-	const UScriptStruct* FindDynamicColumn(const UE::Editor::DataStorage::FDynamicColumnDescription& Description) const override;
-	const UScriptStruct* GenerateDynamicColumn(const UE::Editor::DataStorage::FDynamicColumnDescription& Description) override;
+	virtual const UScriptStruct* FindDynamicColumn(const UE::Editor::DataStorage::FDynamicColumnDescription& Description) const override;
+	virtual const UScriptStruct* GenerateDynamicColumn(const UE::Editor::DataStorage::FDynamicColumnDescription& Description) override;
 
 	void RegisterTickGroup(FName GroupName, EQueryTickPhase Phase, FName BeforeGroup, FName AfterGroup, UE::Editor::DataStorage::EExecutionMode ExecutionMode);
 	void UnregisterTickGroup(FName GroupName, EQueryTickPhase Phase);
