Temporary workaround for a automated test crash.

#rnx

diff --git a/Engine/Plugins/Experimental/EditorDataStorage/Source/TedsCore/Private/TypedElementDatabaseIndexTable.cpp b/Engine/Plugins/Experimental/EditorDataStorage/Source/TedsCore/Private/TypedElementDatabaseIndexTable.cpp
--- a/Engine/Plugins/Experimental/EditorDataStorage/Source/TedsCore/Private/TypedElementDatabaseIndexTable.cpp
+++ b/Engine/Plugins/Experimental/EditorDataStorage/Source/TedsCore/Private/TypedElementDatabaseIndexTable.cpp
@@ -250,13 +250,22 @@ namespace UE::Editor::DataStorage
 		IndexHash Hash = Hashes[Index];
 
 		int32 StoredIndex;
-		bool Result = IndexLookupMap.RemoveAndCopyValue(Hash, StoredIndex);
-		checkf(Result && StoredIndex == Index, TEXT("Attempting to remove an invalid row but it doesn't exist in the lookup map."));
-
-		Rows[Index] = InvalidRowHandle;
-		Hashes[Index] = 0;
+		if (IndexLookupMap.RemoveAndCopyValue(Hash, StoredIndex) && StoredIndex == Index)
+		{
+			Rows[Index] = InvalidRowHandle;
+			Hashes[Index] = 0;
 
-		FreeList.PushFirst(Index);
+			FreeList.PushFirst(Index);
+		}
+		else
+		{
+			// TODO: This case was hit in the "Test 'SceneGraph.EditorWorkflows.Non-OFPE.Level Editor.Viewport.Transform.Mirror",
+			//		indicating that hash collisions are happening. This fix ignores those for now at the cost of leaving some
+			//		old values in, but a more sustainable solution needs to be created.
+			// Original code:
+			// bool Result = IndexLookupMap.RemoveAndCopyValue(Hash, StoredIndex);
+			// checkf(Result && StoredIndex == Index, TEXT("Attempting to remove an invalid row but it doesn't exist in the lookup map."));
+		}
 	}
 
 	void FIndexTable::AdjustJobCount(FTimespan RemainingFrameTime)
