[TEDS] Dynamic Column template query

* Enables TEDS query Where() clause to specify dynamic column templates as a type without the need to know the identifier.  The query will match all columns that were generated from the template.
* Dynamic column templates are more strongly typed.  It is no longer possible add a dynamic column template (without an ID) on a row.  The declaration for a dynamic column template must include the metadata "EditorDataStorage_DynamicColumnTemplate"
* Dynami column identifiers can be queried via metadata added during column generation.
* AttributeBinding API has been specialized for dynamic column templates.  Note that the signature has changed so that the Identifier is the first parameter.

Details:
* QueryStore needed to be notified of newly generated dynamic columns. This notification updates registered queries that have templates as part of the conditions to be sensitive to the new dynamic column type.
* Code generation will add a EditorDataStorage_DynamicColumnTemplate attribute to relevant types to allow the use of concepts that guide the intended use of the API.
* Dynamic column template usage has been updated in plugins
* Unit tests updated for Dynamic Columns (previous they were disabled)

diff --git a/Engine/Plugins/Experimental/EditorDataStorage/Source/TedsCore/Private/Debug/TypedElementDatabaseDebugCommands.cpp b/Engine/Plugins/Experimental/EditorDataStorage/Source/TedsCore/Private/Debug/TypedElementDatabaseDebugCommands.cpp
--- a/Engine/Plugins/Experimental/EditorDataStorage/Source/TedsCore/Private/Debug/TypedElementDatabaseDebugCommands.cpp
+++ b/Engine/Plugins/Experimental/EditorDataStorage/Source/TedsCore/Private/Debug/TypedElementDatabaseDebugCommands.cpp
@@ -502,30 +502,36 @@ static FAutoConsoleCommand CVarCountDynamicTagWithQuery(
 	}),
 	ECVF_Default);
 
+
+static TArray<UE::Editor::DataStorage::QueryHandle> DynamicColumnQueries;
 /**
  * Registers an activatable query which will run against rows that have the FTestDynamicTag::<Identifier> column
  */
 static FAutoConsoleCommand CVarRegisterListDynamicColumnQuery(
 	TEXT("TEDS.Debug.DynamicColumn.RegisterListDynamicColumnQuery"),
-	TEXT("Argument: Identifier ActivationGroup\n"),
+	TEXT("Argument: ActivationGroup [Identifier]\n"),
 	FConsoleCommandWithArgsDelegate::CreateLambda([](const TArray<FString>& Args)
 	{
 		using namespace UE::Editor::DataStorage::Queries;
 				
 		// Print column using query
-		if (Args.Num() != 2)
+		if (Args.Num() < 1 || Args.Num() > 2)
 		{
-			UE_LOG(LogEditorDataStorage, Error, TEXT("Invalid number of arguments. Identifier ActivationGroup"));
+			UE_LOG(LogEditorDataStorage, Error, TEXT("Invalid number of arguments. ActivationGroup [Identifer]"));
 			return;
 		}
 				
 		IEditorDataStorageProvider* DataStorage = GetMutableDataStorageFeature<IEditorDataStorageProvider>(StorageFeatureName);
-				
-		const FName Identifier(*Args[0]);
-		const FName ActivationGroup(*Args[1]);
+		
+		const FName ActivationGroup(*Args[0]);
+		FName Identifier;
+		if (Args.Num() == 2)
+		{
+			Identifier = *Args[1];
+		}
 
 		// Lists the rows processed that have 
-		DataStorage->RegisterQuery(
+		QueryHandle Handle = DataStorage->RegisterQuery(
 			Select(
 				TEXT("ProcessDynamicTagColumns"),
 				FProcessor(EQueryTickPhase::FrameEnd, DataStorage->GetQueryTickGroupName(EQueryTickGroups::Default))
@@ -535,14 +541,34 @@ static FAutoConsoleCommand CVarRegisterListDynamicColumnQuery(
 					auto RowView = MakeConstArrayView(Rows, Context.GetRowCount());
 					for (RowHandle Row : RowView)
 					{
-						UE_LOG(LogEditorDataStorage, Warning, TEXT("- '%llu'\n"), Row);
+						UE_LOG(LogEditorDataStorage, Log, TEXT("- '%llu'\n"), Row);
 					}
 				})
-			.Where()
-				.All<FTestDynamicTag>(Identifier)
+			.Where(TColumn<FTestDynamicTag>(Identifier))
 			.Compile());
 
-		UE_LOG(LogEditorDataStorage, Warning, TEXT("Query registered for Dynamic Column FTestDynamicTag::%s with activation group '%s'"), *Identifier.ToString(), *ActivationGroup.ToString());
+		DynamicColumnQueries.Add(Handle);
+
+		UE_LOG(LogEditorDataStorage, Log, TEXT("Query registered for Dynamic Column FTestDynamicTag::%s with activation group '%s'"), *Identifier.ToString(), *ActivationGroup.ToString());
+	}),
+	ECVF_Default);
+
+/**
+ * Registers an activatable query which will run against rows that have the FTestDynamicTag::<Identifier> column
+ */
+static FAutoConsoleCommand CVarUnRegisterListDynamicColumnQueries(
+	TEXT("TEDS.Debug.DynamicColumn.UnregisterListDynamicColumnQueries"),
+	TEXT(""),
+	FConsoleCommandWithArgsDelegate::CreateLambda([](const TArray<FString>& Args)
+	{
+		IEditorDataStorageProvider* DataStorage = UE::Editor::DataStorage::GetMutableDataStorageFeature<IEditorDataStorageProvider>(UE::Editor::DataStorage::StorageFeatureName);
+
+		for (UE::Editor::DataStorage::QueryHandle& Handle : DynamicColumnQueries)
+		{
+			DataStorage->UnregisterQuery(Handle);
+		}
+		DynamicColumnQueries.Empty();
+		
 	}),
 	ECVF_Default);
 
diff --git a/Engine/Plugins/Experimental/EditorDataStorage/Source/TedsCore/Private/Debug/TypedElementDatabaseDebugTypes.h b/Engine/Plugins/Experimental/EditorDataStorage/Source/TedsCore/Private/Debug/TypedElementDatabaseDebugTypes.h
--- a/Engine/Plugins/Experimental/EditorDataStorage/Source/TedsCore/Private/Debug/TypedElementDatabaseDebugTypes.h
+++ b/Engine/Plugins/Experimental/EditorDataStorage/Source/TedsCore/Private/Debug/TypedElementDatabaseDebugTypes.h
@@ -22,7 +22,7 @@ enum class ETedsDebugEnum : int8
 	Purple
 };
 
-USTRUCT()
+USTRUCT(meta=(EditorDataStorage_DynamicColumnTemplate))
 struct FTestDynamicColumn : public FEditorDataStorageColumn
 {
 	GENERATED_BODY()
@@ -32,7 +32,7 @@ struct FTestDynamicColumn : public FEditorDataStorageColumn
 	TArray<int32> IntArray;
 };
 
-USTRUCT()
+USTRUCT(meta=(EditorDataStorage_DynamicColumnTemplate))
 struct FTestDynamicTag : public FEditorDataStorageTag
 {
 	GENERATED_BODY()
diff --git a/Engine/Plugins/Experimental/EditorDataStorage/Source/TedsCore/Private/DynamicColumnGenerator.cpp b/Engine/Plugins/Experimental/EditorDataStorage/Source/TedsCore/Private/DynamicColumnGenerator.cpp
--- a/Engine/Plugins/Experimental/EditorDataStorage/Source/TedsCore/Private/DynamicColumnGenerator.cpp
+++ b/Engine/Plugins/Experimental/EditorDataStorage/Source/TedsCore/Private/DynamicColumnGenerator.cpp
@@ -6,26 +6,44 @@
 #include "TypedElementDataStorageSharedColumn.h"
 #include "Elements/Common/TypedElementCommonTypes.h"
 #include "Kismet2/ReloadUtilities.h"
+#include "Queries/TypedElementExtendedQueryStore.h"
+#include "UObject/MetaData.h"
 
 namespace UE::Editor::DataStorage
 {
+	void FDynamicColumnGenerator::SetQueryStore(FExtendedQueryStore& InQueryStore)
+	{
+		QueryStore = &InQueryStore;
+	}
+
 	FDynamicColumnGeneratorInfo FDynamicColumnGenerator::GenerateColumn(const UScriptStruct& Template, const FName& Identifier)
 	{
-		const FGeneratedColumnKey Key
+		const FTemplateIdKey Key
 		{
 			.Template = Template,
 			.Identifier = Identifier,
 		};
+
+		if (!ensureMsgf(!Identifier.IsNone(), TEXT("Identifier cannot be None")))
+		{
+			return FDynamicColumnGeneratorInfo
+			{
+				.Type = nullptr,
+				.Template = nullptr,
+				.bNewlyGenerated = false
+			};
+		}
 	
 		FDynamicColumnGeneratorInfo GeneratedColumnInfo;
 		UE_MT_SCOPED_WRITE_ACCESS(AccessDetector);
 		
 		{
-			const int32* GeneratedColumnIndex = GeneratedColumnLookup.Find(Key);
+			const int32* GeneratedColumnIndex = TemplateIdLookup.Find(Key);
 			if (GeneratedColumnIndex != nullptr)
 			{
 				const FGeneratedColumnRecord& GeneratedColumnRecord = GeneratedColumnData[*GeneratedColumnIndex];
 				GeneratedColumnInfo.Type = GeneratedColumnRecord.Type;
+				GeneratedColumnInfo.Template = GeneratedColumnRecord.Template;
 				GeneratedColumnInfo.bNewlyGenerated = false;
 				return GeneratedColumnInfo;
 			}		
@@ -34,18 +52,36 @@ namespace UE::Editor::DataStorage
 		auto IsValidType = [](const UScriptStruct& Template)
 		{
 			return
-				Template.IsChildOf(UE::Editor::DataStorage::FColumn::StaticStruct()) ||
-				Template.IsChildOf(UE::Editor::DataStorage::FTag::StaticStruct()) ||
+				Template.IsChildOf(FColumn::StaticStruct()) ||
+				Template.IsChildOf(FTag::StaticStruct()) ||
 				Template.IsChildOf(FTedsSharedColumn::StaticStruct());
 		};
-		if (!ensureMsgf(IsValidType(Template), TEXT("Template struct [%s] must derive from Column, Tag or SharedColumn"), *Template.GetName()))
+		auto HasTemplateMetadata = [](const UScriptStruct& Template)
+		{
+			return Template.FindMetaData(TEXT("EditorDataStorage_DynamicColumnTemplate")) != nullptr;
+		};
+		bool bCanUseTemplate = true;
+		if (!IsValidType(Template))
+		{
+			ensureAlwaysMsgf(false,  TEXT("Template struct [%s] must derive from Column, Tag or SharedColumn"), *Template.GetName());
+			bCanUseTemplate = false;
+		}
+		if (!HasTemplateMetadata(Template))
+		{
+			ensureAlwaysMsgf(false,  TEXT("Template struct [%s] must be declared with 'meta=(EditorDataStorage_DynamicColumnTemplate)'"), *Template.GetName());
+			bCanUseTemplate = false;
+		}
+		if (!bCanUseTemplate)
 		{
 			return FDynamicColumnGeneratorInfo
 			{
 				.Type = nullptr,
+				.Template = nullptr,
 				.bNewlyGenerated = false
 			};
 		}
+
+		TemplateTypeLookup.Add(&Template);
 	
 		{
 			checkf(Template.GetCppStructOps() != nullptr && Template.IsNative(), TEXT("Can only create column from native struct"));
@@ -56,7 +92,6 @@ namespace UE::Editor::DataStorage
 			ObjectNameBuilder.Append(Identifier.ToString());
 
 			const FName ObjectName = FName(ObjectNameBuilder);
-			const FTopLevelAssetPath AssetPath(GetTransientPackage()->GetFName(), ObjectName);
 			
 			UScriptStruct* NewScriptStruct = NewObject<UScriptStruct>(GetTransientPackage(), ObjectName);
 			// Ensure it is not garbage collected
@@ -69,28 +104,62 @@ namespace UE::Editor::DataStorage
 			
 			NewScriptStruct->Bind();
 			NewScriptStruct->PrepareCppStructOps();
-			NewScriptStruct->StaticLink(true);			
+			NewScriptStruct->StaticLink(true);
+			// Set metadata to indicate that this struct is derived from a dynamic template
+			// The value will contain the identifier
+			TStringBuilder<256> IdentifierStringBuilder;
+			Identifier.AppendString(IdentifierStringBuilder);
+			NewScriptStruct->GetPackage()->GetMetaData()->SetValue(NewScriptStruct, TEXT("EditorDataStorage_DerivedFromDynamicTemplate"), IdentifierStringBuilder.ToString());
 			const int32 Index = GeneratedColumnData.Emplace(FGeneratedColumnRecord
 			{
 				.Identifier = Identifier,
 				.Template = &Template,
-				.Type = NewScriptStruct,
-				.AssetPath = AssetPath
+				.Type = NewScriptStruct
 			});
 			
-			GeneratedColumnLookup.Add(Key, Index);
-						
+			TemplateIdLookup.Add(Key, Index);
+			GeneratedTypeLookup.Add(NewScriptStruct, Index);
+			
 			GeneratedColumnInfo.Type = NewScriptStruct;
+			GeneratedColumnInfo.Template = &Template;
 			GeneratedColumnInfo.bNewlyGenerated = true;
+
+			if (QueryStore)
+			{
+				QueryStore->NotifyNewDynamicColumn(GeneratedColumnInfo);
+			}
+			
 			return GeneratedColumnInfo;
 		}
 	}
-	
-	const UScriptStruct* FDynamicColumnGenerator::FindColumn(const UScriptStruct& Template, const FName& Identifier) const
+
+	void FDynamicColumnGenerator::ForEachDynamicColumn(const UScriptStruct& Template, TFunctionRef<void(const FDynamicColumnGeneratorInfo&)> Callback)
+	{
+		UE_MT_SCOPED_READ_ACCESS(AccessDetector);
+
+		for (const FGeneratedColumnRecord& Record : GeneratedColumnData)
+		{
+			if (Record.Type->IsChildOf(&Template))
+			{
+				FDynamicColumnGeneratorInfo Info;
+				Info.Template = &Template;
+				Info.Type = Record.Type;
+				Info.bNewlyGenerated = false;
+				Callback(Info);
+			}
+		}
+	}
+
+	bool FDynamicColumnGenerator::IsDynamicTemplate(const UScriptStruct& InCandidate) const
+	{
+		return TemplateTypeLookup.Find(&InCandidate) != nullptr;
+	}
+
+	const UScriptStruct* FDynamicColumnGenerator::FindByTemplateId(const UScriptStruct& Template, const FName& Identifier) const
 	{
 		UE_MT_SCOPED_READ_ACCESS(AccessDetector);
 	
-		const int32* IndexPtr = GeneratedColumnLookup.Find(FGeneratedColumnKey{
+		const int32* IndexPtr = TemplateIdLookup.Find(FTemplateIdKey{
 			.Template = Template,
 			.Identifier = Identifier
 			});
@@ -100,7 +169,23 @@ namespace UE::Editor::DataStorage
 		}
 		return nullptr;
 	}
-	
+
+	bool FDynamicColumnGenerator::FindByGeneratedType(const UScriptStruct& GeneratedType, FDynamicColumnGeneratorInfo& Out) const
+	{
+		UE_MT_SCOPED_READ_ACCESS(AccessDetector);
+
+		const int32* IndexPtr = GeneratedTypeLookup.Find(&GeneratedType);
+		if (IndexPtr)
+		{
+			const FGeneratedColumnRecord& Record = GeneratedColumnData[*IndexPtr];
+			Out.Type = Record.Type;
+			Out.Template = Record.Template;
+			Out.bNewlyGenerated = false;
+			return true;
+		}
+		return false;
+	}
+
 	FValueTagManager::FValueTagManager(FDynamicColumnGenerator& InColumnGenerator)
 		: ColumnGenerator(InColumnGenerator)
 	{
diff --git a/Engine/Plugins/Experimental/EditorDataStorage/Source/TedsCore/Private/DynamicColumnGenerator.h b/Engine/Plugins/Experimental/EditorDataStorage/Source/TedsCore/Private/DynamicColumnGenerator.h
--- a/Engine/Plugins/Experimental/EditorDataStorage/Source/TedsCore/Private/DynamicColumnGenerator.h
+++ b/Engine/Plugins/Experimental/EditorDataStorage/Source/TedsCore/Private/DynamicColumnGenerator.h
@@ -27,6 +27,7 @@ struct FTedsValueTagColumn : public FTedsSharedColumn
 
 namespace UE::Editor::DataStorage
 {
+	class FExtendedQueryStore;
 	using FValueTagColumn = FTedsValueTagColumn;
 
 	struct FDynamicColumnInfo
@@ -37,6 +38,7 @@ namespace UE::Editor::DataStorage
 	struct FDynamicColumnGeneratorInfo
 	{
 		const UScriptStruct* Type;
+		const UScriptStruct* Template;
 		bool bNewlyGenerated;
 	};
 
@@ -46,12 +48,20 @@ namespace UE::Editor::DataStorage
 	class FDynamicColumnGenerator
 	{
 	public:
+		void SetQueryStore(FExtendedQueryStore& InQueryStore);
+		
 		/**
 		 * Generates a dynamic TEDS column type based on a Template type (if it hasn't been generated before)
 		 */
 		FDynamicColumnGeneratorInfo GenerateColumn(const UScriptStruct& Template, const FName& Identifier);
 
-		const UScriptStruct* FindColumn(const UScriptStruct& Template, const FName& Identifier) const;
+		void ForEachDynamicColumn(const UScriptStruct& Template, TFunctionRef<void(const FDynamicColumnGeneratorInfo&)> Callback);
+
+		bool IsDynamicTemplate(const UScriptStruct& InCandidate) const;
+		const UScriptStruct* FindByTemplateId(const UScriptStruct& Template, const FName& Identifier) const;
+		// Returns true if found, false otherwise.
+		// If found, Out will be populated
+		bool FindByGeneratedType(const UScriptStruct& GeneratedType, FDynamicColumnGeneratorInfo& Out) const;
 	private:
 
 		struct FGeneratedColumnRecord
@@ -59,20 +69,24 @@ namespace UE::Editor::DataStorage
 			FName Identifier;
 			const UScriptStruct* Template;
 			const UScriptStruct* Type;
-			FTopLevelAssetPath AssetPath;
 		};
 
-		struct FGeneratedColumnKey
+		struct FTemplateIdKey
 		{
 			const UScriptStruct& Template;
 			FName Identifier;
 
-			friend bool operator==(const FGeneratedColumnKey& Lhs, const FGeneratedColumnKey& Rhs)
+			friend bool operator==(const FTemplateIdKey& Lhs, const FTemplateIdKey& Rhs)
 			{
 				return Lhs.Identifier == Rhs.Identifier && &Lhs.Template == &Rhs.Template;
 			}
 
-			friend uint32 GetTypeHash(const FGeneratedColumnKey& Key)
+			friend bool operator!=(const FTemplateIdKey& Lhs, const FTemplateIdKey& Rhs)
+			{
+				return !operator==(Lhs, Rhs);
+			}
+
+			friend uint32 GetTypeHash(const FTemplateIdKey& Key)
 			{
 				return HashCombineFast(GetTypeHash(Key.Identifier), PointerHash(&Key.Template));
 			}
@@ -81,10 +95,17 @@ namespace UE::Editor::DataStorage
 		UE_MT_DECLARE_RW_ACCESS_DETECTOR(AccessDetector);
 
 		TArray<FGeneratedColumnRecord> GeneratedColumnData;
+
+		TSet<const UScriptStruct*> TemplateTypeLookup;
 		
-		// Looks up generated column index by the parameters used to generate it
+		// Looks up generated column index by template and ID used to generate it
 		// Used to de-duplicate
-		TMap<FGeneratedColumnKey, int32> GeneratedColumnLookup;
+		TMap<FTemplateIdKey, int32> TemplateIdLookup;
+
+		// Looks up generated column data by the generated column type
+		TMap<const UScriptStruct*, int32> GeneratedTypeLookup;
+
+		FExtendedQueryStore* QueryStore = nullptr;
 	};
 
 	class FValueTagManager
diff --git a/Engine/Plugins/Experimental/EditorDataStorage/Source/TedsCore/Private/Queries/TypedElementExtendedQueryStore.cpp b/Engine/Plugins/Experimental/EditorDataStorage/Source/TedsCore/Private/Queries/TypedElementExtendedQueryStore.cpp
--- a/Engine/Plugins/Experimental/EditorDataStorage/Source/TedsCore/Private/Queries/TypedElementExtendedQueryStore.cpp
+++ b/Engine/Plugins/Experimental/EditorDataStorage/Source/TedsCore/Private/Queries/TypedElementExtendedQueryStore.cpp
@@ -140,6 +140,11 @@ namespace UE::Editor::DataStorage
 		}
 
 	}
+	
+	FExtendedQueryStore::FExtendedQueryStore(FDynamicColumnGenerator& InDynamicColumnGenerator)
+		: DynamicColumnGenerator(&InDynamicColumnGenerator)
+	{
+	}
 
 	FExtendedQueryStore::Handle FExtendedQueryStore::RegisterQuery(
 		FQueryDescription Query,
@@ -480,6 +485,70 @@ namespace UE::Editor::DataStorage
 			PhaseFinalizationQueries[static_cast<QueryTickPhaseType>(Phase)]);
 	}
 
+	void FExtendedQueryStore::NotifyNewDynamicColumn(const FDynamicColumnGeneratorInfo& GeneratedColumnInfo)
+	{
+		if (!Private::bUseTedsMatchFunction)
+		{
+			return;
+		}
+		// Find all queries that use the template as a condition and add the newly generated type
+		// to the query
+		Queries.ListAliveEntries([&](QueryStore::Handle Handle, FExtendedQuery& Query)
+		{
+			if (!Query.Description.Conditions.IsSet())
+			{
+				// Generic dynamic column template queries only supported for when query conditions used
+				return;
+			}
+			// Handle dynamic Selected columns
+			for (int32 Index = 0, End = Query.Description.DynamicSelectionTypes.Num(); Index < End; ++Index)
+			{
+				// Only consider query descriptions that have a dynamic selection type that is supposed to match all
+				// columns 
+				const bool bIsTemplateSelectionType = Query.Description.DynamicSelectionTypes[Index].Identifier.IsNone();
+				const bool bTemplateMatches = Query.Description.DynamicSelectionTypes[Index].TemplateType == GeneratedColumnInfo.Template;
+				if (bIsTemplateSelectionType && bTemplateMatches)
+				{
+					const UScriptStruct* GeneratedColumnType = GeneratedColumnInfo.Type;
+					const EMassFragmentAccess AccessMode = ConvertToNativeAccessType(Query.Description.DynamicSelectionAccessTypes[Index]);
+					const EMassFragmentPresence PresenceMode = EMassFragmentPresence::Any;
+					if (GeneratedColumnType->IsChildOf<FEditorDataStorageColumn>())
+					{
+						Query.GetQuery()->AddRequirement(GeneratedColumnType, AccessMode, PresenceMode);
+					}
+					else if (GeneratedColumnType->IsChildOf<FEditorDataStorageTag>())
+					{
+						Query.GetQuery()->AddTagRequirement(*GeneratedColumnType, PresenceMode);
+					}
+				}
+			}
+
+			// Handle dynamic conditional columns
+			for (int32 Index = 0, End = Query.Description.DynamicConditionDescriptions.Num(); Index < End; ++Index)
+			{
+				// Only consider query descriptions that have a dynamic selection type that is supposed to match all
+				// columns 
+				const bool bIsTemplateSelectionType = Query.Description.DynamicConditionDescriptions[Index].Identifier.IsNone();
+				const bool bTemplateMatches = Query.Description.DynamicConditionDescriptions[Index].TemplateType == GeneratedColumnInfo.Template;
+				if (bIsTemplateSelectionType && bTemplateMatches)
+				{
+					const UScriptStruct* GeneratedColumnType = GeneratedColumnInfo.Type;
+					const EMassFragmentAccess AccessMode = EMassFragmentAccess::None;
+					const EMassFragmentPresence PresenceMode = EMassFragmentPresence::Any;
+
+					if (GeneratedColumnType->IsChildOf<FEditorDataStorageColumn>())
+					{
+						Query.GetQuery()->AddRequirement(GeneratedColumnType, AccessMode, PresenceMode);
+					}
+					else if (GeneratedColumnType->IsChildOf<FEditorDataStorageTag>())
+					{
+						Query.GetQuery()->AddTagRequirement(*GeneratedColumnType, PresenceMode);
+					}
+				}
+			}
+		});
+	}
+
 	void FExtendedQueryStore::DebugPrintQueryCallbacks(FOutputDevice& Output) const
 	{
 		Output.Log(TEXT("The Typed Elements Data Storage has the following query callbacks:"));
@@ -512,7 +581,12 @@ namespace UE::Editor::DataStorage
 	}
 
 	bool FExtendedQueryStore::SetupDynamicColumns(FQueryDescription& Query, FEnvironment& Environment)
-	{	
+	{
+		auto IsDynamicColumnTemplate = [](const UScriptStruct& Struct, const FName& Identifier)->bool
+		{
+			return Identifier.IsNone();
+		};
+		
 		const int32 SelectionCount = Query.DynamicSelectionTypes.Num();
 
 		for (int32 Index = 0, End = SelectionCount; Index < End; ++Index)
@@ -542,13 +616,28 @@ namespace UE::Editor::DataStorage
 			{
 				continue;
 			}
-			const UScriptStruct* DynamicColumnType = Environment.GenerateDynamicColumn(*Description.TemplateType, Description.Identifier);
-			FQueryDescription::EOperatorType Operation = Query.DynamicConditionOperations[Index];
-			Query.ConditionTypes.Add(Operation);
-			Query.ConditionOperators.Add(FQueryDescription::FOperator
+			if (IsDynamicColumnTemplate(*Description.TemplateType, Description.Identifier))
 			{
-				.Type = DynamicColumnType
-			});
+				FQueryDescription::EOperatorType Operation = Query.DynamicConditionOperations[Index];
+				DynamicColumnGenerator->ForEachDynamicColumn(*Description.TemplateType, [&Query, Operation](const FDynamicColumnGeneratorInfo& Info)
+				{
+					Query.ConditionTypes.Add(Operation);
+					Query.ConditionOperators.Add(FQueryDescription::FOperator
+					{
+						.Type = Info.Type
+					});
+				});
+			}
+			else
+			{
+				const UScriptStruct* DynamicColumnType = DynamicColumnGenerator->GenerateColumn(*Description.TemplateType, Description.Identifier).Type;
+				FQueryDescription::EOperatorType Operation = Query.DynamicConditionOperations[Index];
+				Query.ConditionTypes.Add(Operation);
+				Query.ConditionOperators.Add(FQueryDescription::FOperator
+				{
+					.Type = DynamicColumnType
+				});
+			}
 		}
 
 		return true;
@@ -739,9 +828,20 @@ namespace UE::Editor::DataStorage
 
 			// Note: need to remove duplicates to appease Mass
 			TSet<const UScriptStruct*, DefaultKeyFuncs<const UScriptStruct*>, TInlineSetAllocator<128>> DedupedColumns;
-			for (const auto& Column : Columns)
+			for (const TWeakObjectPtr<const UScriptStruct>& Column : Columns)
 			{
-				DedupedColumns.Add(Column.Get());
+				// If any of the columns are Dynamic Column Templates, then expand them
+				if (DynamicColumnGenerator->IsDynamicTemplate(*Column.Get()))
+				{
+					DynamicColumnGenerator->ForEachDynamicColumn(*Column.Get(), [&DedupedColumns ](const FDynamicColumnGeneratorInfo& Info)
+					{
+						DedupedColumns.Add(Info.Type);
+					});
+				}
+				else
+				{
+					DedupedColumns.Add(Column.Get());
+				}
 			}
 
 			// Everything is going to be Any, Mass' internals checks requirements 
@@ -782,22 +882,8 @@ namespace UE::Editor::DataStorage
 				const FQueryDescription::FOperator* Operand = Query.ConditionOperators.GetData();
 				for (FQueryDescription::EOperatorType Type : Query.ConditionTypes)
 				{
-					EMassFragmentPresence Presence;
-					switch (Type)
-					{
-					case FQueryDescription::EOperatorType::SimpleAll:
-						Presence = EMassFragmentPresence::All;
-						break;
-					case FQueryDescription::EOperatorType::SimpleAny:
-						Presence = EMassFragmentPresence::Any;
-						break;
-					case FQueryDescription::EOperatorType::SimpleNone:
-						Presence = EMassFragmentPresence::None;
-						break;
-					default:
-						continue;
-					}
-
+					const EMassFragmentPresence Presence = ConvertToNativePresenceType(Type);
+					
 					if (UE::Mass::IsA<FMassTag>(Operand->Type.Get()))
 					{
 						NativeQuery.AddTagRequirement(*(Operand->Type), Presence);
@@ -1050,6 +1136,27 @@ namespace UE::Editor::DataStorage
 		}
 	}
 
+	EMassFragmentPresence FExtendedQueryStore::ConvertToNativePresenceType(FQueryDescription::EOperatorType OperatorType)
+	{
+		EMassFragmentPresence Presence;
+		switch (OperatorType)
+		{
+		case FQueryDescription::EOperatorType::SimpleAll:
+			Presence = EMassFragmentPresence::All;
+			break;
+		case FQueryDescription::EOperatorType::SimpleAny:
+			Presence = EMassFragmentPresence::Any;
+			break;
+		case FQueryDescription::EOperatorType::SimpleNone:
+			Presence = EMassFragmentPresence::None;
+			break;
+		default:
+			Presence = EMassFragmentPresence::Any;
+			ensureAlwaysMsgf(false, TEXT("OperatorType '%d' cannot be converted to native type"), static_cast<int32>(OperatorType));
+		}
+		return Presence;
+	}
+
 	void FExtendedQueryStore::RegisterPreambleQuery(EQueryTickPhase Phase, Handle Query)
 	{
 		PhasePreparationQueries[static_cast<QueryTickPhaseType>(Phase)].Add(Query);
diff --git a/Engine/Plugins/Experimental/EditorDataStorage/Source/TedsCore/Private/Queries/TypedElementExtendedQueryStore.h b/Engine/Plugins/Experimental/EditorDataStorage/Source/TedsCore/Private/Queries/TypedElementExtendedQueryStore.h
--- a/Engine/Plugins/Experimental/EditorDataStorage/Source/TedsCore/Private/Queries/TypedElementExtendedQueryStore.h
+++ b/Engine/Plugins/Experimental/EditorDataStorage/Source/TedsCore/Private/Queries/TypedElementExtendedQueryStore.h
@@ -16,6 +16,9 @@ class FOutputDevice;
 namespace UE::Editor::DataStorage
 {
 	class FEnvironment;
+	
+	class FDynamicColumnGenerator;
+	struct FDynamicColumnGeneratorInfo;
 
 	struct FExtendedQuery
 	{
@@ -37,7 +40,7 @@ namespace UE::Editor::DataStorage
 		}
 		
 	};
-
+	
 	/**
 	 * Storage and utilities for Typed Element queries after they've been processed by the Data Storage implementation.
 	 */
@@ -49,6 +52,8 @@ namespace UE::Editor::DataStorage
 		using Handle = QueryStore::Handle;
 		using ListAliveEntriesConstCallback = QueryStore::ListAliveEntriesConstCallback;
 
+		explicit FExtendedQueryStore(FDynamicColumnGenerator& InDynamicColumnGenerator);
+
 		/**
 		 * @section Registration
 		 * @description A set of functions to manage the registration of queries.
@@ -146,6 +151,8 @@ namespace UE::Editor::DataStorage
 			IEditorDataStorageProvider::EQueryTickPhase Phase,
 			float DeltaTime);
 
+		void NotifyNewDynamicColumn(const FDynamicColumnGeneratorInfo& GeneratedColumnInfo);
+
 		void DebugPrintQueryCallbacks(FOutputDevice& Output) const;
 
 	private:
@@ -191,6 +198,7 @@ namespace UE::Editor::DataStorage
 
 		EMassFragmentAccess ConvertToNativeAccessType(IEditorDataStorageProvider::EQueryAccessType AccessType);
 		EMassFragmentPresence ConvertToNativePresenceType(IEditorDataStorageProvider::EQueryAccessType AccessType);
+		EMassFragmentPresence ConvertToNativePresenceType(FQueryDescription::EOperatorType OperatorType);
 
 		void RegisterPreambleQuery(IEditorDataStorageProvider::EQueryTickPhase Phase, Handle Query);
 		void RegisterPostambleQuery(IEditorDataStorageProvider::EQueryTickPhase Phase, Handle Query);
@@ -210,5 +218,6 @@ namespace UE::Editor::DataStorage
 		TArray<Handle> PhaseFinalizationQueries[MaxTickPhase];
 		TArray<Handle> PendingActivatables;
 		TArray<Handle> ActiveActivatables;
+		FDynamicColumnGenerator* DynamicColumnGenerator;
 	};
 } // namepsace UE::Editor::DataStorage
diff --git a/Engine/Plugins/Experimental/EditorDataStorage/Source/TedsCore/Private/TypedElementDatabaseEnvironment.cpp b/Engine/Plugins/Experimental/EditorDataStorage/Source/TedsCore/Private/TypedElementDatabaseEnvironment.cpp
--- a/Engine/Plugins/Experimental/EditorDataStorage/Source/TedsCore/Private/TypedElementDatabaseEnvironment.cpp
+++ b/Engine/Plugins/Experimental/EditorDataStorage/Source/TedsCore/Private/TypedElementDatabaseEnvironment.cpp
@@ -10,11 +10,13 @@ namespace UE::Editor::DataStorage
 		: DataStorage(InDataStorage)
 		, DirectDeferredCommands(*this)
 		, IndexTable(InDataStorage)
+		, Queries(DynamicColumnGenerator)
 		, MementoSystem(InDataStorage)
 		, ValueTagManager(DynamicColumnGenerator)
 		, MassEntityManager(InMassEntityManager)
 		, MassPhaseManager(InMassPhaseManager)
 	{
+		DynamicColumnGenerator.SetQueryStore(Queries);
 	}
 
 	Legacy::FCommandBuffer& FEnvironment::GetDirectDeferredCommands()
@@ -94,7 +96,7 @@ namespace UE::Editor::DataStorage
 
 	const UScriptStruct* FEnvironment::FindDynamicColumn(const UScriptStruct& Template, const FName& Identifier)
 	{
-		return DynamicColumnGenerator.FindColumn(Template, Identifier);
+		return DynamicColumnGenerator.FindByTemplateId(Template, Identifier);
 	}
 
 	const UScriptStruct* FEnvironment::GenerateDynamicColumn(const UScriptStruct& Template, const FName& Identifier)
diff --git a/Engine/Plugins/Experimental/EditorDataStorage/Source/TedsCore/Tests/DynamicColumnTests.cpp b/Engine/Plugins/Experimental/EditorDataStorage/Source/TedsCore/Tests/DynamicColumnTests.cpp
--- a/Engine/Plugins/Experimental/EditorDataStorage/Source/TedsCore/Tests/DynamicColumnTests.cpp
+++ b/Engine/Plugins/Experimental/EditorDataStorage/Source/TedsCore/Tests/DynamicColumnTests.cpp
@@ -54,52 +54,12 @@ namespace UE::Editor::DataStorage::Tests
 			return QueryHandle;
 		}
 
-		// Wait at least FrameCount gamethread ticks
-		void WaitFrames(int32 FrameCount = 2)
-		{
-			if(ensureMsgf(!IsInGameThread(), TEXT("Must not be on async thread")))
-			{
-				TPromise<void> Promise;
-				TFuture<void> Future = Promise.GetFuture();
-				bool bFutureSet = false;
-
-				FDelegateHandle DelegateHandle;
-
-				Async(EAsyncExecution::TaskGraphMainThread, [&Promise, FrameCount, &DelegateHandle, &bFutureSet]()
-				{
-					const uint64 LastFrameCount = GFrameCounter;
-					const uint64 UnblockAt = LastFrameCount + FrameCount;
-					DelegateHandle = FCoreDelegates::OnEndFrame.AddLambda([UnblockAt, &Promise, &bFutureSet]()
-					{
-						const uint64 CurrentFrameCount = GFrameCounter;
-						if (UnblockAt < CurrentFrameCount)
-						{
-							// Handle case where test thread hasn't yet unregistered the delegate
-							if (!bFutureSet)
-							{
-								Promise.SetValue();
-								bFutureSet = true;
-							}
-						}
-					});
-				});
-
-				Future.Wait();
-				
-				Async(EAsyncExecution::TaskGraphMainThread, [&DelegateHandle]()
-				{
-					FCoreDelegates::OnEndFrame.Remove(DelegateHandle);
-					DelegateHandle = FDelegateHandle();
-				}).Wait();
-			}
-		}
-
 	
 	END_DEFINE_SPEC(DynamicColumnTestFixture)
 
 	void DynamicColumnTestFixture::Define()
 	{
-		xBeforeEach([this]()
+		BeforeEach([this]()
 		{
 			TedsInterface = GetMutableDataStorageFeature<IEditorDataStorageProvider>(StorageFeatureName);
 			TestTable = RegisterTestTable();
@@ -107,12 +67,11 @@ namespace UE::Editor::DataStorage::Tests
 			TestTrue("", TedsInterface != nullptr);
 		});
 			
-		xDescribe("", [this]
+		Describe("", [this]
 		{
-			It(TEXT("Tags"), EAsyncExecution::ThreadPool, [this]()
+			It(TEXT("Tags"), EAsyncExecution::TaskGraphMainThread, [this]()
 			{
 				// Add dynamic columns that are actually tags (ie. dataless)
-				Async(EAsyncExecution::TaskGraphMainThread,[this]()
 				{
 					for (int32 Index = 0; Index <= 2; ++Index)
 					{
@@ -147,11 +106,10 @@ namespace UE::Editor::DataStorage::Tests
 
 					TestFalse("Unexpected columns found", TedsInterface->HasColumns(Rows[2], MakeConstArrayView({DynamicTagTypes[0]})));
 					TestTrue("Expected columns not found", TedsInterface->HasColumns(Rows[2], MakeConstArrayView({DynamicTagTypes[1]})));
-				}).Wait();
+				}
 				
 				// Direct Query
 				{
-					Async(EAsyncExecution::TaskGraphMainThread,[this]()
 					{
 						using namespace UE::Editor::DataStorage::Queries;
 
@@ -187,9 +145,7 @@ namespace UE::Editor::DataStorage::Tests
 						{
 							QueryHandle Query = RegisterQuery(
 								Select().
-								Where().
-									All<FTestDynamicTag>(Identifiers[0]).
-									All<FTestDynamicTag>(Identifiers[1]).
+								Where(TColumn<FTestDynamicTag>(Identifiers[0]) && TColumn<FTestDynamicTag>(Identifiers[1])).
 								Compile());
 						
 							SetExpectedMatches({Rows[0]});
@@ -200,223 +156,49 @@ namespace UE::Editor::DataStorage::Tests
 							// Should match Rows 0 and 1
 							QueryHandle Query = RegisterQuery(
 								Select().
-									Where().
-										All<FTestDynamicTag>(Identifiers[0]).
+								Where(TColumn<FTestDynamicTag>(Identifiers[0])).
 								Compile());
 							SetExpectedMatches({Rows[0], Rows[1]});
 							FQueryResult Result = TedsInterface->RunQuery(Query, Callback);
 							TestEqual("Match Row[0] and Row[1]", Result.Count, GetMatchCount());
 						}
-						{
-							// Should match Rows 1
-							QueryHandle Query = RegisterQuery(
-								Select().
-									Where().
-										All<FTestDynamicTag>(Identifiers[0]).
-										None<FTestDynamicTag>(Identifiers[1]).
-								Compile());
-							SetExpectedMatches({Rows[1]});
-							FQueryResult Result = TedsInterface->RunQuery(Query, Callback);
-							TestEqual("Match Row[1]", Result.Count, GetMatchCount());
-						}
 						{
 							// Should match Rows 0 and 2
 							QueryHandle Query = RegisterQuery(
 								Select().
-									Where().
-										All<FTestDynamicTag>(Identifiers[1]).
+								Where(TColumn<FTestDynamicTag>(Identifiers[1])).
 								Compile());
 							SetExpectedMatches({Rows[0], Rows[2]});
 							FQueryResult Result = TedsInterface->RunQuery(Query, Callback);
 							TestEqual("Match Row[0] and Row[2]", Result.Count, GetMatchCount());
 						}
 						{
-							// Should match Rows 0 and 2
+							// Should match Rows 0, 1 and 2
 							QueryHandle Query = RegisterQuery(
 								Select().
-									Where().
-										None<FTestDynamicTag>(Identifiers[0]).
-										All<FTestDynamicTag>(Identifiers[1]).
+								Where(TColumn<FTestDynamicTag>(Identifiers[0]) || TColumn<FTestDynamicTag>(Identifiers[1])).
 								Compile());
-							SetExpectedMatches({Rows[2]});
+							SetExpectedMatches({Rows[0], Rows[1], Rows[2]});
 							FQueryResult Result = TedsInterface->RunQuery(Query, Callback);
-							TestEqual("Match Row[2]", Result.Count, GetMatchCount());
+							TestEqual("Match All Rows", Result.Count, GetMatchCount());
 						}
 						{
 							// Should match Rows 0, 1 and 2
 							QueryHandle Query = RegisterQuery(
 								Select().
-									Where().
-										Any<FTestDynamicTag>(Identifiers[0]).
-										Any<FTestDynamicTag>(Identifiers[1]).
+								Where(TColumn<FTestDynamicTag>()).
 								Compile());
 							SetExpectedMatches({Rows[0], Rows[1], Rows[2]});
 							FQueryResult Result = TedsInterface->RunQuery(Query, Callback);
-							TestEqual("Match All Rows", Result.Count, GetMatchCount());
-						}
-					}).Wait();
-				}
-
-				// Processor Query
-				{
-					TArray<RowHandle> RowsToMatch;
-					TArray<int32> MatchCount;
-					int32 UnexpectedRowCount;
-					auto SetExpectedMatches = [&RowsToMatch, &MatchCount, &UnexpectedRowCount](TConstArrayView<RowHandle> Expectation)
-					{
-						RowsToMatch.Empty(RowsToMatch.Num());
-						RowsToMatch.Append(Expectation);
-						MatchCount.Empty(MatchCount.Num());
-						MatchCount.SetNum(Expectation.Num());
-						for (int32 I = 0; I < MatchCount.Num(); ++I)
-						{
-							MatchCount[I] = 0;
-						}
-						UnexpectedRowCount = 0;
-					};
-
-					TArray<FName> ActivationKeys;
-					TArray<TArray<RowHandle>> QueryExpectedMatchRows;
-					
-					Async(EAsyncExecution::TaskGraphMainThread,[this, &RowsToMatch, &MatchCount, &UnexpectedRowCount, &ActivationKeys, &QueryExpectedMatchRows]()
-					{
-						auto Callback = [this, &RowsToMatch, &MatchCount, &UnexpectedRowCount] (IQueryContext& Context, const RowHandle* CallbackRows)
-						{
-							TConstArrayView<RowHandle> RowsView = MakeConstArrayView(CallbackRows, Context.GetRowCount());
-							for (auto Row : RowsView)
-							{
-								int32 Index = RowsToMatch.Find(Row);
-								TestTrue(TEXT("Returned row in query is within expected match array"), Index != INDEX_NONE);
-								if (Index != INDEX_NONE)
-								{
-									++MatchCount[Index];
-								}
-								else
-								{
-									++UnexpectedRowCount;
-								}
-							}
-						};
-
-						using namespace UE::Editor::DataStorage::Queries;
-						// Setup an activatable processor for the test
-						ActivationKeys.Emplace(TEXT("TEST: Match Row[0]"));
-						QueryExpectedMatchRows.Emplace(TArray({Rows[0]}));
-						RegisterQuery(
-							Select(
-								ActivationKeys.Last(),
-								FProcessor(EQueryTickPhase::FrameEnd, TedsInterface->GetQueryTickGroupName(EQueryTickGroups::SyncDataStorageToExternal))
-									.MakeActivatable(ActivationKeys.Last()),
-									CopyTemp(Callback)
-							)
-							.Where()
-								.All<FTestDynamicTag>(Identifiers[0])
-								.All<FTestDynamicTag>(Identifiers[1])
-							.Compile());
-
-						ActivationKeys.Emplace(TEXT("TEST: Match Row[0] and Row[1]"));
-						QueryExpectedMatchRows.Emplace(TArray({Rows[0], Rows[1]}));
-						RegisterQuery(
-							Select(
-								ActivationKeys.Last(),
-								FProcessor(EQueryTickPhase::FrameEnd, TedsInterface->GetQueryTickGroupName(EQueryTickGroups::SyncDataStorageToExternal))
-									.MakeActivatable(ActivationKeys.Last()),
-									CopyTemp(Callback)
-							)
-							.Where()
-								.All<FTestDynamicTag>(Identifiers[0])
-							.Compile());
-
-						ActivationKeys.Emplace(TEXT("TEST: Match Row[1]"));
-						QueryExpectedMatchRows.Emplace(TArray({Rows[1]}));
-						RegisterQuery(
-							Select(
-								ActivationKeys.Last(),
-								FProcessor(EQueryTickPhase::FrameEnd, TedsInterface->GetQueryTickGroupName(EQueryTickGroups::SyncDataStorageToExternal))
-									.MakeActivatable(ActivationKeys.Last()),
-									CopyTemp(Callback)
-							)
-							.Where()
-								.All<FTestDynamicTag>(Identifiers[0])
-								.None<FTestDynamicTag>(Identifiers[1])
-							.Compile());
-
-						ActivationKeys.Emplace(TEXT("TEST: Match Row[0] and Row[2]"));
-						QueryExpectedMatchRows.Emplace(TArray({Rows[0], Rows[2]}));
-						RegisterQuery(
-							Select(
-								ActivationKeys.Last(),
-								FProcessor(EQueryTickPhase::FrameEnd, TedsInterface->GetQueryTickGroupName(EQueryTickGroups::SyncDataStorageToExternal))
-									.MakeActivatable(ActivationKeys.Last()),
-									CopyTemp(Callback)
-							)
-							.Where()
-								.All<FTestDynamicTag>(Identifiers[1])
-							.Compile());
-
-						ActivationKeys.Emplace(TEXT("TEST: Match Row[2]"));
-						QueryExpectedMatchRows.Emplace(TArray({Rows[2]}));
-						RegisterQuery(
-							Select(
-								ActivationKeys.Last(),
-								FProcessor(EQueryTickPhase::FrameEnd, TedsInterface->GetQueryTickGroupName(EQueryTickGroups::SyncDataStorageToExternal))
-									.MakeActivatable(ActivationKeys.Last()),
-									CopyTemp(Callback)
-							)
-							.Where()
-								.None<FTestDynamicTag>(Identifiers[0])
-								.All<FTestDynamicTag>(Identifiers[1])
-							.Compile());
-
-						ActivationKeys.Emplace(TEXT("TEST: Match all rows"));
-						QueryExpectedMatchRows.Emplace(TArray({Rows[0], Rows[1], Rows[2]}));
-						RegisterQuery(
-							Select(
-								ActivationKeys.Last(),
-								FProcessor(EQueryTickPhase::FrameEnd, TedsInterface->GetQueryTickGroupName(EQueryTickGroups::SyncDataStorageToExternal))
-									.MakeActivatable(ActivationKeys.Last()),
-									CopyTemp(Callback)
-							)
-							.Where()
-								.Any<FTestDynamicTag>(Identifiers[0])
-								.Any<FTestDynamicTag>(Identifiers[1])
-							.Compile());
-					}).Wait();
-
-					for (int32 TestIndex = 0; TestIndex < ActivationKeys.Num(); ++TestIndex)
-					{
-						SetExpectedMatches(QueryExpectedMatchRows[TestIndex]);
-						// Kick off the activation
-						TFuture<void> Future = Async(EAsyncExecution::TaskGraphMainThread, [this, &ActivationKeys, TestIndex]()
-						{
-							TedsInterface->ActivateQueries(ActivationKeys[TestIndex]);
-						});
-						Future.Wait();
-
-						// Wait for execution to occur
-						WaitFrames();
-						
-						// Check results
-						TStringBuilder<256> TestMessage;
-						TestMessage.Appendf(TEXT("'%s': Check that no rows processed that weren't expected"), *ActivationKeys[TestIndex].ToString());
-						TestEqual(TestMessage.ToString(), UnexpectedRowCount, static_cast<int32>(0));
-
-						TestMessage.Reset();
-						for (int32 RowIndex = 0; RowIndex < MatchCount.Num(); ++RowIndex)
-						{
-							TestMessage.Appendf(TEXT("'%s': Check row matched 1 times"), *ActivationKeys[TestIndex].ToString());
-							TestEqual(TestMessage.ToString(), MatchCount[RowIndex], 1);
+							TestEqual("Match All Rows using template", Result.Count, GetMatchCount());
 						}
 					}
 				}
-
-				WaitFrames();
 			});
 
-			It(TEXT("Columns"), EAsyncExecution::ThreadPool, [this]()
+			It(TEXT("Columns"), EAsyncExecution::TaskGraphMainThread, [this]()
 			{
 				// Add dynamic columns that have data
-				Async(EAsyncExecution::TaskGraphMainThread,[this]()
 				{
 					for (int32 Index = 0; Index <= 2; ++Index)
 					{
@@ -440,7 +222,7 @@ namespace UE::Editor::DataStorage::Tests
 
 					TedsInterface->AddColumn(Rows[2], Identifiers[1], FTestDynamicColumn
 					{
-						.IntArray = {11, 22, 33, 44},
+						.IntArray = {11, 22, 33, 43},
 					});
 
 					// Check they were added
@@ -464,135 +246,23 @@ namespace UE::Editor::DataStorage::Tests
 
 					TestFalse("Unexpected columns found", TedsInterface->HasColumns(Rows[2], MakeConstArrayView({DynamicTagTypes[0]})));
 					TestTrue("Expected columns not found", TedsInterface->HasColumns(Rows[2], MakeConstArrayView({DynamicTagTypes[1]})));
-				}).Wait();
-				
-				// Processor Query
-				{
-					TArray<FName> ActivationKeys;
-					
-					struct FExpectation
-					{
-						RowHandle Row;
-						TArray<int32> ColumnValues;
-
-						int32 MatchCount = 0;
-						bool bValuesMatch = false;
-					};
-					
-					TArray<TArray<FExpectation>> AllTestExpectations;
-
-					// Expectation for the current test
-					int32 UnexpectedRowCount;
 					
-					Async(EAsyncExecution::TaskGraphMainThread,[this, &ActivationKeys, &AllTestExpectations, &UnexpectedRowCount]()
-					{
-						using namespace UE::Editor::DataStorage::Queries;
-
-						auto RunTest = [this, &AllTestExpectations, &UnexpectedRowCount](IQueryContext& Context, const RowHandle* RowsPtr, int32 TestIndex)
-						{
-							TConstArrayView<RowHandle> RowView = MakeConstArrayView(RowsPtr, Context.GetRowCount());
-							TConstArrayView<FTestDynamicColumn> ColumnView = MakeConstArrayView(
-								Context.GetColumn<FTestDynamicColumn>(Identifiers[0]),
-								Context.GetRowCount());
+					using namespace UE::Editor::DataStorage::Queries;
 
-							TArray<FExpectation>& Expectations = AllTestExpectations[TestIndex];
-							for (int32 RowIndex = 0, End = Context.GetRowCount(); RowIndex < End; ++RowIndex)
-							{
-								RowHandle Row = RowView[RowIndex];
-																		
-								FExpectation* Expectation = Expectations.FindByPredicate([Row](const FExpectation& Expectation)
-								{
-									return Expectation.Row == Row;
-								});
-								if (Expectation)
-								{
-									Expectation->MatchCount++;
-									// Check values
-									const FTestDynamicColumn& Column = ColumnView[RowIndex];
-									if (Column.IntArray.Num() == Expectation->ColumnValues.Num())
-									{
-										bool bAllMatch = true;
-										for (int32 ColumnDataIndex = 0; ColumnDataIndex < Column.IntArray.Num(); ++ColumnDataIndex)
-										{
-											if (Column.IntArray[ColumnDataIndex] != Expectation->ColumnValues[ColumnDataIndex])
-											{
-												bAllMatch = false;
-												break;
-											}
-										}
-										Expectation->bValuesMatch = bAllMatch;
-									}
-								}
-								else
-								{
-									++UnexpectedRowCount;
-								}
-							}
-						};
-						
-						// Setup an activatable processor for the test
-						ActivationKeys.Emplace(TEXT("TEST: Match Row[0]"));
-						AllTestExpectations.Add(
-							{
-								{
-									.Row = Rows[0],
-									.ColumnValues = {1, 2, 3}
-								},
-								{
-									.Row = Rows[1],
-									.ColumnValues = {14, 15, 16}
-								}
-							});
-
-						RegisterQuery(
-							Select(
-								ActivationKeys.Last(),
-								FProcessor(EQueryTickPhase::FrameEnd, TedsInterface->GetQueryTickGroupName(EQueryTickGroups::SyncDataStorageToExternal))
-									.MakeActivatable(ActivationKeys.Last()),
-								[this, TestIndex = AllTestExpectations.Num() - 1, &RunTest](IQueryContext& Context, const RowHandle* RowsPtr)
-								{
-									RunTest(Context, RowsPtr, TestIndex);
-								})
-							.ReadOnly<FTestDynamicColumn>(Identifiers[0])
-							.Compile());
-					});
-
-					for (int32 TestIndex = 0; TestIndex < ActivationKeys.Num(); ++TestIndex)
 					{
-						// Kick off the activation
-						UnexpectedRowCount = 0;
-						TFuture<void> Future = Async(EAsyncExecution::TaskGraphMainThread,
-							[this, &ActivationKeys, TestIndex]()
-							{
-								TedsInterface->ActivateQueries(ActivationKeys[TestIndex]);
-							});
-						Future.Wait();
-						
-						WaitFrames();
-						
-						// Check results
-						TStringBuilder<256> TestMessage;
-						TestMessage.Appendf(TEXT("'%s': Check that no rows processed that weren't expected"), *ActivationKeys[TestIndex].ToString());
-						TestEqual(TestMessage.ToString(), UnexpectedRowCount, static_cast<int32>(0));
-
-						TestMessage.Reset();
-						
-						TArray<FExpectation>& TestExpectations = AllTestExpectations[TestIndex];
-						for (int32 RowIndex = 0; RowIndex < TestExpectations.Num(); ++RowIndex)
-						{
-							TestMessage.Appendf(TEXT("'%s': Each row matched 1 times"), *ActivationKeys[TestIndex].ToString());
-							TestEqual(TestMessage.ToString(), TestExpectations[RowIndex].MatchCount, 1);
+						QueryHandle Query = RegisterQuery(
+									Select().
+									Where(TColumn<FTestDynamicColumn>()).
+									Compile());
 
-							TestMessage.Reset();
-							TestMessage.Appendf(TEXT("'%s': Column values are expected"));
-							TestEqual(TestMessage.ToString(), TestExpectations[RowIndex].bValuesMatch, true);
-						}
+						FQueryResult Result = TedsInterface->RunQuery(Query);
+						TestEqual("Unexpected number of rows queried", Result.Count, 3);
 					}
 				}
 			});
 		});
 		
-		xAfterEach([this]()
+		AfterEach([this]()
 		{
 			Identifiers.Empty(Identifiers.Num());
 			for (RowHandle Row : Rows)
diff --git a/Engine/Plugins/Experimental/EditorDataStorage/Source/UhtExtension/EditorDataStorageUhtExtension.cs b/Engine/Plugins/Experimental/EditorDataStorage/Source/UhtExtension/EditorDataStorageUhtExtension.cs
new file mode 100644
--- /dev/null
+++ b/Engine/Plugins/Experimental/EditorDataStorage/Source/UhtExtension/EditorDataStorageUhtExtension.cs
@@ -0,0 +1,28 @@
+// Copyright Epic Games, Inc. All Rights Reserved.
+
+using System;
+using System.Text;
+using EpicGames.UHT.Tables;
+using EpicGames.UHT.Utils;
+using EpicGames.UHT.Exporters.CodeGen;
+using EpicGames.UHT.Types;
+
+namespace EditorDataStorageUhtExtension
+{
+	[UnrealHeaderTool]
+	class Extension
+    {
+	    private const string DynamicTemplateMetadata = "EditorDataStorage_DynamicColumnTemplate";
+	    
+	    [UhtCodeGeneratorInjector(
+		    UhtType = typeof(UhtScriptStruct), 
+		    Location = UhtCodeGeneratorInjectionLocation.GeneratedMacro)]
+	    public static void InjectDynamicColumnTemplateStaticAttribute(StringBuilder builder, UhtType uhtType, int leadingTabs, string eolSequence)
+	    {
+		    if (uhtType.MetaData.ContainsKey(DynamicTemplateMetadata))
+		    {
+			    builder.Append('\t', leadingTabs).Append("struct EditorDataStorage_DynamicColumnTemplate{};").Append(eolSequence);
+		    }
+	    }
+    }
+}
\ No newline at end of file
diff --git a/Engine/Plugins/Experimental/EditorDataStorage/Source/UhtExtension/EditorDataStorageUhtExtension.ubtplugin.csproj b/Engine/Plugins/Experimental/EditorDataStorage/Source/UhtExtension/EditorDataStorageUhtExtension.ubtplugin.csproj
new file mode 100644
--- /dev/null
+++ b/Engine/Plugins/Experimental/EditorDataStorage/Source/UhtExtension/EditorDataStorageUhtExtension.ubtplugin.csproj
@@ -0,0 +1,50 @@
+<Project Sdk="Microsoft.NET.Sdk">
+  <Import Project="EditorDataStorageUhtExtension.ubtplugin.csproj.props" Condition="Exists('EditorDataStorageUhtExtension.ubtplugin.csproj.props')" />
+  <Import Project="$(EngineDir)\Source\Programs\Shared\UnrealEngine.csproj.props" />
+  
+  <PropertyGroup>
+    <TargetFramework>net8.0</TargetFramework>
+    <Configuration Condition=" '$(Configuration)' == '' ">Development</Configuration>
+    <OutputType>Library</OutputType>
+    <AppendTargetFrameworkToOutputPath>false</AppendTargetFrameworkToOutputPath>
+    <GenerateAssemblyInfo>false</GenerateAssemblyInfo>
+    <GenerateTargetFrameworkAttribute>false</GenerateTargetFrameworkAttribute>
+    <Configurations>Debug;Release;Development</Configurations>
+    <RootNamespace>EditorDataStorageUhtExtension</RootNamespace>
+    <AssemblyName>EditorDataStorageUhtExtension</AssemblyName>
+    <OutputPath>..\..\..\..\..\Binaries\DotNET\UnrealBuildTool\Plugins\EditorDataStorage\</OutputPath>
+    <TreatWarningsAsErrors>true</TreatWarningsAsErrors>
+    <DebugType>pdbonly</DebugType>
+    <Nullable>enable</Nullable>
+  </PropertyGroup>
+  <PropertyGroup Condition=" '$(Configuration)|$(Platform)' == 'Development|AnyCPU' ">
+    <DefineConstants>$(DefineConstants);TRACE</DefineConstants>
+    <Optimize>true</Optimize>
+  </PropertyGroup>
+  <PropertyGroup Condition=" '$(Configuration)|$(Platform)' == 'Release|AnyCPU' ">
+    <DefineConstants>$(DefineConstants);TRACE</DefineConstants>
+    <Optimize>true</Optimize>
+  </PropertyGroup>
+  <PropertyGroup Condition=" '$(Configuration)|$(Platform)' == 'Debug|AnyCPU' ">
+    <DefineConstants>$(DefineConstants);DEBUG;TRACE</DefineConstants>
+    <DebugSymbols>true</DebugSymbols>
+    <DebugType>full</DebugType>
+  </PropertyGroup>
+  <ItemGroup>
+    <PackageReference Include="Microsoft.CSharp" Version="4.7.0" />
+  </ItemGroup>
+  <ItemGroup>
+    <Reference Include="EpicGames.Build">
+      <HintPath>$(EngineDir)\Binaries\DotNET\UnrealBuildTool\EpicGames.Build.dll</HintPath>
+    </Reference>
+    <Reference Include="EpicGames.Core">
+      <HintPath>$(EngineDir)\Binaries\DotNET\UnrealBuildTool\EpicGames.Core.dll</HintPath>
+    </Reference>
+    <Reference Include="EpicGames.UHT">
+      <HintPath>$(EngineDir)\Binaries\DotNET\UnrealBuildTool\EpicGames.UHT.dll</HintPath>
+    </Reference>
+    <Reference Include="UnrealBuildTool">
+      <HintPath>$(EngineDir)\Binaries\DotNET\UnrealBuildTool\UnrealBuildTool.dll</HintPath>
+    </Reference>
+  </ItemGroup>
+</Project>
\ No newline at end of file
diff --git a/Engine/Plugins/Experimental/EditorDataStorageFeatures/Source/TedsAssetData/Private/Widgets/SkeletalMeshSkeletonWidget.cpp b/Engine/Plugins/Experimental/EditorDataStorageFeatures/Source/TedsAssetData/Private/Widgets/SkeletalMeshSkeletonWidget.cpp
--- a/Engine/Plugins/Experimental/EditorDataStorageFeatures/Source/TedsAssetData/Private/Widgets/SkeletalMeshSkeletonWidget.cpp
+++ b/Engine/Plugins/Experimental/EditorDataStorageFeatures/Source/TedsAssetData/Private/Widgets/SkeletalMeshSkeletonWidget.cpp
@@ -54,11 +54,11 @@ TSharedPtr<SWidget> FSkeletalMeshSkeletonWidgetConstructor::CreateWidget(IEditor
 			.AutoWidth()
 			[
 				SNew(SHyperlink)
-				.Text(Binder.BindData(&FItemStringAttributeColumn_Experimental::Value, [](const FString& SkeletonPath)
+				.Text(Binder.BindData(FName("Skeleton"), &FItemStringAttributeColumn_Experimental::Value, [](const FString& SkeletonPath)
 				{
 					return FText::FromString(UE::SkeletalMeshSkeletonWidget::Private::GetAssetName(SkeletonPath));
-				}, FString(), FName("Skeleton")))
-				.ToolTipText(Binder.BindText(&FItemStringAttributeColumn_Experimental::Value, FName("Skeleton")))
+				}, FString()))
+				.ToolTipText(Binder.BindText(FName("Skeleton"), &FItemStringAttributeColumn_Experimental::Value))
 				.OnNavigate_Static(&UE::SkeletalMeshSkeletonWidget::Private::NavigateToAsset, DataStorage, TargetRow)
 				.Style(FAppStyle::Get(), "Common.GotoBlueprintHyperlink")
 			];
diff --git a/Engine/Plugins/Experimental/EditorDataStorageFeatures/Source/TedsAssetData/Private/Widgets/StaticMeshTrianglesWidget.cpp b/Engine/Plugins/Experimental/EditorDataStorageFeatures/Source/TedsAssetData/Private/Widgets/StaticMeshTrianglesWidget.cpp
--- a/Engine/Plugins/Experimental/EditorDataStorageFeatures/Source/TedsAssetData/Private/Widgets/StaticMeshTrianglesWidget.cpp
+++ b/Engine/Plugins/Experimental/EditorDataStorageFeatures/Source/TedsAssetData/Private/Widgets/StaticMeshTrianglesWidget.cpp
@@ -27,6 +27,6 @@ TSharedPtr<SWidget> FStaticMeshTrianglesWidgetConstructor::CreateWidget(IEditorD
 	UE::Editor::DataStorage::FAttributeBinder Binder(TargetRow, DataStorage);
 	
 	return SNew(STextBlock)
-			.Text(Binder.BindText(&FItemStringAttributeColumn_Experimental::Value, FName("Triangles")));
+			.Text(Binder.BindText(FName("Triangles"), &FItemStringAttributeColumn_Experimental::Value));
 			
 }
diff --git a/Engine/Plugins/Experimental/EditorDataStorageFeatures/Source/TedsAssetData/Public/TedsAssetDataColumns.h b/Engine/Plugins/Experimental/EditorDataStorageFeatures/Source/TedsAssetData/Public/TedsAssetDataColumns.h
--- a/Engine/Plugins/Experimental/EditorDataStorageFeatures/Source/TedsAssetData/Public/TedsAssetDataColumns.h
+++ b/Engine/Plugins/Experimental/EditorDataStorageFeatures/Source/TedsAssetData/Public/TedsAssetDataColumns.h
@@ -175,7 +175,7 @@ struct FItemTextAttributeColumn_Experimental : public FItemAttributeBaseColumn_E
 	FText Value;
 };
 
-USTRUCT()
+USTRUCT(meta=(EditorDataStorage_DynamicColumnTemplate))
 struct FItemStringAttributeColumn_Experimental : public FItemAttributeBaseColumn_Experimental
 {
 	GENERATED_BODY()
diff --git a/Engine/Source/Runtime/TypedElementFramework/Private/Elements/Common/TypedElementQueryConditions.cpp b/Engine/Source/Runtime/TypedElementFramework/Private/Elements/Common/TypedElementQueryConditions.cpp
--- a/Engine/Source/Runtime/TypedElementFramework/Private/Elements/Common/TypedElementQueryConditions.cpp
+++ b/Engine/Source/Runtime/TypedElementFramework/Private/Elements/Common/TypedElementQueryConditions.cpp
@@ -2,11 +2,15 @@
 
 #include "Elements/Common/TypedElementQueryConditions.h"
 
+#include "RHIStats.h"
 #include "Elements/Interfaces/TypedElementDataStorageInterface.h"
 #include "UObject/Class.h"
 
 namespace UE::Editor::DataStorage::Queries
 {
+	static const FName NAME_DynamicTemplateMetadata(TEXT("EditorDataStorage_DynamicColumnTemplate"));
+	static const FName NAME_DerivedFromDynamicTemplateMetadata(TEXT("EditorDataStorage_DerivedFromDynamicTemplate"));
+	
 	FEditorStorageQueryConditionCompileContext::FEditorStorageQueryConditionCompileContext(IEditorDataStorageProvider* InDataStorage)
 		: DataStorage(InDataStorage)
 	{
@@ -81,22 +85,103 @@ namespace UE::Editor::DataStorage::Queries
 		}
 	}
 
-	bool FConditions::Verify(TConstArrayView<FColumnBase> AvailableColumns) const
+	template <typename AvailableColumnType, typename ProjectionFunctionType>
+	bool FConditions::VerifyWithDynamicColumn(TConstArrayView<AvailableColumnType> AvailableColumns, uint64& Matches, ProjectionFunctionType Projection) const
 	{
-		checkf(bIsCompiled, TEXT("Query Conditions must call Compile() before you can use them"));
+		// Upfront process AvailableColumns since AvailableColumns may be searched through more than once
+		TBitArray<TInlineAllocator<512>> DerivedFromDynamicTemplate;
+#if WITH_EDITORONLY_DATA
+		DerivedFromDynamicTemplate.SetNum(AvailableColumns.Num(), false);
+		const FName DerivedFromDynamicTemplateMetadata = NAME_DerivedFromDynamicTemplateMetadata;
+		for (int32 AvailableColumnIndex = 0; AvailableColumnIndex < AvailableColumns.Num(); ++AvailableColumnIndex)
+		{
+			const TWeakObjectPtr<const UScriptStruct>& AvailableColumn = Projection(AvailableColumns[AvailableColumnIndex]);
+			if (const UScriptStruct* AvailableColumnStruct = AvailableColumn.Get())
+			{
+				DerivedFromDynamicTemplate[AvailableColumnIndex] = AvailableColumnStruct->FindMetaData(DerivedFromDynamicTemplateMetadata) != nullptr;
+			}
+		}
+#endif
 		
-		return VerifyBootstrap(
-			[&AvailableColumns](uint8_t ColumnIndex, TWeakObjectPtr<const UScriptStruct> Column)
+		bool Result = VerifyBootstrap(
+			[this, &Matches, &AvailableColumns, &DerivedFromDynamicTemplate, &Projection](uint8_t ColumnIndex, TWeakObjectPtr<const UScriptStruct> Column)
 			{
-				for (const FColumnBase& Target : AvailableColumns)
+				bool bContains = false;
+				// Looping through each column in the condition, if the column is a template then determine if any of the available
+				// columns that are being matched against are derived from that template
+				if ((ColumnFlags[ColumnIndex] & EColumnFlags::DynamicColumnTemplate) == EColumnFlags::DynamicColumnTemplate)
 				{
-					if (Target.TypeInfo == Column)
+					for (int32 AvailableColumnIndex = 0; AvailableColumnIndex < AvailableColumns.Num(); ++AvailableColumnIndex)
 					{
-						return true;
+						// Only check if this available column is derived from our dynamic template if we know it derives from any DynamicTemplate
+						if (DerivedFromDynamicTemplate[AvailableColumnIndex])
+						{
+							if (Projection(AvailableColumns[AvailableColumnIndex])->IsChildOf(Column.Get()))
+							{
+								Matches |= (uint64)1 << ColumnIndex;
+								bContains = true;
+							}
+						}
+						else
+						{
+							if (Projection(AvailableColumns[AvailableColumnIndex]) == Column.Get())
+							{
+								Matches |= (uint64)1 << ColumnIndex;
+								bContains = true;
+							}
+						}
 					}
 				}
-		return false;
+				else
+				{
+					for (int32 AvailableColumnIndex = 0; AvailableColumnIndex < AvailableColumns.Num(); ++AvailableColumnIndex)
+					{
+						if (Projection(AvailableColumns[AvailableColumnIndex]) == Column.Get())
+						{
+							Matches |= (uint64)1 << ColumnIndex;
+							bContains = true;
+						}
+					}
+					//auto FindPredicate = ;
+					if ( AvailableColumns.FindByPredicate([&Column, &Projection](const AvailableColumnType& Entry) { return Projection(Entry) == Column.Get(); }) != nullptr)
+					{
+						Matches |= (uint64)1 << ColumnIndex;
+						bContains = true;
+					}
+				}
+				return bContains;
 			});
+
+		return Result;
+	}
+	
+	bool FConditions::Verify(TConstArrayView<FColumnBase> AvailableColumns) const
+	{
+		checkf(bIsCompiled, TEXT("Query Conditions must call Compile() before you can use them"));
+
+		bool Result;
+		if (!UsesDynamicTemplates())
+		{
+			Result = VerifyBootstrap(
+				[&AvailableColumns](uint8_t ColumnIndex, TWeakObjectPtr<const UScriptStruct> Column)
+				{
+					for (const FColumnBase& Target : AvailableColumns)
+					{
+						if (Target.TypeInfo == Column)
+						{
+							return true;
+						}
+					}
+					return false;
+				});
+		}
+		else
+		{
+			uint64 Matches = 0;
+			Result = VerifyWithDynamicColumn(AvailableColumns, Matches, [](const FColumnBase& Column) { return Column.TypeInfo; });
+		}
+		
+		return Result;
 	}
 
 	bool FConditions::Verify(TArray<TWeakObjectPtr<const UScriptStruct>>& MatchedColumns, TConstArrayView<FColumnBase> AvailableColumns,
@@ -107,34 +192,42 @@ namespace UE::Editor::DataStorage::Queries
 		checkf(bIsCompiled, TEXT("Query Conditions must call Compile() before you can use them"));
 		
 		uint64 Matches = 0;
-		bool Result = AvailableColumnsAreSorted
-			? VerifyBootstrap(
-				[&Matches, &AvailableColumns](uint8_t ColumnIndex, TWeakObjectPtr<const UScriptStruct> Column)
-				{
-					auto Projection = [](const FColumnBase& ColumnBase)
-					{
-						return ColumnBase.TypeInfo.Get();
-					};
-					if (Algo::BinarySearchBy(AvailableColumns, Column.Get(), Projection) != INDEX_NONE)
-					{
-						Matches |= (uint64)1 << ColumnIndex;
-						return true;
-					}
-					return false;
-				})
-			: VerifyBootstrap(
-				[&Matches, &AvailableColumns](uint8_t ColumnIndex, TWeakObjectPtr<const UScriptStruct> Column)
-				{
-					for (const FColumnBase& Target : AvailableColumns)
+		bool Result;
+		if (!UsesDynamicTemplates())
+		{
+			Result = AvailableColumnsAreSorted
+				? VerifyBootstrap(
+					[&Matches, &AvailableColumns](uint8_t ColumnIndex, TWeakObjectPtr<const UScriptStruct> Column)
 					{
-						if (Target.TypeInfo == Column)
+						auto Projection = [](const FColumnBase& ColumnBase)
+						{
+							return ColumnBase.TypeInfo.Get();
+						};
+						if (Algo::BinarySearchBy(AvailableColumns, Column.Get(), Projection) != INDEX_NONE)
 						{
 							Matches |= (uint64)1 << ColumnIndex;
 							return true;
 						}
-					}
-					return false;
-				});
+						return false;
+					})
+				: VerifyBootstrap(
+					[&Matches, &AvailableColumns](uint8_t ColumnIndex, TWeakObjectPtr<const UScriptStruct> Column)
+					{
+						for (const FColumnBase& Target : AvailableColumns)
+						{
+							if (Target.TypeInfo == Column)
+							{
+								Matches |= (uint64)1 << ColumnIndex;
+								return true;
+							}
+						}
+						return false;
+					});
+		}
+		else
+		{
+			Result = VerifyWithDynamicColumn(AvailableColumns, Matches, [](const FColumnBase& ColumnBase) { return ColumnBase.TypeInfo; });
+		}
 		if (Result)
 		{
 			ConvertColumnBitToArray(MatchedColumns, Matches);
@@ -146,22 +239,32 @@ namespace UE::Editor::DataStorage::Queries
 		bool AvailableColumnsAreSorted) const
 	{
 		checkf(bIsCompiled, TEXT("Query Conditions must call Compile() before you can use them"));
-		
-		return AvailableColumnsAreSorted
-			? VerifyBootstrap(
-				[&AvailableColumns](uint8_t ColumnIndex, TWeakObjectPtr<const UScriptStruct> Column)
-				{
-					auto Projection = [](TWeakObjectPtr<const UScriptStruct> ColumnBase)
+
+		bool Result;
+		if (!UsesDynamicTemplates())
+		{
+			Result = AvailableColumnsAreSorted
+				? VerifyBootstrap(
+					[&AvailableColumns](uint8_t ColumnIndex, TWeakObjectPtr<const UScriptStruct> Column)
 					{
-						return ColumnBase.Get();
-					};
-					return Algo::BinarySearchBy(AvailableColumns, Column.Get(), Projection) != INDEX_NONE;
-				})
-			: VerifyBootstrap(
-				[&AvailableColumns](uint8_t ColumnIndex, TWeakObjectPtr<const UScriptStruct> Column)
-				{
-					return AvailableColumns.Find(Column) != INDEX_NONE;
-				});
+						auto Projection = [](TWeakObjectPtr<const UScriptStruct> ColumnBase)
+						{
+							return ColumnBase.Get();
+						};
+						return Algo::BinarySearchBy(AvailableColumns, Column.Get(), Projection) != INDEX_NONE;
+					})
+				: VerifyBootstrap(
+					[&AvailableColumns](uint8_t ColumnIndex, TWeakObjectPtr<const UScriptStruct> Column)
+					{
+						return AvailableColumns.Find(Column) != INDEX_NONE;
+					});
+		}
+		else
+		{
+			uint64 Matches = 0;
+			Result = VerifyWithDynamicColumn(AvailableColumns, Matches, [](TWeakObjectPtr<const UScriptStruct> ColumnType) { return ColumnType; });
+		}
+		return Result;
 	}
 
 	bool FConditions::Verify(TArray<TWeakObjectPtr<const UScriptStruct>>& MatchedColumns,
@@ -173,26 +276,39 @@ namespace UE::Editor::DataStorage::Queries
 		checkf(bIsCompiled, TEXT("Query Conditions must call Compile() before you can use them"));
 		
 		uint64 Matches = 0;
-		bool Result = AvailableColumnsAreSorted
-			? VerifyBootstrap(
-				[&Matches, &AvailableColumns](uint8_t ColumnIndex, TWeakObjectPtr<const UScriptStruct> Column)
-				{
-					auto Projection = [](TWeakObjectPtr<const UScriptStruct> ColumnBase)
+		bool Result;
+		if (!UsesDynamicTemplates())
+		{
+			Result = AvailableColumnsAreSorted
+				? VerifyBootstrap(
+					[&Matches, &AvailableColumns](uint8_t ColumnIndex, TWeakObjectPtr<const UScriptStruct> Column)
 					{
-						return ColumnBase.Get();
-					};
-					if (Algo::BinarySearchBy(AvailableColumns, Column.Get(), Projection) != INDEX_NONE)
+						auto Projection = [](TWeakObjectPtr<const UScriptStruct> ColumnBase)
+						{
+							return ColumnBase.Get();
+						};
+						if (Algo::BinarySearchBy(AvailableColumns, Column.Get(), Projection) != INDEX_NONE)
+						{
+							Matches |= (uint64)1 << ColumnIndex;
+							return true;
+						}
+						return false;
+					})
+				: VerifyBootstrap(
+					[&Matches, &AvailableColumns](uint8_t ColumnIndex, TWeakObjectPtr<const UScriptStruct> Column)
 					{
-						Matches |= (uint64)1 << ColumnIndex;
-						return true;
-					}
-					return false;
-				})
-			: VerifyBootstrap(
-				[&Matches, &AvailableColumns](uint8_t ColumnIndex, TWeakObjectPtr<const UScriptStruct> Column)
-				{
-					return AvailableColumns.Find(Column) != INDEX_NONE;
-				});
+						if ( AvailableColumns.Find(Column) != INDEX_NONE)
+						{
+							Matches |= (uint64)1 << ColumnIndex;
+							return true;
+						}
+						return false;
+					});
+		}
+		else
+		{
+			Result = VerifyWithDynamicColumn(AvailableColumns, Matches, [](TWeakObjectPtr<const UScriptStruct> ColumnType) { return ColumnType; });
+		}
 		if (Result)
 		{
 			ConvertColumnBitToArray(MatchedColumns, Matches);
@@ -253,6 +369,8 @@ namespace UE::Editor::DataStorage::Queries
 		{
 			return *this;
 		}
+		
+		ColumnFlags.SetNumZeroed(ColumnCount);
 
 		// Resolve all dynamic columns by generating their UScriptStruct
 		for (uint8_t ColumnIndex = 0; ColumnIndex < ColumnCount; ++ColumnIndex)
@@ -268,6 +386,16 @@ namespace UE::Editor::DataStorage::Queries
 				
 				Columns[ColumnIndex] = ResolvedDynamicColumn;
 			}
+			else
+			{
+#if WITH_EDITORONLY_DATA
+				bool bIsDynamicTemplate = Columns[ColumnIndex]->FindMetaData(NAME_DynamicTemplateMetadata) != nullptr;
+				if (bIsDynamicTemplate)
+				{
+					ColumnFlags[ColumnIndex] |= EColumnFlags::DynamicColumnTemplate;
+				}
+#endif
+			}
 		}
 
 		bIsCompiled = true;
@@ -447,6 +575,18 @@ namespace UE::Editor::DataStorage::Queries
 		return Result;
 	}
 
+	bool FConditions::UsesDynamicTemplates() const
+	{
+		for (EColumnFlags Flags : ColumnFlags)
+		{
+			if((Flags & EColumnFlags::DynamicColumnTemplate) == EColumnFlags::DynamicColumnTemplate)
+			{
+				return true;
+			}
+		}
+		return false;
+	}
+
 	void FConditions::AppendQuery(FConditions& Target, const FConditions& Source)
 	{
 		checkf(Target.ColumnCount + Source.ColumnCount < MaxColumnCount, TEXT("Too many columns in the query."));
@@ -553,4 +693,6 @@ namespace UE::Editor::DataStorage::Queries
 		Result.Tokens[Result.TokenCount++] = FConditions::Token::ScopeClose;
 		return Result;
 	}
+
+	ENUM_CLASS_FLAGS(FConditions::EColumnFlags)
 } // namespace UE::Editor::DataStorage::Queries
\ No newline at end of file
diff --git a/Engine/Source/Runtime/TypedElementFramework/Private/Elements/Framework/TypedElementQueryBuilder.cpp b/Engine/Source/Runtime/TypedElementFramework/Private/Elements/Framework/TypedElementQueryBuilder.cpp
--- a/Engine/Source/Runtime/TypedElementFramework/Private/Elements/Framework/TypedElementQueryBuilder.cpp
+++ b/Engine/Source/Runtime/TypedElementFramework/Private/Elements/Framework/TypedElementQueryBuilder.cpp
@@ -495,11 +495,14 @@ namespace UE::Editor::DataStorage::Queries
 		return *this;
 	}
 
-	Select& Select::ReadWrite(const FDynamicColumnDescription& Descirption)
+	Select& Select::ReadWrite(const FDynamicColumnDescription& Description)
 	{
-		Query.DynamicSelectionTypes.Emplace(Descirption);
-		Query.DynamicSelectionAccessTypes.Emplace(EQueryAccessType::ReadWrite);
-		Query.DynamicSelectionMetaData.Emplace(FColumnMetaData::EFlags::IsMutable);
+		if (ensureMsgf(!Description.Identifier.IsNone(), TEXT("Cannot pass special identifier None to select a specific dynamic column")))
+		{
+			Query.DynamicSelectionTypes.Emplace(Description);
+			Query.DynamicSelectionAccessTypes.Emplace(EQueryAccessType::ReadWrite);
+			Query.DynamicSelectionMetaData.Emplace(FColumnMetaData::EFlags::IsMutable);
+		}
 		return *this;
 	}
 
diff --git a/Engine/Source/Runtime/TypedElementFramework/Public/Elements/Common/TypedElementCommonTypes.h b/Engine/Source/Runtime/TypedElementFramework/Public/Elements/Common/TypedElementCommonTypes.h
--- a/Engine/Source/Runtime/TypedElementFramework/Public/Elements/Common/TypedElementCommonTypes.h
+++ b/Engine/Source/Runtime/TypedElementFramework/Public/Elements/Common/TypedElementCommonTypes.h
@@ -84,15 +84,19 @@ namespace UE
 		using ColumnListWithDataCallbackRef = TFunctionRef<void(void* Column, const UScriptStruct& ColumnType)>;
 		using ColumnCopyOrMoveCallback = void (*)(const UScriptStruct& ColumnType, void* Destination, void* Source);
 
-
-
+		template<typename T>
+		concept THasDynamicColumnTemplateSpecifier = std::is_empty_v<typename T::EditorDataStorage_DynamicColumnTemplate>;
+		
+		template<typename T>
+		concept TDynamicColumnTemplate = (UE::derived_from<T, FColumn> || UE::derived_from<T, FTag>) && THasDynamicColumnTemplateSpecifier<T>;
+		
 		// Template concepts to enforce type correctness.
 		template<typename T>
-		concept TDataColumnType = UE::derived_from<T, FColumn>;
+		concept TDataColumnType = UE::derived_from<T, FColumn> && !THasDynamicColumnTemplateSpecifier<T>;
 
 		template<typename T>
-		concept TTagColumnType = UE::derived_from<T, FTag>;
-
+		concept TTagColumnType = UE::derived_from<T, FTag> && !THasDynamicColumnTemplateSpecifier<T>;
+		
 		template<typename T>
 		concept TColumnType = TDataColumnType<T> || TTagColumnType<T>;
 
diff --git a/Engine/Source/Runtime/TypedElementFramework/Public/Elements/Common/TypedElementQueryConditions.h b/Engine/Source/Runtime/TypedElementFramework/Public/Elements/Common/TypedElementQueryConditions.h
--- a/Engine/Source/Runtime/TypedElementFramework/Public/Elements/Common/TypedElementQueryConditions.h
+++ b/Engine/Source/Runtime/TypedElementFramework/Public/Elements/Common/TypedElementQueryConditions.h
@@ -128,6 +128,9 @@ namespace UE::Editor::DataStorage::Queries
 		TYPEDELEMENTFRAMEWORK_API bool IsEmpty() const;
 
 	private:
+		template<typename AvailableColumnType, typename ProjectionFunction>
+		bool VerifyWithDynamicColumn(TConstArrayView<AvailableColumnType> AvailableColumns, uint64& OutMatches, ProjectionFunction Projection) const;
+		
 		void AppendName(FString& Output, TWeakObjectPtr<const UScriptStruct> TypeInfo) const;
 
 		bool EntersScopeNext(uint8_t Index) const;
@@ -141,6 +144,8 @@ namespace UE::Editor::DataStorage::Queries
 
 		uint8_t MinimumColumnMatchRequiredRange(uint8_t& Front) const;
 
+		bool UsesDynamicTemplates() const;
+
 		static void AppendQuery(FConditions& Target, const FConditions& Source);
 
 		static constexpr SIZE_T MaxColumnCount = 32;
@@ -154,11 +159,20 @@ namespace UE::Editor::DataStorage::Queries
 			ScopeOpen,
 			ScopeClose
 		};
+
+		enum class EColumnFlags : uint8_t
+		{
+			DynamicColumnTemplate = 1 << 0
+		};
+		FRIEND_ENUM_CLASS_FLAGS(EColumnFlags);
+		
 		using ColumnArray = TWeakObjectPtr<const UScriptStruct>[MaxColumnCount];
 		using TokenArray = Token[MaxTokenCount];
 		using IdentifierArray = TArray<FName, TInlineAllocator<MaxColumnCount>>;
+		using ColumnFlagArray = TArray<EColumnFlags, TInlineAllocator<MaxColumnCount>>;
 
 		ColumnArray Columns{ nullptr };
+		ColumnFlagArray ColumnFlags;
 		TokenArray Tokens{ Token::None };
 		IdentifierArray Identifiers;
 		uint8_t ColumnCount = 0;
diff --git a/Engine/Source/Runtime/TypedElementFramework/Public/Elements/Framework/TypedElementAttributeBinding.h b/Engine/Source/Runtime/TypedElementFramework/Public/Elements/Framework/TypedElementAttributeBinding.h
--- a/Engine/Source/Runtime/TypedElementFramework/Public/Elements/Framework/TypedElementAttributeBinding.h
+++ b/Engine/Source/Runtime/TypedElementFramework/Public/Elements/Framework/TypedElementAttributeBinding.h
@@ -36,16 +36,28 @@ namespace UE::Editor::DataStorage
 		 * 
 		 * @param InVariable The data member inside a column to be bound
 		 * @param InDefaultValue The default value to be used when the column isn't present on a row
-		 * @param InIdentifier The identifier for this column if it is a dynamic column, NAME_None if it is not a dynamic column
-		 * @return A TAttribute bound to the row, column pair specified
 		 *
 		 * Example:
 		 * FAttributeBinder Binder(RowHandle);
 		 * TAttribute<FString> TestAttribute(Binder.BindData(&FTypedElementLabelColumn::Label))
 		 */
 		template <typename AttributeType, TDataColumnType ColumnType>
-		TAttribute<AttributeType> BindData(AttributeType ColumnType::* InVariable, const AttributeType& InDefaultValue = AttributeType(),
-			const FName& InIdentifier = NAME_None);
+		TAttribute<AttributeType> BindData(AttributeType ColumnType::* InVariable, const AttributeType& InDefaultValue = AttributeType());
+
+		/**
+		 * Bind a specific data member inside a TEDS column to an attribute of the same type as the data
+		 * 
+		 * @param InVariable The data member inside a column to be bound
+		 * @param InDefaultValue The default value to be used when the column isn't present on a row
+		 * @param InIdentifier The identifier for this dynamic column
+		 * @return A TAttribute bound to the row, column pair specified
+		 *
+		 * Example:
+		 * FAttributeBinder Binder(RowHandle);
+		 * TAttribute<FString> TestAttribute(Binder.BindData(&FTypedElementLabelColumn::Label, FName(TEXT("Id"))))
+		 */
+		template <typename AttributeType, TDynamicColumnTemplate ColumnType>
+		TAttribute<AttributeType> BindData(const FName& InIdentifier, AttributeType ColumnType::* InVariable, const AttributeType& InDefaultValue = AttributeType());
 
 		/**
 		 * Bind a specific data member inside a TEDS column to an attribute of a different type than the data by providing a conversion function
@@ -67,7 +79,40 @@ namespace UE::Editor::DataStorage
 		 *                                 ));
 		 */
 		template <typename AttributeType, typename DataType, TDataColumnType ColumnType>
-		TAttribute<AttributeType> BindData(DataType ColumnType::* InVariable, TFunction<AttributeType(const DataType&)> InConverter, const DataType& InDefaultValue = DataType(), const FName& InIdentifier = NAME_None);
+		TAttribute<AttributeType> BindData(DataType ColumnType::* InVariable, TFunction<AttributeType(const DataType&)> InConverter, const DataType& InDefaultValue = DataType());
+
+		/**
+		 * Bind a specific data member inside a TEDS column to an attribute of a different type than the data by providing a conversion function
+		 * NOTE: the default value is not the actual attribute type but rather the data type in the column and it gets passed to the conversion function
+		 *
+		 * @param InVariable The data member inside a column to be bound
+		 * @param InConverter Conversion function to convert from DataType -> AttributeType
+		*  @param InIdentifier The identifier for this column if it is a dynamic column
+		 * @param InDefaultValue The default value to be used when the column isn't present on a row
+		 * @return A TAttribute bound to the row, column pair specified
+		 *
+		 * Example:
+		 * UE::Editor::DataStorage::FAttributeBinder Binder(RowHandle);
+		 * TAttribute<FText> TestAttribute(Binder.BindData(&FTypedElementLabelColumn::Label),
+		 *                                 [](const FString& Data)
+		 *                                   {
+		 *                                      return FText::FromString(Data);
+		 *                                   }
+		 *                                 ));
+		 */
+		template <typename AttributeType, typename DataType, TDynamicColumnTemplate ColumnType>
+		TAttribute<AttributeType> BindData(const FName& InIdentifier, DataType ColumnType::* InVariable, TFunction<AttributeType(const DataType&)> InConverter, const DataType& InDefaultValue = DataType());
+
+		/**
+		 * Overload for the conversion binder to accept lambdas instead of TFunctions
+		 *
+		 * @param InVariable The data member inside a column to be bound
+		 * @param InConverter Conversion function to convert from DataType -> AttributeType
+		 * @param InDefaultValue The default value to be used when the column isn't present on a row\
+		 */
+		template <typename DataType, TDataColumnType ColumnType, typename FunctionType>
+			requires Private::AttributeBinderInvocable<FunctionType, DataType>
+		auto BindData(DataType ColumnType::* InVariable, FunctionType InConverter, const DataType& InDefaultValue = DataType());
 
 		/**
 		 * Overload for the conversion binder to accept lambdas instead of TFunctions
@@ -78,10 +123,9 @@ namespace UE::Editor::DataStorage
 		 * @param InIdentifier The identifier for this column if it is a dynamic column, NAME_None if it is not a dynamic column
 		 * @return A TAttribute bound to the row, column pair specified
 		 */
-		template <typename DataType, TDataColumnType ColumnType, typename FunctionType>
+		template <typename DataType, TDynamicColumnTemplate ColumnType, typename FunctionType>
 			requires Private::AttributeBinderInvocable<FunctionType, DataType>
-		auto BindData(DataType ColumnType::* InVariable, FunctionType InConverter, const DataType& InDefaultValue = DataType(),
-			const FName& InIdentifier = NAME_None);
+		auto BindData(const FName& InIdentifier, DataType ColumnType::* InVariable, FunctionType InConverter, const DataType& InDefaultValue = DataType());
 
 		/**
 		 * Bind a delegate inside a TEDS column to a SLATE_EVENT macro on a widget
@@ -90,8 +134,10 @@ namespace UE::Editor::DataStorage
 		 * @return A delegate that can be provided to an event on a slate widget
 		 */
 		template <typename InRetValType, typename... ParamTypes, TDataColumnType ColumnType>
-		TDelegate<InRetValType(ParamTypes...)> BindEvent(TDelegate<InRetValType(ParamTypes...)> ColumnType::* InVariable,
-			const FName& InIdentifier = NAME_None);
+		TDelegate<InRetValType(ParamTypes...)> BindEvent(TDelegate<InRetValType(ParamTypes...)> ColumnType::* InVariable);
+
+		template <typename InRetValType, typename... ParamTypes, TDynamicColumnTemplate ColumnType>
+		TDelegate<InRetValType(ParamTypes...)> BindEvent(const FName& InIdentifier, TDelegate<InRetValType(ParamTypes...)> ColumnType::* InVariable);
 
 		/**
 		 * Directly bind an FString member in a TEDS column to an FText attribute as a shortcut
@@ -100,7 +146,10 @@ namespace UE::Editor::DataStorage
 		 * @return A delegate that can be provided to a text widget in Slate (e.g STextBlock)
 		 */
 		template <TDataColumnType ColumnType>
-		TAttribute<FText> BindText(FString ColumnType::* InFStringVariable, const FName& InIdentifier = NAME_None);
+		TAttribute<FText> BindText(FString ColumnType::* InFStringVariable);
+		
+		template <TDynamicColumnTemplate ColumnType>
+		TAttribute<FText> BindText(const FName& InIdentifier, FString ColumnType::* InFStringVariable);
 
 		/**
 		 * Directly bind an FName member in a TEDS column to an FText attribute as a shortcut
@@ -109,7 +158,10 @@ namespace UE::Editor::DataStorage
 		 * @return A delegate that can be provided to a text widget in Slate (e.g STextBlock)
 		 */
 		template <TDataColumnType ColumnType>
-		TAttribute<FText> BindText(FName ColumnType::* InFNameVariable, const FName& InIdentifier = NAME_None);
+		TAttribute<FText> BindText(FName ColumnType::* InFNameVariable);
+		
+		template <TDynamicColumnTemplate ColumnType>
+		TAttribute<FText> BindText(const FName& InIdentifier, FName ColumnType::* InFNameVariable);
 
 		/**
 		 * Composite a FText attribute using a format string and named arguments bound to TEDS columns.
diff --git a/Engine/Source/Runtime/TypedElementFramework/Public/Elements/Framework/TypedElementAttributeBinding.inl b/Engine/Source/Runtime/TypedElementFramework/Public/Elements/Framework/TypedElementAttributeBinding.inl
--- a/Engine/Source/Runtime/TypedElementFramework/Public/Elements/Framework/TypedElementAttributeBinding.inl
+++ b/Engine/Source/Runtime/TypedElementFramework/Public/Elements/Framework/TypedElementAttributeBinding.inl
@@ -12,8 +12,34 @@ namespace UE::Editor::DataStorage
 	//
 
 	template <typename AttributeType, TDataColumnType ColumnType>
-	TAttribute<AttributeType> FAttributeBinder::BindData(AttributeType ColumnType::* InVariable, const AttributeType& InDefaultValue,
-		const FName& InIdentifier)
+	TAttribute<AttributeType> FAttributeBinder::BindData(AttributeType ColumnType::* InVariable, const AttributeType& InDefaultValue)
+	{
+		using namespace Private;
+
+		if(!DataStorage)
+		{
+			return TAttribute<AttributeType>();
+		}
+		
+		// Create a direct property and bind it to the given variable
+		TProperty<AttributeType> Prop;
+		Prop.Bind(InVariable);
+
+		// We don't want any references to this in the lambda because binders are designed to be used and destructed on the stack
+		return TAttribute<AttributeType>::CreateLambda([Property = MoveTemp(Prop), Storage = DataStorage, Row = TargetRow,
+			DefaultValue = InDefaultValue]()
+		{
+			// Get the column from the given row and use that to return the stored property
+			if(ColumnType* Column = Private::GetColumn<ColumnType>(Storage, Row))
+			{
+				return Property.Get(Column, ColumnType::StaticStruct());
+			}
+			return DefaultValue;
+		});
+	}
+
+	template <typename AttributeType, TDynamicColumnTemplate ColumnType>
+	TAttribute<AttributeType> FAttributeBinder::BindData(const FName& InIdentifier, AttributeType ColumnType::* InVariable, const AttributeType& InDefaultValue)
 	{
 		using namespace Private;
 
@@ -40,7 +66,41 @@ namespace UE::Editor::DataStorage
 	}
 
 	template <typename AttributeType, typename DataType, TDataColumnType ColumnType>
-	TAttribute<AttributeType> FAttributeBinder::BindData(DataType ColumnType::* InVariable, TFunction<AttributeType(const DataType&)> InConverter, const DataType& InDefaultValue, const FName& InIdentifier)
+	TAttribute<AttributeType> FAttributeBinder::BindData(DataType ColumnType::* InVariable, TFunction<AttributeType(const DataType&)> InConverter, const DataType& InDefaultValue)
+	{
+		using namespace Private;
+
+		if(!DataStorage)
+		{
+			return TAttribute<AttributeType>();
+		}
+		
+		AttributeType ConvertedDefault = InConverter(InDefaultValue);
+
+		// Create a convertible property and bind it to the given variable
+		TProperty<AttributeType> Prop;
+		Prop.Bind(InVariable, MoveTemp(InConverter));
+	
+		return TAttribute<AttributeType>::CreateLambda(
+			[
+				Property = MoveTemp(Prop), 
+				Storage = DataStorage, 
+				Row = TargetRow, 
+				DefaultValue = InDefaultValue, 
+				ConvertedDefault = MoveTemp(ConvertedDefault)
+			]()
+			{
+				if(ColumnType* Column = Private::GetColumn<ColumnType>(Storage, Row))
+				{
+					return Property.Get(Column, ColumnType::StaticStruct());
+				}
+
+				return ConvertedDefault;
+			});
+	}
+
+	template <typename AttributeType, typename DataType, TDynamicColumnTemplate ColumnType>
+	TAttribute<AttributeType> FAttributeBinder::BindData(const FName& InIdentifier, DataType ColumnType::* InVariable, TFunction<AttributeType(const DataType&)> InConverter, const DataType& InDefaultValue)
 	{
 		using namespace Private;
 
@@ -76,17 +136,53 @@ namespace UE::Editor::DataStorage
 	
 	template <typename DataType, TDataColumnType ColumnType, typename FunctionType>
 		requires Private::AttributeBinderInvocable<FunctionType, DataType>
-	auto FAttributeBinder::BindData(DataType ColumnType::* InVariable, FunctionType InConverter, const DataType& InDefaultValue, const FName& InIdentifier)
+	auto FAttributeBinder::BindData(DataType ColumnType::* InVariable, FunctionType InConverter, const DataType& InDefaultValue)
 	{
 		// Deduce the attribute type from the return value of the converter function
 		using AttributeType = decltype(InConverter(std::declval<DataType>()));
 		
-		return BindData<AttributeType, DataType>(InVariable, TFunction<AttributeType(const DataType&)>(MoveTemp(InConverter)), InDefaultValue, InIdentifier);
+		return BindData<AttributeType, DataType>(InVariable, TFunction<AttributeType(const DataType&)>(MoveTemp(InConverter)), InDefaultValue);
+	}
+	
+	template <typename DataType, TDynamicColumnTemplate ColumnType, typename FunctionType>
+		requires Private::AttributeBinderInvocable<FunctionType, DataType>
+	auto FAttributeBinder::BindData(const FName& InIdentifier, DataType ColumnType::* InVariable, FunctionType InConverter, const DataType& InDefaultValue)
+	{
+		// Deduce the attribute type from the return value of the converter function
+		using AttributeType = decltype(InConverter(std::declval<DataType>()));
+		
+		return BindData<AttributeType, DataType>(InIdentifier, InVariable, TFunction<AttributeType(const DataType&)>(MoveTemp(InConverter)), InDefaultValue);
 	}
 
 	template <typename InRetValType, typename... ParamTypes, TDataColumnType ColumnType>
-		TDelegate<InRetValType(ParamTypes...)> FAttributeBinder::BindEvent(TDelegate<InRetValType(ParamTypes...)> ColumnType::* InVariable,
-			const FName& InIdentifier)
+		TDelegate<InRetValType(ParamTypes...)> FAttributeBinder::BindEvent(TDelegate<InRetValType(ParamTypes...)> ColumnType::* InVariable)
+	{
+		using namespace Private;
+		
+		// Create a property for the delegate
+		TProperty<TDelegate<InRetValType(ParamTypes...)>> Prop;
+		Prop.Bind(InVariable);
+	
+		return TDelegate<InRetValType(ParamTypes...)>::CreateLambda([Property = MoveTemp(Prop), Storage = DataStorage,
+			Row = TargetRow](ParamTypes&&... Params)
+		{
+			if(ColumnType* Column = Private::GetColumn<ColumnType>(Storage, Row))
+			{
+				// Get the delegate in the bound column for the specified row
+				TDelegate<InRetValType(ParamTypes...)> Delegate = Property.Get(Column, ColumnType::StaticStruct());
+
+				// Execute the delegate if it is bound
+				if(Delegate.IsBound())
+				{
+					return Delegate.Execute(Forward<ParamTypes>(Params)...);
+				}
+			}
+			return InRetValType();
+		});
+	}
+
+	template <typename InRetValType, typename... ParamTypes, TDynamicColumnTemplate ColumnType>
+	TDelegate<InRetValType(ParamTypes...)> FAttributeBinder::BindEvent(const FName& InIdentifier, TDelegate<InRetValType(ParamTypes...)> ColumnType::* InVariable)
 	{
 		using namespace Private;
 		
@@ -113,21 +209,39 @@ namespace UE::Editor::DataStorage
 	}
 	
 	template <TDataColumnType ColumnType>
-	TAttribute<FText> FAttributeBinder::BindText(FString ColumnType::* InFStringVariable, const FName& InIdentifier)
+	TAttribute<FText> FAttributeBinder::BindText(FString ColumnType::* InFStringVariable)
 	{
 		return BindData(InFStringVariable, [](const FString& InString)
 		{
 			return FText::FromString(InString);
-		}, FString(), InIdentifier);
+		}, FString());
+	}
+	
+	template <TDynamicColumnTemplate ColumnType>
+	TAttribute<FText> FAttributeBinder::BindText(const FName& InIdentifier, FString ColumnType::* InFStringVariable)
+	{
+		return BindData(InIdentifier, InFStringVariable, [](const FString& InString)
+		{
+			return FText::FromString(InString);
+		}, FString());
 	}
 
 	template <TDataColumnType ColumnType>
-	TAttribute<FText> FAttributeBinder::BindText(FName ColumnType::* InFNameVariable, const FName& InIdentifier)
+	TAttribute<FText> FAttributeBinder::BindText(FName ColumnType::* InFNameVariable)
 	{
 		return BindData(InFNameVariable, [](const FName& InName)
 		{
 			return FText::FromName(InName);
-		}, FName(), InIdentifier);
+		}, FName());
+	}
+	
+	template <TDynamicColumnTemplate ColumnType>
+	TAttribute<FText> FAttributeBinder::BindText(const FName& InIdentifier, FName ColumnType::* InFNameVariable)
+	{
+		return BindData(InIdentifier, InFNameVariable, [](const FName& InName)
+		{
+			return FText::FromName(InName);
+		}, FName());
 	}
 
 } // namespace UE::Editor::DataStorage
diff --git a/Engine/Source/Runtime/TypedElementFramework/Public/Elements/Framework/TypedElementAttributeBindingProperty.h b/Engine/Source/Runtime/TypedElementFramework/Public/Elements/Framework/TypedElementAttributeBindingProperty.h
--- a/Engine/Source/Runtime/TypedElementFramework/Public/Elements/Framework/TypedElementAttributeBindingProperty.h
+++ b/Engine/Source/Runtime/TypedElementFramework/Public/Elements/Framework/TypedElementAttributeBindingProperty.h
@@ -15,23 +15,22 @@ namespace UE::Editor::DataStorage
 		// Concept for a conversion function used by the attribute binder to bind a column data member to an attribute of a different type
 		template <typename FunctionType, typename ArgumentType>
 		concept AttributeBinderInvocable = std::is_invocable_v<std::decay_t<FunctionType>, const ArgumentType&>;
+
+		template <TColumnType ColumnType>
+		ColumnType* GetColumn(IEditorDataStorageProvider* DataStorage, RowHandle Row)
+		{
+			return DataStorage->GetColumn<ColumnType>(Row);
+		}
 		
-		template <TDataColumnType ColumnType>
+		template <TDynamicColumnTemplate ColumnType>
 		ColumnType* GetColumn(IEditorDataStorageProvider* DataStorage, RowHandle Row, const FName& InIdentifier)
 		{
-			if(InIdentifier != NAME_None)
+			if(ensureMsgf(InIdentifier != NAME_None, TEXT("None identifier passed to dynamic column version of GetColumn. Will always return nullptr")))
 			{
-				if(ColumnType* DynamicColumn = DataStorage->GetColumn<ColumnType>(Row, InIdentifier))
-				{
-					return DynamicColumn;
-				}
-
-				// The dynamic column wasn't found on this row
-				return nullptr;
+				return DataStorage->GetColumn<ColumnType>(Row, InIdentifier);
 			}
-
-			// If we weren't provided a dynamic column identifier, we are looking for the regular non-dynamic column
-			return DataStorage->GetColumn<ColumnType>(Row);
+			
+			return nullptr;
 		}
 
 		template <typename PropertyType>
@@ -65,24 +64,30 @@ namespace UE::Editor::DataStorage
 			const PropertyType* Get(const void* Object) const;
 		};
 
+		template<typename T>
+		concept TDataColumnOrDynamicDataColumnTemplate = requires(T t)
+		{
+			TDataColumnType<T> || THasDynamicColumnTemplateSpecifier<T>;
+		};
+
 		template <typename PropertyType>
 		class TProperty
 		{
 		public:
 			// Bind this property directly
-			template <TDataColumnType ObjectType>
+			template <TDataColumnOrDynamicDataColumnTemplate ObjectType>
 			void Bind(PropertyType ObjectType::* Variable);
 
 			// Bind this property using a conversion function
-			template <typename InputType, TDataColumnType ObjectType>
+			template <typename InputType, TDataColumnOrDynamicDataColumnTemplate ObjectType>
 			void Bind(InputType ObjectType::* Variable, TFunction<PropertyType(const InputType&)> Converter);
 
 			// Get the bound property for the specified object
-			template <TDataColumnType ObjectType>
+			template <TDataColumnOrDynamicDataColumnTemplate ObjectType>
 			PropertyType& Get(ObjectType& Object) const;
 
 			// Get the bound property for the specified object
-			template <TDataColumnType ObjectType>
+			template <TDataColumnOrDynamicDataColumnTemplate ObjectType>
 			const PropertyType& Get(const ObjectType& Object) const;
 
 			// Get the bound property for this specified object ptr by providing type information about the object
@@ -170,7 +175,7 @@ namespace UE::Editor::DataStorage
 		//
 
 		template <typename PropertyType>
-		template <TDataColumnType ObjectType>
+		template <TDataColumnOrDynamicDataColumnTemplate ObjectType>
 		void TProperty<PropertyType>::Bind(PropertyType ObjectType::* Variable)
 		{
 			TDirectProperty<PropertyType> Result;
@@ -181,7 +186,7 @@ namespace UE::Editor::DataStorage
 		}
 
 		template <typename PropertyType>
-		template <typename InputType, TDataColumnType ObjectType>
+		template <typename InputType, TDataColumnOrDynamicDataColumnTemplate ObjectType>
 		void TProperty<PropertyType>::Bind(InputType ObjectType::* Variable, TFunction<PropertyType(const InputType&)> Converter)
 		{
 			TConvertibleProperty<PropertyType> Result;
@@ -195,14 +200,14 @@ namespace UE::Editor::DataStorage
 		}
 
 		template <typename PropertyType>
-		template <TDataColumnType ObjectType>
+		template <TDataColumnOrDynamicDataColumnTemplate ObjectType>
 		PropertyType& TProperty<PropertyType>::Get(ObjectType& Object) const
 		{
 			return Get(&Object, ObjectType::StaticStruct());
 		}
 
 		template <typename PropertyType>
-		template <TDataColumnType ObjectType>
+		template <TDataColumnOrDynamicDataColumnTemplate ObjectType>
 		const PropertyType& TProperty<PropertyType>::Get(const ObjectType& Object) const
 		{
 			return Get(&Object, ObjectType::StaticStruct());
diff --git a/Engine/Source/Runtime/TypedElementFramework/Public/Elements/Framework/TypedElementQueryBuilder.h b/Engine/Source/Runtime/TypedElementFramework/Public/Elements/Framework/TypedElementQueryBuilder.h
--- a/Engine/Source/Runtime/TypedElementFramework/Public/Elements/Framework/TypedElementQueryBuilder.h
+++ b/Engine/Source/Runtime/TypedElementFramework/Public/Elements/Framework/TypedElementQueryBuilder.h
@@ -164,7 +164,7 @@ namespace UE::Editor::DataStorage::Queries
 
 		// Dynamic Column Support
 		TYPEDELEMENTFRAMEWORK_API FSimpleQuery& Any(const FDynamicColumnDescription& Description);
-		template<TColumnType T>
+		template<TDynamicColumnTemplate T>
 		FSimpleQuery& Any(const FName&);
 
 		// ValueTags not yet supported for Any
@@ -187,7 +187,7 @@ namespace UE::Editor::DataStorage::Queries
 
 		// Dynamic Column Support
 		TYPEDELEMENTFRAMEWORK_API FSimpleQuery& None(const FDynamicColumnDescription& Description);
-		template<TColumnType T>
+		template<TDynamicColumnTemplate T>
 		FSimpleQuery& None(const FName&);
 		
 		// Value Tags
@@ -203,7 +203,7 @@ namespace UE::Editor::DataStorage::Queries
 
 		// Dynamic Column support
 		TYPEDELEMENTFRAMEWORK_API FSimpleQuery& All(const FDynamicColumnDescription& Description);
-		template<TColumnType T>
+		template<TDynamicColumnTemplate T>
 		FSimpleQuery& All(const FName&);
 
 		// Adds a filter to the query which must match a given ValueTag.  The value of the DynamicTag will not be checked.
@@ -408,8 +408,10 @@ namespace UE::Editor::DataStorage::Queries
 		 * it will not be used for finding matching tables. Columns bound with optional can not be bound to a query callback argument.
 		 */
 		TYPEDELEMENTFRAMEWORK_API Select& ReadOnly(TConstArrayView<const UScriptStruct*> Targets, EOptional Optional);
-		template<TDataColumnType Target>
+		template<TDynamicColumnTemplate Target>
 		Select& ReadOnly(const FName& Identifier);
+		template<TDynamicColumnTemplate Target>
+		Select& ReadOnly();
 		/** Request read and write access to the listed columns. */
 		template<TDataColumnType... TargetTypes>
 		Select& ReadWrite();
@@ -418,8 +420,10 @@ namespace UE::Editor::DataStorage::Queries
 		/** Request read and write access to the listed columns. */
 		TYPEDELEMENTFRAMEWORK_API Select& ReadWrite(TConstArrayView<const UScriptStruct*> Targets);
 		TYPEDELEMENTFRAMEWORK_API Select& ReadWrite(const FDynamicColumnDescription& Description);
-		template<TDataColumnType Target>
+		template<TDynamicColumnTemplate Target>
 		Select& ReadWrite(const FName& Identifier);
+		template<TDynamicColumnTemplate Target>
+		Select& ReadWrite();
 
 		// Query Condition API
 		TYPEDELEMENTFRAMEWORK_API FQueryConditionQuery Where(const Queries::FConditions& Condition);
diff --git a/Engine/Source/Runtime/TypedElementFramework/Public/Elements/Framework/TypedElementQueryBuilder.inl b/Engine/Source/Runtime/TypedElementFramework/Public/Elements/Framework/TypedElementQueryBuilder.inl
--- a/Engine/Source/Runtime/TypedElementFramework/Public/Elements/Framework/TypedElementQueryBuilder.inl
+++ b/Engine/Source/Runtime/TypedElementFramework/Public/Elements/Framework/TypedElementQueryBuilder.inl
@@ -1138,7 +1138,7 @@ e.g. void(FCachedQueryContext<Subsystem1, const Subsystem2>& Context, RowHandle
 		return *this;
 	}
 
-	template<TDataColumnType Target>
+	template<TDynamicColumnTemplate Target>
 	Select& Select::ReadOnly(const FName& Identifier)
 	{
 		ReadOnly(FDynamicColumnDescription
@@ -1156,7 +1156,7 @@ e.g. void(FCachedQueryContext<Subsystem1, const Subsystem2>& Context, RowHandle
 		return *this;
 	}
 
-	template<TDataColumnType Target>
+	template<TDynamicColumnTemplate Target>
 	Select& Select::ReadWrite(const FName& Identifier)
 	{
 		ReadWrite(FDynamicColumnDescription
@@ -1167,7 +1167,6 @@ e.g. void(FCachedQueryContext<Subsystem1, const Subsystem2>& Context, RowHandle
 		return *this;
 	}
 
-
 	//
 	// FSimpleQuery
 	//
@@ -1186,7 +1185,7 @@ e.g. void(FCachedQueryContext<Subsystem1, const Subsystem2>& Context, RowHandle
 		return *this;
 	}
 
-	template <TColumnType DynamicColumnTemplate>
+	template <TDynamicColumnTemplate DynamicColumnTemplate>
 	FSimpleQuery& FSimpleQuery::Any(const FName& Identifier)
 	{
 		return Any(FDynamicColumnDescription
@@ -1215,7 +1214,7 @@ e.g. void(FCachedQueryContext<Subsystem1, const Subsystem2>& Context, RowHandle
 		return All(FValueTag(Tag), Value);
 	}
 	
-	template <TColumnType DynamicColumnTemplate>
+	template <TDynamicColumnTemplate DynamicColumnTemplate>
 	FSimpleQuery& FSimpleQuery::None(const FName& Identifier)
 	{
 		return None(FDynamicColumnDescription
@@ -1225,7 +1224,7 @@ e.g. void(FCachedQueryContext<Subsystem1, const Subsystem2>& Context, RowHandle
 			});
 	}
 
-	template <TColumnType DynamicColumnTemplate>
+	template <TDynamicColumnTemplate DynamicColumnTemplate>
 	FSimpleQuery& FSimpleQuery::All(const FName& Identifier)
 	{
 		return All(FDynamicColumnDescription
diff --git a/Engine/Source/Runtime/TypedElementFramework/Public/Elements/Interfaces/TypedElementDataStorageInterface.h b/Engine/Source/Runtime/TypedElementFramework/Public/Elements/Interfaces/TypedElementDataStorageInterface.h
--- a/Engine/Source/Runtime/TypedElementFramework/Public/Elements/Interfaces/TypedElementDataStorageInterface.h
+++ b/Engine/Source/Runtime/TypedElementFramework/Public/Elements/Interfaces/TypedElementDataStorageInterface.h
@@ -222,10 +222,10 @@ public:
 	template<auto Value, UE::Editor::DataStorage::TEnumType EnumT = decltype(Value)>
 	void AddColumn(RowHandle Row);
 
-	template<UE::Editor::DataStorage::TColumnType DynamicColumnTemplate>
+	template<UE::Editor::DataStorage::TDynamicColumnTemplate DynamicColumnTemplate>
 	void AddColumn(RowHandle Row, const FName& Identifier);
 	
-	template<UE::Editor::DataStorage::TColumnType DynamicColumnTemplate>
+	template<UE::Editor::DataStorage::TDynamicColumnTemplate DynamicColumnTemplate>
 	void AddColumn(RowHandle Row, const FName& Identifier, DynamicColumnTemplate&& TemplateInstance);
 
 	/**
@@ -253,7 +253,7 @@ public:
 	template<typename T>
 	void RemoveColumn(RowHandle Row, const FName& Tag) = delete;
 
-	template<UE::Editor::DataStorage::TColumnType DynamicColumnTemplateType>
+	template<UE::Editor::DataStorage::TDynamicColumnTemplate DynamicColumnTemplateType>
 	void RemoveColumn(RowHandle Row, const FName& Identifier);
 	
 	template<>
@@ -290,9 +290,9 @@ public:
 	template<UE::Editor::DataStorage::TDataColumnType ColumnType>
 	const ColumnType* GetColumn(RowHandle Row) const;
 	// Gets a dynamic column identified by the ColumnTypeTemplate and Identifier
-	template<UE::Editor::DataStorage::TDataColumnType ColumnTypeTemplate>
+	template<UE::Editor::DataStorage::TDynamicColumnTemplate ColumnTypeTemplate>
 	ColumnTypeTemplate* GetColumn(RowHandle Row, const FName& Identifer);
-	template<UE::Editor::DataStorage::TDataColumnType ColumnTypeTemplate>
+	template<UE::Editor::DataStorage::TDynamicColumnTemplate ColumnTypeTemplate>
 	const ColumnTypeTemplate* GetColumn(RowHandle Row, const FName& Identifer) const;
 	
 	/** 
@@ -524,7 +524,7 @@ inline void IEditorDataStorageProvider::RemoveColumn<UE::Editor::DataStorage::FV
 	RemoveColumn(Row, FValueTag(Tag));
 }
 
-template<UE::Editor::DataStorage::TColumnType DynamicColumnTemplate>
+template<UE::Editor::DataStorage::TDynamicColumnTemplate DynamicColumnTemplate>
 void IEditorDataStorageProvider::RemoveColumn(RowHandle Row, const FName& Identifier)
 {
 	const UE::Editor::DataStorage::FDynamicColumnDescription Description
@@ -559,11 +559,9 @@ void IEditorDataStorageProvider::AddColumn(RowHandle Row)
 	AddColumn<EnumT>(Row, Value);
 }
 
-template <UE::Editor::DataStorage::TColumnType DynamicColumnTemplate>
+template <UE::Editor::DataStorage::TDynamicColumnTemplate DynamicColumnTemplate>
 void IEditorDataStorageProvider::AddColumn(RowHandle Row, const FName& Identifier)
 {
-	static_assert(UE::Editor::DataStorage::TDataColumnType<DynamicColumnTemplate> || UE::Editor::DataStorage::TTagColumnType<DynamicColumnTemplate>,
-		"DynamicColumnTemplate must be derived from either a Tag or Column");
 	const UE::Editor::DataStorage::FDynamicColumnDescription Description
 	{
 		.TemplateType = DynamicColumnTemplate::StaticStruct(),
@@ -573,11 +571,9 @@ void IEditorDataStorageProvider::AddColumn(RowHandle Row, const FName& Identifie
 	AddColumn(Row, StructInfo);
 }
 
-template <UE::Editor::DataStorage::TColumnType DynamicColumnTemplate>
+template <UE::Editor::DataStorage::TDynamicColumnTemplate DynamicColumnTemplate>
 void IEditorDataStorageProvider::AddColumn(RowHandle Row, const FName& Identifier, DynamicColumnTemplate&& TemplateInstance)
 {
-	static_assert(UE::Editor::DataStorage::TDataColumnType<DynamicColumnTemplate> || UE::Editor::DataStorage::TTagColumnType<DynamicColumnTemplate>,
-		"DynamicColumnTemplate must be derived from either a Tag or Column");
 	const UE::Editor::DataStorage::FDynamicColumnDescription Description
 	{
 		.TemplateType = DynamicColumnTemplate::StaticStruct(),
@@ -663,7 +659,7 @@ const ColumnType* IEditorDataStorageProvider::GetColumn(RowHandle Row) const
 	return reinterpret_cast<const ColumnType*>(GetColumnData(Row, ColumnType::StaticStruct()));
 }
 
-template <UE::Editor::DataStorage::TDataColumnType DynamicColumnTemplate>
+template <UE::Editor::DataStorage::TDynamicColumnTemplate DynamicColumnTemplate>
 DynamicColumnTemplate* IEditorDataStorageProvider::GetColumn(RowHandle Row, const FName& Identifier)
 {
 	const UE::Editor::DataStorage::FDynamicColumnDescription Description
@@ -679,7 +675,7 @@ DynamicColumnTemplate* IEditorDataStorageProvider::GetColumn(RowHandle Row, cons
 	return nullptr;
 }
 
-template <UE::Editor::DataStorage::TDataColumnType DynamicColumnTemplate>
+template <UE::Editor::DataStorage::TDynamicColumnTemplate DynamicColumnTemplate>
 const DynamicColumnTemplate* IEditorDataStorageProvider::GetColumn(RowHandle Row, const FName& Identifier) const
 {
 	const UE::Editor::DataStorage::FDynamicColumnDescription Description
diff --git a/Engine/Source/Runtime/TypedElementFramework/Public/Elements/Interfaces/TypedElementQueryStorageInterfaces.h b/Engine/Source/Runtime/TypedElementFramework/Public/Elements/Interfaces/TypedElementQueryStorageInterfaces.h
--- a/Engine/Source/Runtime/TypedElementFramework/Public/Elements/Interfaces/TypedElementQueryStorageInterfaces.h
+++ b/Engine/Source/Runtime/TypedElementFramework/Public/Elements/Interfaces/TypedElementQueryStorageInterfaces.h
@@ -21,7 +21,7 @@ namespace UE::Editor::DataStorage
 
 	using SubqueryCallback = TFunction<void(const FQueryDescription&, ISubqueryContext&)>;
 	using SubqueryCallbackRef = TFunctionRef<void(const FQueryDescription&, ISubqueryContext&)>;
-
+	
 	/**
 	 * Base interface for any contexts provided to query callbacks.
 	 */
@@ -40,16 +40,16 @@ namespace UE::Editor::DataStorage
 		/** Return the address of a immutable column matching the requested type or a nullptr if not found. */
 		virtual const void* GetColumn(const UScriptStruct* ColumnType) const = 0;
 		/** Return the address of a immutable column matching the requested type or a nullptr if not found. */
-		template<typename Column>
+		template<TDataColumnType Column>
 		const Column* GetColumn() const;
-		template<TDataColumnType TemplateType>
+		template<TDynamicColumnTemplate TemplateType>
 		const TemplateType* GetColumn(const FName& Identifier) const;
 		/** Return the address of a mutable column matching the requested type or a nullptr if not found. */
 		virtual void* GetMutableColumn(const UScriptStruct* ColumnType) = 0;
 		/** Return the address of a mutable column matching the requested type or a nullptr if not found. */
 		template<typename Column>
 		Column* GetMutableColumn();
-		template<TDataColumnType TemplateType>
+		template<TDynamicColumnTemplate TemplateType>
 		TemplateType* GetMutableColumn(const FName& Identifier);
 
 		/**
@@ -92,7 +92,7 @@ namespace UE::Editor::DataStorage
 		template<typename Column>
 		bool HasColumn(RowHandle Row) const;
 
-		template<TColumnType DynamicColumnTemplate>
+		template<TDynamicColumnTemplate DynamicColumnTemplate>
 		bool HasColumn(RowHandle Row, const FName& Identifier) const;
 
 		/**
@@ -100,7 +100,7 @@ namespace UE::Editor::DataStorage
 		 * @return The UScriptStruct for a previously created Dynamic Column.  If no column exists, then nullptr
 		 */
 		virtual const UScriptStruct* FindDynamicColumnType(const FDynamicColumnDescription& Description) const = 0;
-		template<TColumnType TemplateType>
+		template<TDynamicColumnTemplate TemplateType>
 		const UScriptStruct* FindDynamicColumnType(const FName& Identifier) const;
 	};
 
@@ -400,13 +400,13 @@ namespace UE::Editor::DataStorage
 	// ICommonQueryContext
 	//
 
-	template<typename Column>
+	template<TDataColumnType Column>
 	const Column* ICommonQueryContext::GetColumn() const
 	{
 		return reinterpret_cast<const Column*>(GetColumn(Column::StaticStruct()));
 	}
 
-	template <TDataColumnType TemplateType>
+	template <TDynamicColumnTemplate TemplateType>
 	const TemplateType* ICommonQueryContext::GetColumn(const FName& Identifier) const
 	{
 		const UScriptStruct* DynamicColumnType = FindDynamicColumnType<TemplateType>(Identifier);
@@ -419,7 +419,7 @@ namespace UE::Editor::DataStorage
 		return reinterpret_cast<Column*>(GetMutableColumn(Column::StaticStruct()));
 	}
 
-	template <TDataColumnType TemplateType>
+	template <TDynamicColumnTemplate TemplateType>
 	TemplateType* ICommonQueryContext::GetMutableColumn(const FName& Identifier)
 	{
 		const UScriptStruct* DynamicColumnType = FindDynamicColumnType<TemplateType>(Identifier);
@@ -438,7 +438,7 @@ namespace UE::Editor::DataStorage
 		return HasColumn(Row, Column::StaticStruct());
 	}
 
-	template <TColumnType DynamicColumnTemplate>
+	template <TDynamicColumnTemplate DynamicColumnTemplate>
 	bool ICommonQueryContext::HasColumn(RowHandle Row, const FName& Identifier) const
 	{
 		if (const UScriptStruct* DynamicColumnType = FindDynamicColumnType<DynamicColumnTemplate>(Identifier))
@@ -448,7 +448,7 @@ namespace UE::Editor::DataStorage
 		return false;
 	}
 
-	template <TColumnType TemplateType>
+	template <TDynamicColumnTemplate TemplateType>
 	const UScriptStruct* ICommonQueryContext::FindDynamicColumnType(const FName& Identifier) const
 	{
 		return FindDynamicColumnType(FDynamicColumnDescription
