[TEDS] Temporarily removing UHT code generation for dynamic columns
* Manual addition of empty struct attribute into the dynamic column
* The addition of the plugin exposed an issue with the .uhtplugin.cs framework which caused a performance regression and log spam.

diff --git a/Engine/Plugins/Experimental/EditorDataStorage/Source/TedsCore/Private/Debug/TypedElementDatabaseDebugTypes.h b/Engine/Plugins/Experimental/EditorDataStorage/Source/TedsCore/Private/Debug/TypedElementDatabaseDebugTypes.h
--- a/Engine/Plugins/Experimental/EditorDataStorage/Source/TedsCore/Private/Debug/TypedElementDatabaseDebugTypes.h
+++ b/Engine/Plugins/Experimental/EditorDataStorage/Source/TedsCore/Private/Debug/TypedElementDatabaseDebugTypes.h
@@ -27,6 +27,8 @@ struct FTestDynamicColumn : public FEditorDataStorageColumn
 {
 	GENERATED_BODY()
 
+	struct EditorDataStorage_DynamicColumnTemplate{};
+
 	// Note: Not trivial type to check that non-default move/copying works
 	// Should work regardless of UPROPERTY
 	TArray<int32> IntArray;
@@ -36,4 +38,6 @@ USTRUCT(meta=(EditorDataStorage_DynamicColumnTemplate))
 struct FTestDynamicTag : public FEditorDataStorageTag
 {
 	GENERATED_BODY()
+	
+	struct EditorDataStorage_DynamicColumnTemplate{};
 };
diff --git a/Engine/Plugins/Experimental/EditorDataStorage/Source/UhtExtension/EditorDataStorageUhtExtension.cs b/Engine/Plugins/Experimental/EditorDataStorage/Source/UhtExtension/EditorDataStorageUhtExtension.cs
deleted file mode 100644
diff --git a/Engine/Plugins/Experimental/EditorDataStorage/Source/UhtExtension/EditorDataStorageUhtExtension.ubtplugin.csproj b/Engine/Plugins/Experimental/EditorDataStorage/Source/UhtExtension/EditorDataStorageUhtExtension.ubtplugin.csproj
deleted file mode 100644
diff --git a/Engine/Plugins/Experimental/EditorDataStorageFeatures/Source/TedsAssetData/Public/TedsAssetDataColumns.h b/Engine/Plugins/Experimental/EditorDataStorageFeatures/Source/TedsAssetData/Public/TedsAssetDataColumns.h
--- a/Engine/Plugins/Experimental/EditorDataStorageFeatures/Source/TedsAssetData/Public/TedsAssetDataColumns.h
+++ b/Engine/Plugins/Experimental/EditorDataStorageFeatures/Source/TedsAssetData/Public/TedsAssetDataColumns.h
@@ -180,6 +180,8 @@ struct FItemStringAttributeColumn_Experimental : public FItemAttributeBaseColumn
 {
 	GENERATED_BODY()
 
+	struct EditorDataStorage_DynamicColumnTemplate{};
+
 	UPROPERTY()
 	FString Value;
 };
\ No newline at end of file
