TEDS UI: Move STedsWidget from TypedElementFramework to the EditorDataStorage plugin and replace with an interface
This was done so that STedsWidget can use classes in the Slate module which we did not want TypedElementFramework to depend on

diff --git a/Engine/Plugins/Experimental/EditorDataStorage/Source/TedsCore/Private/TypedElementDatabaseUI.cpp b/Engine/Plugins/Experimental/EditorDataStorage/Source/TedsCore/Private/TypedElementDatabaseUI.cpp
--- a/Engine/Plugins/Experimental/EditorDataStorage/Source/TedsCore/Private/TypedElementDatabaseUI.cpp
+++ b/Engine/Plugins/Experimental/EditorDataStorage/Source/TedsCore/Private/TypedElementDatabaseUI.cpp
@@ -12,6 +12,7 @@
 #include "Elements/Interfaces/TypedElementDataStorageCompatibilityInterface.h"
 #include "GenericPlatform/GenericPlatformMemory.h"
 #include "Widgets/SlateControlledConstruction.h"
+#include "Widgets/STedsWidget.h"
 
 DEFINE_LOG_CATEGORY(LogEditorDataStorageUI);
 
@@ -455,6 +456,18 @@ bool UEditorDataStorageUi::SupportsExtension(FName Extension) const
 
 void UEditorDataStorageUi::ListExtensions(TFunctionRef<void(FName)> Callback) const
 {
+	
+}
+
+TSharedPtr<UE::Editor::DataStorage::ITedsWidget> UEditorDataStorageUi::CreateContainerTedsWidget(UE::Editor::DataStorage::RowHandle UiRowHandle) const
+{
+	return SNew(UE::Editor::DataStorage::Widgets::STedsWidget)
+			.UiRowHandle(UiRowHandle);
+}
+
+UE::Editor::DataStorage::TableHandle UEditorDataStorageUi::GetWidgetTable() const
+{
+	return WidgetTable;
 }
 
 void UEditorDataStorageUi::CreateStandardArchetypes()
diff --git a/Engine/Plugins/Experimental/EditorDataStorage/Source/TedsCore/Private/TypedElementDatabaseUI.h b/Engine/Plugins/Experimental/EditorDataStorage/Source/TedsCore/Private/TypedElementDatabaseUI.h
--- a/Engine/Plugins/Experimental/EditorDataStorage/Source/TedsCore/Private/TypedElementDatabaseUI.h
+++ b/Engine/Plugins/Experimental/EditorDataStorage/Source/TedsCore/Private/TypedElementDatabaseUI.h
@@ -55,6 +55,12 @@ public:
 	bool SupportsExtension(FName Extension) const override;
 	void ListExtensions(TFunctionRef<void(FName)> Callback) const override;
 
+	/** Create the container widget that every TEDS UI widget is stored in */
+	virtual TSharedPtr<UE::Editor::DataStorage::ITedsWidget> CreateContainerTedsWidget(UE::Editor::DataStorage::RowHandle UiRowHandle) const override;
+	
+	/** Get the table where TEDS UI widgets are stored */
+	virtual UE::Editor::DataStorage::TableHandle GetWidgetTable() const override;
+
 private:
 	struct FWidgetFactory
 	{
diff --git a/Engine/Plugins/Experimental/EditorDataStorage/Source/TedsCore/Private/Widgets/STedsWidget.cpp b/Engine/Plugins/Experimental/EditorDataStorage/Source/TedsCore/Private/Widgets/STedsWidget.cpp
new file mode 100644
--- /dev/null
+++ b/Engine/Plugins/Experimental/EditorDataStorage/Source/TedsCore/Private/Widgets/STedsWidget.cpp
@@ -0,0 +1,140 @@
+// Copyright Epic Games, Inc. All Rights Reserved.
+
+#include "STedsWidget.h"
+
+#include "Elements/Columns/TypedElementSlateWidgetColumns.h"
+#include "Elements/Common//EditorDataStorageFeatures.h"
+#include "Elements/Framework/TypedElementAttributeBinding.h"
+#include "Elements/Interfaces/TypedElementDataStorageInterface.h"
+#include "Elements/Interfaces/TypedElementDataStorageUiInterface.h"
+
+namespace UE::Editor::DataStorage::TedsWidget::Private
+{
+	bool bUseDefaultAttributeBindings = false;
+
+	static FAutoConsoleVariableRef CVarUseDefaultAttributeBindings(
+		TEXT("TEDS.UI.UseDefaultAttributeBindings"),
+		bUseDefaultAttributeBindings,
+		TEXT("If true, widgets created through TEDS UI will provide a default set of attribute bindings. Does not apply to existing widgets retroactively"),
+		ECVF_Default
+	);
+}
+
+namespace UE::Editor::DataStorage::Widgets
+{
+	void STedsWidget::Construct(const FArguments& InArgs)
+	{
+		UiRowHandle = InArgs._UiRowHandle;
+
+		// If the Ui row wasn't already registered externally, register it with Teds
+		if (UiRowHandle == UE::Editor::DataStorage::InvalidRowHandle)
+		{
+			RegisterTedsWidget(InArgs._Content.Widget);
+		}
+
+		if (UE::Editor::DataStorage::TedsWidget::Private::bUseDefaultAttributeBindings)
+		{
+			UE::Editor::DataStorage::FAttributeBinder Binder(UiRowHandle, GetStorageIfAvailable());
+
+			SetColorAndOpacity(Binder.BindData(&FSlateColorColumn::Color, [](const FSlateColor& InColor)
+			{
+				if(InColor.IsColorSpecified())
+				{
+					return InColor.GetSpecifiedColor();
+				}
+
+				return FLinearColor::White;
+			}, FSlateColor::UseForeground()));
+		}
+
+		SetContent(InArgs._Content.Widget);
+	}
+
+	void STedsWidget::RegisterTedsWidget(const TSharedPtr<SWidget>& InContentWidget)
+	{
+		using namespace UE::Editor::DataStorage;
+		IEditorDataStorageProvider* Storage = GetStorageIfAvailable();
+		IEditorDataStorageUiProvider* StorageUi = GetStorageUiIfAvailable();
+
+		// If TEDS is not enabled, STedsWidget will just behave like a regular widget
+		if (!Storage || !StorageUi)
+		{
+			return;
+		}
+		
+		const TableHandle WidgetTable = StorageUi->GetWidgetTable();
+		if (WidgetTable == InvalidTableHandle)
+		{
+			return;
+		}
+		
+		UiRowHandle = Storage->AddRow(WidgetTable);
+		
+		if (FTypedElementSlateWidgetReferenceColumn* WidgetReferenceColumn = Storage->GetColumn<FTypedElementSlateWidgetReferenceColumn>(UiRowHandle))
+		{
+			WidgetReferenceColumn->TedsWidget = SharedThis(this);
+			WidgetReferenceColumn->Widget = InContentWidget;
+		}
+	}
+
+	void STedsWidget::SetContent(const TSharedRef< SWidget >& InContent)
+	{
+		if (IEditorDataStorageProvider* Storage = GetStorageIfAvailable())
+		{
+			if (FTypedElementSlateWidgetReferenceColumn* WidgetReferenceColumn = Storage->GetColumn<FTypedElementSlateWidgetReferenceColumn>(UiRowHandle))
+			{
+				// First we set the widget reference on the column
+				WidgetReferenceColumn->Widget = InContent;
+
+				if (TSharedPtr<FTypedElementWidgetConstructor> Constructor = WidgetReferenceColumn->WidgetConstructor.Pin())
+				{
+					TConstArrayView<const UScriptStruct*> AdditionalColumns = Constructor->GetAdditionalColumnsList();
+					
+					// When we are setting valid content, we want to add the additional columns to the widget row so it gets picked up by any queries
+					// the widget constructor requires.
+					if (InContent != SNullWidget::NullWidget)
+					{
+						Storage->AddColumns(UiRowHandle, AdditionalColumns);
+					}
+					// If we are setting to null content, i.e removing the widget we also remove the additional columns so the row stops matching any
+					// queries designed to operate on the widget
+					else
+					{
+						// RemoveColumns ensures if the columns don't exist on the row
+						if (Storage->HasColumns(UiRowHandle, AdditionalColumns))
+						{
+							Storage->RemoveColumns(UiRowHandle, AdditionalColumns);
+						}
+					}
+				}
+			}
+		}
+		
+		ChildSlot
+		[
+			InContent
+		];
+	}
+
+	UE::Editor::DataStorage::RowHandle STedsWidget::GetRowHandle() const
+	{
+		return UiRowHandle;
+	}
+
+	TSharedRef<SWidget> STedsWidget::AsWidget()
+	{
+		return SharedThis(this);
+	}
+
+	IEditorDataStorageProvider* STedsWidget::GetStorageIfAvailable()
+	{
+		using namespace UE::Editor::DataStorage;
+		return GetMutableDataStorageFeature<IEditorDataStorageProvider>(StorageFeatureName);
+	}
+
+	IEditorDataStorageUiProvider* STedsWidget::GetStorageUiIfAvailable()
+	{
+		using namespace UE::Editor::DataStorage;
+		return GetMutableDataStorageFeature<IEditorDataStorageUiProvider>(UiFeatureName);
+	}
+}
\ No newline at end of file
diff --git a/Engine/Plugins/Experimental/EditorDataStorage/Source/TedsCore/Private/Widgets/STedsWidget.h b/Engine/Plugins/Experimental/EditorDataStorage/Source/TedsCore/Private/Widgets/STedsWidget.h
new file mode 100644
--- /dev/null
+++ b/Engine/Plugins/Experimental/EditorDataStorage/Source/TedsCore/Private/Widgets/STedsWidget.h
@@ -0,0 +1,60 @@
+// Copyright Epic Games, Inc. All Rights Reserved.
+
+#pragma once
+
+#include "Elements/Framework/TypedElementDataStorageWidget.h"
+#include "Elements/Interfaces/TypedElementDataStorageInterface.h"
+#include "Widgets/SCompoundWidget.h"
+
+struct FTypedElementWidgetConstructor;
+
+class IEditorDataStorageUiProvider;
+
+namespace UE::Editor::DataStorage::Widgets
+{
+	/*
+	 * All Teds widgets will be contained inside STedsWidget which acts like a container widget
+	 * so we can have guaranteed access to the contents inside to dynamically update them if required.
+	 * This widget is created and returned for any Teds widget requested for a row, regardless of if
+	 * the actual internal widget exists or not.
+	 * 
+	 * Currently this is simply an SCompoundWidget
+	 */
+	class STedsWidget : public SCompoundWidget, public UE::Editor::DataStorage::ITedsWidget
+	{
+	public:
+
+		SLATE_BEGIN_ARGS(STedsWidget)
+			: _UiRowHandle(UE::Editor::DataStorage::InvalidRowHandle)
+			, _Content()
+		{
+		}
+
+		// The UI Row this widget will be assigned to
+		SLATE_ARGUMENT(UE::Editor::DataStorage::RowHandle, UiRowHandle)
+	
+		/** The actual widget content */
+		SLATE_DEFAULT_SLOT(FArguments, Content)
+	
+		SLATE_END_ARGS()
+
+		void Construct( const FArguments& InArgs );
+	
+		virtual void SetContent(const TSharedRef< SWidget >& InContent) override;
+
+		virtual UE::Editor::DataStorage::RowHandle GetRowHandle() const override;
+	
+		virtual TSharedRef<SWidget> AsWidget() override;
+
+	private:
+
+		void RegisterTedsWidget(const TSharedPtr<SWidget>& InContentWidget);
+	
+		static IEditorDataStorageProvider* GetStorageIfAvailable();
+		static IEditorDataStorageUiProvider* GetStorageUiIfAvailable();
+
+	private:
+	
+		RowHandle UiRowHandle = InvalidRowHandle;
+	};
+}
diff --git a/Engine/Plugins/Experimental/EditorDataStorage/Source/TedsCore/TedsCore.Build.cs b/Engine/Plugins/Experimental/EditorDataStorage/Source/TedsCore/TedsCore.Build.cs
--- a/Engine/Plugins/Experimental/EditorDataStorage/Source/TedsCore/TedsCore.Build.cs
+++ b/Engine/Plugins/Experimental/EditorDataStorage/Source/TedsCore/TedsCore.Build.cs
@@ -24,6 +24,7 @@ public class TedsCore : ModuleRules
 					"MassEntityEditor",
 					"TypedElementFramework",
 					"SlateCore",
+					"Slate",
 					"UnrealEd"
 				});
 
diff --git a/Engine/Plugins/Experimental/EditorDataStorageFeatures/Source/TedsTableViewer/Private/TedsTableViewerColumn.cpp b/Engine/Plugins/Experimental/EditorDataStorageFeatures/Source/TedsTableViewer/Private/TedsTableViewerColumn.cpp
--- a/Engine/Plugins/Experimental/EditorDataStorageFeatures/Source/TedsTableViewer/Private/TedsTableViewerColumn.cpp
+++ b/Engine/Plugins/Experimental/EditorDataStorageFeatures/Source/TedsTableViewer/Private/TedsTableViewerColumn.cpp
@@ -310,7 +310,7 @@ namespace UE::Editor::DataStorage
 			}
 
 			// Check if the container TEDSWidget exists, if not we cannot update this widget
-			const TSharedPtr<STedsWidget> TedsWidget = WidgetColumn->TedsWidget.Pin();
+			const TSharedPtr<ITedsWidget> TedsWidget = WidgetColumn->TedsWidget.Pin();
 			if(!TedsWidget)
 			{
 				continue;
diff --git a/Engine/Plugins/Experimental/EditorDataStorageFeatures/Source/TedsTableViewer/Private/TedsTableViewerModel.cpp b/Engine/Plugins/Experimental/EditorDataStorageFeatures/Source/TedsTableViewer/Private/TedsTableViewerModel.cpp
--- a/Engine/Plugins/Experimental/EditorDataStorageFeatures/Source/TedsTableViewer/Private/TedsTableViewerModel.cpp
+++ b/Engine/Plugins/Experimental/EditorDataStorageFeatures/Source/TedsTableViewer/Private/TedsTableViewerModel.cpp
@@ -140,6 +140,11 @@ namespace UE::Editor::DataStorage
 		return Storage;
 	}
 
+	IEditorDataStorageUiProvider* FTedsTableViewerModel::GetDataStorageUiProvider() const
+	{
+		return StorageUi;
+	}
+
 	void FTedsTableViewerModel::GenerateColumns()
 	{
 		using MatchApproach = IEditorDataStorageUiProvider::EMatchApproach;
diff --git a/Engine/Plugins/Experimental/EditorDataStorageFeatures/Source/TedsTableViewer/Private/Widgets/STedsTableViewer.cpp b/Engine/Plugins/Experimental/EditorDataStorageFeatures/Source/TedsTableViewer/Private/Widgets/STedsTableViewer.cpp
--- a/Engine/Plugins/Experimental/EditorDataStorageFeatures/Source/TedsTableViewer/Private/Widgets/STedsTableViewer.cpp
+++ b/Engine/Plugins/Experimental/EditorDataStorageFeatures/Source/TedsTableViewer/Private/Widgets/STedsTableViewer.cpp
@@ -25,11 +25,18 @@ namespace UE::Editor::DataStorage
 		
 		HeaderRowWidget = SNew( SHeaderRow )
 							.CanSelectGeneratedColumn(true);
+		
+		IEditorDataStorageUiProvider* StorageUi = Model->GetDataStorageUiProvider();
 
-		ChildSlot
-		[
-			SAssignNew(TedsWidget, STedsWidget)
-		];
+		if (ensure(StorageUi))
+		{
+			TedsWidget = StorageUi->CreateContainerTedsWidget(InvalidRowHandle);
+			
+			ChildSlot
+			[
+				TedsWidget->AsWidget()
+			];
+		}
 		
 		AddWidgetColumns();
 
diff --git a/Engine/Plugins/Experimental/EditorDataStorageFeatures/Source/TedsTableViewer/Private/Widgets/STedsTileViewer.cpp b/Engine/Plugins/Experimental/EditorDataStorageFeatures/Source/TedsTableViewer/Private/Widgets/STedsTileViewer.cpp
--- a/Engine/Plugins/Experimental/EditorDataStorageFeatures/Source/TedsTableViewer/Private/Widgets/STedsTileViewer.cpp
+++ b/Engine/Plugins/Experimental/EditorDataStorageFeatures/Source/TedsTableViewer/Private/Widgets/STedsTileViewer.cpp
@@ -22,10 +22,17 @@ namespace UE::Editor::DataStorage
 		Model = MakeShared<FTedsTableViewerModel>(InArgs._QueryStack, Columns, Purposes,
 			FTedsTableViewerModel::FIsItemVisible::CreateSP(this, &STedsTileViewer::IsItemVisible));
 
-		ChildSlot
-		[
-			SAssignNew(TedsWidget, STedsWidget)
-		];
+		IEditorDataStorageUiProvider* StorageUi = Model->GetDataStorageUiProvider();
+		
+		if (ensure(StorageUi))
+		{
+			TedsWidget = StorageUi->CreateContainerTedsWidget(InvalidRowHandle);
+			
+			ChildSlot
+			[
+				TedsWidget->AsWidget()
+			];
+		}
 		
 		AddWidgetColumns();
 		
diff --git a/Engine/Plugins/Experimental/EditorDataStorageFeatures/Source/TedsTableViewer/Public/TedsTableViewerModel.h b/Engine/Plugins/Experimental/EditorDataStorageFeatures/Source/TedsTableViewer/Public/TedsTableViewerModel.h
--- a/Engine/Plugins/Experimental/EditorDataStorageFeatures/Source/TedsTableViewer/Public/TedsTableViewerModel.h
+++ b/Engine/Plugins/Experimental/EditorDataStorageFeatures/Source/TedsTableViewer/Public/TedsTableViewerModel.h
@@ -61,6 +61,8 @@ namespace UE::Editor::DataStorage
 
 		TEDSTABLEVIEWER_API IEditorDataStorageProvider* GetDataStorageInterface() const;
 
+		TEDSTABLEVIEWER_API IEditorDataStorageUiProvider* GetDataStorageUiProvider() const;
+
 	protected:
 
 		// Generate the actual columns to display in the UI using TEDS UI
diff --git a/Engine/Plugins/Experimental/EditorDataStorageFeatures/Source/TedsTableViewer/Public/Widgets/STedsTableViewer.h b/Engine/Plugins/Experimental/EditorDataStorageFeatures/Source/TedsTableViewer/Public/Widgets/STedsTableViewer.h
--- a/Engine/Plugins/Experimental/EditorDataStorageFeatures/Source/TedsTableViewer/Public/Widgets/STedsTableViewer.h
+++ b/Engine/Plugins/Experimental/EditorDataStorageFeatures/Source/TedsTableViewer/Public/Widgets/STedsTableViewer.h
@@ -12,7 +12,11 @@
 class ITableRow;
 class STableViewBase;
 class SHeaderRow;
-class STedsWidget;
+
+namespace UE::Editor::DataStorage
+{
+	class ITedsWidget;
+}
 
 namespace UE::Editor::DataStorage
 {
@@ -119,7 +123,7 @@ namespace UE::Editor::DataStorage
 		FOnSelectionChanged OnSelectionChanged;
 
 		// Wrapper Teds Widget around our contents so we can use Teds columns to specify behavior
-		TSharedPtr<STedsWidget> TedsWidget;
+		TSharedPtr<ITedsWidget> TedsWidget;
 
 		// The message to show in place of the table viewer when there are no rows provided by the current query stack
 		TAttribute<FText> EmptyRowsMessage;
diff --git a/Engine/Plugins/Experimental/EditorDataStorageFeatures/Source/TedsTableViewer/Public/Widgets/STedsTileViewer.h b/Engine/Plugins/Experimental/EditorDataStorageFeatures/Source/TedsTableViewer/Public/Widgets/STedsTileViewer.h
--- a/Engine/Plugins/Experimental/EditorDataStorageFeatures/Source/TedsTableViewer/Public/Widgets/STedsTileViewer.h
+++ b/Engine/Plugins/Experimental/EditorDataStorageFeatures/Source/TedsTableViewer/Public/Widgets/STedsTileViewer.h
@@ -12,7 +12,11 @@
 class ITableRow;
 class STableViewBase;
 class SHeaderRow;
-class STedsWidget;
+
+namespace UE::Editor::DataStorage
+{
+	class ITedsWidget;
+}
 
 namespace UE::Editor::DataStorage
 {
@@ -94,7 +98,7 @@ namespace UE::Editor::DataStorage
 		FOnSelectionChanged OnSelectionChanged;
 
 		// Wrapper Teds Widget around our contents so we can use Teds columns to specify behavior
-		TSharedPtr<STedsWidget> TedsWidget;
+		TSharedPtr<ITedsWidget> TedsWidget;
 
 		// The widget purpose used to generate tiles
 		FName WidgetPurpose;
diff --git a/Engine/Source/Runtime/TypedElementFramework/Private/Elements/Framework/TypedElementDataStorageWidget.cpp b/Engine/Source/Runtime/TypedElementFramework/Private/Elements/Framework/TypedElementDataStorageWidget.cpp
deleted file mode 100644
diff --git a/Engine/Source/Runtime/TypedElementFramework/Private/Elements/Interfaces/TypedElementDataStorageUiInterface.cpp b/Engine/Source/Runtime/TypedElementFramework/Private/Elements/Interfaces/TypedElementDataStorageUiInterface.cpp
--- a/Engine/Source/Runtime/TypedElementFramework/Private/Elements/Interfaces/TypedElementDataStorageUiInterface.cpp
+++ b/Engine/Source/Runtime/TypedElementFramework/Private/Elements/Interfaces/TypedElementDataStorageUiInterface.cpp
@@ -7,6 +7,7 @@
 #include "Elements/Columns/TypedElementSlateWidgetColumns.h"
 #include "Elements/Columns/TypedElementUIColumns.h"
 #include "Elements/Framework/TypedElementDataStorageWidget.h"
+#include "Widgets/SNullWidget.h"
 
 FTypedElementWidgetConstructor::FTypedElementWidgetConstructor(const UScriptStruct* InTypeInfo)
 	: TypeInfo(InTypeInfo)
@@ -127,14 +128,15 @@ TSharedPtr<SWidget> FTypedElementWidgetConstructor::ConstructFinalWidget(
 	}
 
 	// Create a container widget to hold the content (even if it doesn't exist yet)
-	TSharedPtr<STedsWidget> ContainerWidget = SNew(STedsWidget)
-	.UiRowHandle(Row)
-	[
-		Widget.ToSharedRef()
-	];
+	if (TSharedPtr<UE::Editor::DataStorage::ITedsWidget> ContainerWidget = DataStorageUi->CreateContainerTedsWidget(Row))
+	{
+		ContainerWidget->SetContent(Widget.ToSharedRef());
+		DataStorage->GetColumn<FTypedElementSlateWidgetReferenceColumn>(Row)->TedsWidget = ContainerWidget;
+
+		return ContainerWidget->AsWidget();
+	}
 	
-	DataStorage->GetColumn<FTypedElementSlateWidgetReferenceColumn>(Row)->TedsWidget = ContainerWidget;
-	return ContainerWidget;
+	return Widget;
 }
 
 TSharedPtr<SWidget> FTypedElementWidgetConstructor::Construct(
diff --git a/Engine/Source/Runtime/TypedElementFramework/Public/Elements/Columns/TypedElementSlateWidgetColumns.h b/Engine/Source/Runtime/TypedElementFramework/Public/Elements/Columns/TypedElementSlateWidgetColumns.h
--- a/Engine/Source/Runtime/TypedElementFramework/Public/Elements/Columns/TypedElementSlateWidgetColumns.h
+++ b/Engine/Source/Runtime/TypedElementFramework/Public/Elements/Columns/TypedElementSlateWidgetColumns.h
@@ -9,9 +9,13 @@
 #include "TypedElementSlateWidgetColumns.generated.h"
 
 class SWidget;
-class STedsWidget;
 struct FTypedElementWidgetConstructor;
 
+namespace UE::Editor::DataStorage
+{
+	class ITedsWidget;
+}
+
 /**
  * Stores a widget reference in the data storage. At the start of processing any
  * columns that are not pointing to a valid widget will be removed. If the
@@ -27,7 +31,7 @@ struct FTypedElementSlateWidgetReferenceColumn final : public FEditorDataStorage
 	TWeakPtr<SWidget> Widget;
 
 	// Reference to the container widget that holds the internal widget
-	TWeakPtr<STedsWidget> TedsWidget;
+	TWeakPtr<UE::Editor::DataStorage::ITedsWidget> TedsWidget;
 
 	// Reference to the widget constructor that was used to create this widget (if applicable)
 	TWeakPtr<FTypedElementWidgetConstructor> WidgetConstructor;
diff --git a/Engine/Source/Runtime/TypedElementFramework/Public/Elements/Framework/TypedElementDataStorageWidget.h b/Engine/Source/Runtime/TypedElementFramework/Public/Elements/Framework/TypedElementDataStorageWidget.h
--- a/Engine/Source/Runtime/TypedElementFramework/Public/Elements/Framework/TypedElementDataStorageWidget.h
+++ b/Engine/Source/Runtime/TypedElementFramework/Public/Elements/Framework/TypedElementDataStorageWidget.h
@@ -2,52 +2,24 @@
 
 #pragma once
 
-#include "Widgets/SCompoundWidget.h"
-#include "Elements/Interfaces/TypedElementDataStorageInterface.h"
+#include "Elements/Common/TypedElementHandles.h"
+#include "Widgets/SWidget.h"
 
-struct FTypedElementWidgetConstructor;
-
-/*
- * All Teds widgets will be contained inside STedsWidget which acts like a container widget
- * so we can have guaranteed access to the contents inside to dynamically update them if required.
- * This widget is created and returned for any Teds widget requested for a row, regardless of if
- * the actual internal widget exists or not.
- * 
- * Currently this is simply an SCompoundWidget
- */
-class STedsWidget : public SCompoundWidget
+namespace UE::Editor::DataStorage
 {
-public:
-
-	SLATE_BEGIN_ARGS(STedsWidget)
-		: _UiRowHandle(UE::Editor::DataStorage::InvalidRowHandle)
-		, _Content()
+	class ITedsWidget
 	{
-	}
-
-	// The UI Row this widget will be assigned to
-	SLATE_ARGUMENT(UE::Editor::DataStorage::RowHandle, UiRowHandle)
-	
-	/** The actual widget content */
-	SLATE_DEFAULT_SLOT(FArguments, Content)
-	
-	SLATE_END_ARGS()
-	
-	TYPEDELEMENTFRAMEWORK_API STedsWidget();
-
-	TYPEDELEMENTFRAMEWORK_API void Construct( const FArguments& InArgs );
-	
-	TYPEDELEMENTFRAMEWORK_API void SetContent(const TSharedRef< SWidget >& InContent);
+	public:
+		virtual ~ITedsWidget() = default;
 
-	TYPEDELEMENTFRAMEWORK_API UE::Editor::DataStorage::RowHandle GetRowHandle() const;
+		virtual RowHandle GetRowHandle() const = 0;
 
-private:
+		virtual void SetContent(const TSharedRef< SWidget >& InContent) = 0;
 
-	void RegisterTedsWidget(const TSharedPtr<SWidget>& InContentWidget);
-	
-	static IEditorDataStorageProvider* GetStorageIfAvailable();
+		virtual TSharedRef<SWidget> AsWidget() = 0;
+	};
+}
 
-private:
-	
-	UE::Editor::DataStorage::RowHandle UiRowHandle;
-};
+// This is for backwards compatibility with the STedsWidget class in the global namespace that used to exist in this file but has been replaced
+// by the ITedsWidget interface
+using STedsWidget UE_DEPRECATED(5.6, "Use the new ITedsWidget interface directly instead") = UE::Editor::DataStorage::ITedsWidget;
\ No newline at end of file
diff --git a/Engine/Source/Runtime/TypedElementFramework/Public/Elements/Interfaces/TypedElementDataStorageUiInterface.h b/Engine/Source/Runtime/TypedElementFramework/Public/Elements/Interfaces/TypedElementDataStorageUiInterface.h
--- a/Engine/Source/Runtime/TypedElementFramework/Public/Elements/Interfaces/TypedElementDataStorageUiInterface.h
+++ b/Engine/Source/Runtime/TypedElementFramework/Public/Elements/Interfaces/TypedElementDataStorageUiInterface.h
@@ -14,6 +14,11 @@
 
 #include "TypedElementDataStorageUiInterface.generated.h"
 
+namespace UE::Editor::DataStorage
+{
+	class ITedsWidget;
+}
+
 class IEditorDataStorageUiProvider;
 class SWidget;
 
@@ -346,6 +351,12 @@ public:
 	virtual bool SupportsExtension(FName Extension) const = 0;
 	/** Provides a list of all extensions that are enabled. */
 	virtual void ListExtensions(TFunctionRef<void(FName)> Callback) const = 0;
+
+	/** Create the container widget that every TEDS UI widget is stored in */
+	virtual TSharedPtr<UE::Editor::DataStorage::ITedsWidget> CreateContainerTedsWidget(UE::Editor::DataStorage::RowHandle UiRowHandle) const = 0;
+
+	/** Get the table where TEDS UI widgets are stored */
+	virtual UE::Editor::DataStorage::TableHandle GetWidgetTable() const = 0;
 };
 
 
