[Truncation Warnings] Fix and enable truncation warnings in TedsCore module

#rnx

diff --git a/Engine/Plugins/Experimental/EditorDataStorage/Source/TedsCore/Private/Compatibility/TypedElementActorViewportProcessors.cpp b/Engine/Plugins/Experimental/EditorDataStorage/Source/TedsCore/Private/Compatibility/TypedElementActorViewportProcessors.cpp
--- a/Engine/Plugins/Experimental/EditorDataStorage/Source/TedsCore/Private/Compatibility/TypedElementActorViewportProcessors.cpp
+++ b/Engine/Plugins/Experimental/EditorDataStorage/Source/TedsCore/Private/Compatibility/TypedElementActorViewportProcessors.cpp
@@ -108,7 +108,7 @@ namespace UE::Editor::DataStorage::Private
 				
 					// Parse the color
 					FColor Color = FColor::FromHex(Args[0]);
-					Color.A = FMath::Clamp(Color.A, 0, 128);
+					Color.A = FMath::Clamp<uint8>(Color.A, 0, 128);
 
 					TArray<RowHandle> RowHandles;
 				
diff --git a/Engine/Plugins/Experimental/EditorDataStorage/Source/TedsCore/Private/Debug/TypedElementDatabaseDebugCommands.cpp b/Engine/Plugins/Experimental/EditorDataStorage/Source/TedsCore/Private/Debug/TypedElementDatabaseDebugCommands.cpp
--- a/Engine/Plugins/Experimental/EditorDataStorage/Source/TedsCore/Private/Debug/TypedElementDatabaseDebugCommands.cpp
+++ b/Engine/Plugins/Experimental/EditorDataStorage/Source/TedsCore/Private/Debug/TypedElementDatabaseDebugCommands.cpp
@@ -320,7 +320,7 @@ static FAutoConsoleCommand CVarAddToDynamicColumn(
 		const uint64 RowAsU64 = FCString::Strtoui64(*Args[0], nullptr, 10);
 		const RowHandle Row = RowAsU64;
 		const FName TagId = FName(*Args[1]);
-		const uint64 Value = FCString::Strtoui64(*Args[2], nullptr, 10);
+		const int32 Value = FCString::Strtoi(*Args[2], nullptr, 10);
 		const uint64 MethodId = Args.Num() >= 4 ? FCString::Strtoui64(*Args[3], nullptr, 10) : 0;
 
 		// Two methods
diff --git a/Engine/Plugins/Experimental/EditorDataStorage/Source/TedsCore/Private/TypedElementDatabase.cpp b/Engine/Plugins/Experimental/EditorDataStorage/Source/TedsCore/Private/TypedElementDatabase.cpp
--- a/Engine/Plugins/Experimental/EditorDataStorage/Source/TedsCore/Private/TypedElementDatabase.cpp
+++ b/Engine/Plugins/Experimental/EditorDataStorage/Source/TedsCore/Private/TypedElementDatabase.cpp
@@ -55,11 +55,11 @@ FAutoConsoleCommandWithOutputDevice PrintSupportedColumnsConsoleCommand(
 					int32 StructureSize = It->GetStructureSize();
 					if (StructureSize >= 1024 * 1024)
 					{
-						Output.Logf(TEXT("    [%6.2f mib] %s"), StructureSize / (1024.0f * 1024.0f), *It->GetFullName());
+						Output.Logf(TEXT("    [%6.2f mib] %s"), static_cast<float>(StructureSize) / (1024.0f * 1024.0f), *It->GetFullName());
 					}
 					else if (StructureSize >= 1024)
 					{
-						Output.Logf(TEXT("    [%6.2f kib] %s"), StructureSize / 1024.0f, *It->GetFullName());
+						Output.Logf(TEXT("    [%6.2f kib] %s"), static_cast<float>(StructureSize) / 1024.0f, *It->GetFullName());
 					}
 					else
 					{
@@ -98,6 +98,12 @@ namespace UE::Editor::DataStorage::Private
 		}
 		return Result;
 	}
+	
+	constexpr int32 ConvertTableHandleToIndex(const TableHandle InTableHandle)
+	{
+		ensure(InTableHandle == InvalidTableHandle || (InTableHandle <= static_cast<TableHandle>(std::numeric_limits<int32>::max()) && InTableHandle >= 0));
+		return static_cast<int32>(InTableHandle);
+	}
 } // namespace UE::Editor::DataStorage::Private
 
 void UEditorDataStorage::Initialize()
@@ -289,13 +295,14 @@ UE::Editor::DataStorage::TableHandle UEditorDataStorage::RegisterTable(UE::Edito
 {
 	using namespace UE::Editor::DataStorage;
 
-	if (ActiveEditorEntityManager && SourceTable < Tables.Num() && !TableNameLookup.Contains(Name))
+	const int32 TableHandleAsIndex = Private::ConvertTableHandleToIndex(SourceTable);
+	if (ActiveEditorEntityManager && TableHandleAsIndex < Tables.Num() && !TableNameLookup.Contains(Name))
 	{
 		TableHandle Result = Tables.Num();
 		FMassArchetypeCreationParams ArchetypeCreationParams;
 		ArchetypeCreationParams.DebugName = Name;
 		ArchetypeCreationParams.ChunkMemorySize = GetTableChunkSize(Name);
-		Tables.Add(ActiveEditorEntityManager->CreateArchetype(Tables[SourceTable], ColumnList, ArchetypeCreationParams));
+		Tables.Add(ActiveEditorEntityManager->CreateArchetype(Tables[TableHandleAsIndex], ColumnList, ArchetypeCreationParams));
 		if (Name.IsValid())
 		{
 			TableNameLookup.Add(Name, Result);
@@ -346,9 +353,12 @@ void UEditorDataStorage::BatchReserveRows(TArrayView<RowHandle> ReservedRows)
 
 UE::Editor::DataStorage::RowHandle UEditorDataStorage::AddRow(TableHandle Table)
 {
-	checkf(Table < Tables.Num(), TEXT("Attempting to add a row to a non-existing table."));
+	using namespace UE::Editor::DataStorage;
+
+	const int32 TableHandleAsIndex = Private::ConvertTableHandleToIndex(Table);
+	checkf(TableHandleAsIndex < Tables.Num(), TEXT("Attempting to add a row to a non-existing table."));
 	return ActiveEditorEntityManager 
-		? ActiveEditorEntityManager->CreateEntity(Tables[Table]).AsNumber() 
+		? ActiveEditorEntityManager->CreateEntity(Tables[TableHandleAsIndex]).AsNumber()
 		: UE::Editor::DataStorage::InvalidRowHandle;
 }
 
@@ -360,12 +370,13 @@ UE::Editor::DataStorage::RowHandle UEditorDataStorage::AddRow(TableHandle Table,
 	OnCreated.CheckCallable();
 	if (ActiveEditorEntityManager)
 	{
-		checkf(Table < Tables.Num(), TEXT("Attempting to a row to a non-existing table."));
+		const int32 TableHandleAsIndex = Private::ConvertTableHandleToIndex(Table);
+		checkf(TableHandleAsIndex < Tables.Num(), TEXT("Attempting to a row to a non-existing table."));
 
 		TArray<FMassEntityHandle> Entity;
 		Entity.Reserve(1);
 		TSharedRef<FMassEntityManager::FEntityCreationContext> Context =
-			ActiveEditorEntityManager->BatchCreateEntities(Tables[Table], 1, Entity);
+			ActiveEditorEntityManager->BatchCreateEntities(Tables[TableHandleAsIndex], 1, Entity);
 
 		checkf(!Entity.IsEmpty(), TEXT("Add row tried to create a new row but none were provided by the backend."));
 		RowHandle Result = Entity[0].AsNumber();
@@ -377,11 +388,14 @@ UE::Editor::DataStorage::RowHandle UEditorDataStorage::AddRow(TableHandle Table,
 
 bool UEditorDataStorage::AddRow(RowHandle ReservedRow, TableHandle Table)
 {
+	using namespace UE::Editor::DataStorage;
+
+	const int32 TableHandleAsIndex = Private::ConvertTableHandleToIndex(Table);
 	checkf(!IsRowAssigned(ReservedRow), TEXT("Attempting to assign a table to row that already has a table assigned."));
-	checkf(Table < Tables.Num(), TEXT("Attempting to add a row to a non-existing table."));
+	checkf(TableHandleAsIndex < Tables.Num(), TEXT("Attempting to add a row to a non-existing table."));
 	if (ActiveEditorEntityManager)
 	{
-		ActiveEditorEntityManager->BuildEntity(FMassEntityHandle::FromNumber(ReservedRow), Tables[Table]);
+		ActiveEditorEntityManager->BuildEntity(FMassEntityHandle::FromNumber(ReservedRow), Tables[TableHandleAsIndex]);
 		return true;
 	}
 	else
@@ -393,13 +407,16 @@ bool UEditorDataStorage::AddRow(RowHandle ReservedRow, TableHandle Table)
 bool UEditorDataStorage::AddRow(RowHandle ReservedRow, TableHandle Table,
 	RowCreationCallbackRef OnCreated)
 {
+	using namespace UE::Editor::DataStorage;
+
 	OnCreated.CheckCallable();
 	if (ActiveEditorEntityManager)
 	{
-		checkf(Table < Tables.Num(), TEXT("Attempting to add a row to a non-existing table."));
+		const int32 TableHandleAsIndex = Private::ConvertTableHandleToIndex(Table);
+		checkf(TableHandleAsIndex < Tables.Num(), TEXT("Attempting to add a row to a non-existing table."));
 		
 		TSharedRef<FMassEntityManager::FEntityCreationContext> Context =
-			ActiveEditorEntityManager->BatchCreateReservedEntities(Tables[Table], { FMassEntityHandle::FromNumber(ReservedRow) });
+			ActiveEditorEntityManager->BatchCreateReservedEntities(Tables[TableHandleAsIndex], { FMassEntityHandle::FromNumber(ReservedRow) });
 
 		OnCreated(ReservedRow);
 		return true;
@@ -410,15 +427,18 @@ bool UEditorDataStorage::AddRow(RowHandle ReservedRow, TableHandle Table,
 bool UEditorDataStorage::BatchAddRow(
 	TableHandle Table, int32 Count, RowCreationCallbackRef OnCreated)
 {
+	using namespace UE::Editor::DataStorage;
+
 	OnCreated.CheckCallable();
 	if (ActiveEditorEntityManager)
 	{
-		checkf(Table < Tables.Num(), TEXT("Attempting to add multiple rows to a non-existing table."));
+		const int32 TableHandleAsIndex = Private::ConvertTableHandleToIndex(Table);
+		checkf(TableHandleAsIndex < Tables.Num(), TEXT("Attempting to add multiple rows to a non-existing table."));
 	
 		TArray<FMassEntityHandle> Entities;
 		Entities.Reserve(Count);
 		TSharedRef<FMassEntityManager::FEntityCreationContext> Context = 
-			ActiveEditorEntityManager->BatchCreateEntities(Tables[Table], Count, Entities);
+			ActiveEditorEntityManager->BatchCreateEntities(Tables[TableHandleAsIndex], Count, Entities);
 		
 		for (FMassEntityHandle Entity : Entities)
 		{
@@ -439,10 +459,11 @@ bool UEditorDataStorage::BatchAddRow(TableHandle Table,
 	OnCreated.CheckCallable();
 	if (ActiveEditorEntityManager)
 	{
-		checkf(Table < Tables.Num(), TEXT("Attempting to add multiple rows to a non-existing table."));
+		const int32 TableHandleAsIndex = Private::ConvertTableHandleToIndex(Table);
+		checkf(TableHandleAsIndex < Tables.Num(), TEXT("Attempting to add multiple rows to a non-existing table."));
 	
 		TSharedRef<FMassEntityManager::FEntityCreationContext> Context =
-			ActiveEditorEntityManager->BatchCreateReservedEntities(Tables[Table], RowsToMassEntitiesConversion(ReservedHandles));
+			ActiveEditorEntityManager->BatchCreateReservedEntities(Tables[TableHandleAsIndex], RowsToMassEntitiesConversion(ReservedHandles));
 
 		for (RowHandle Entity : ReservedHandles)
 		{
@@ -1284,10 +1305,12 @@ TSharedPtr<const UE::Editor::DataStorage::FEnvironment> UEditorDataStorage::GetE
 
 FMassArchetypeHandle UEditorDataStorage::LookupArchetype(UE::Editor::DataStorage::TableHandle InTableHandle) const
 {
-	const uint32 TableIndex = InTableHandle;
-	if (Tables.IsValidIndex(TableIndex))
+	using namespace UE::Editor::DataStorage;
+
+	const int32 TableHandleAsIndex = Private::ConvertTableHandleToIndex(InTableHandle);
+	if (Tables.IsValidIndex(TableHandleAsIndex))
 	{
-		return Tables[TableIndex];
+		return Tables[TableHandleAsIndex];
 	}
 	return FMassArchetypeHandle();
 }
diff --git a/Engine/Plugins/Experimental/EditorDataStorage/Source/TedsCore/Private/TypedElementDatabaseIndexTable.cpp b/Engine/Plugins/Experimental/EditorDataStorage/Source/TedsCore/Private/TypedElementDatabaseIndexTable.cpp
--- a/Engine/Plugins/Experimental/EditorDataStorage/Source/TedsCore/Private/TypedElementDatabaseIndexTable.cpp
+++ b/Engine/Plugins/Experimental/EditorDataStorage/Source/TedsCore/Private/TypedElementDatabaseIndexTable.cpp
@@ -258,7 +258,7 @@ namespace UE::Editor::DataStorage
 			// Skip rows with an invalid handle as those are already freed.
 			if (*Front != InvalidRowHandle && !DataStorage.IsRowAvailableUnsafe(*Front))
 			{
-				Result.Add(Front - Begin);
+				Result.Add(static_cast<int32>(Front - Begin));
 			}
 		}
 
@@ -317,11 +317,11 @@ namespace UE::Editor::DataStorage
 	{
 		if (RemainingFrameTime >= FCleanUpInfo::JobGrowthThreshold)
 		{
-			CleanUpInfo.MaxNumJobs *= FCleanUpInfo::JobGrowthFactor;
+			CleanUpInfo.MaxNumJobs = FMath::RoundToInt32(static_cast<float>(CleanUpInfo.MaxNumJobs) * FCleanUpInfo::JobGrowthFactor);
 		}
 		else if (RemainingFrameTime <= FCleanUpInfo::JobShrinkThreshold)
 		{
-			CleanUpInfo.MaxNumJobs *= FCleanUpInfo::JobShrinkFactor;
+			CleanUpInfo.MaxNumJobs = FMath::RoundToInt32(static_cast<float>(CleanUpInfo.MaxNumJobs) * FCleanUpInfo::JobShrinkFactor);
 		}
 
 		CleanUpInfo.MaxNumJobs = FMath::Clamp(CleanUpInfo.MaxNumJobs, FCleanUpInfo::MinJobCount, FCleanUpInfo::MaxJobCount);
@@ -332,14 +332,14 @@ namespace UE::Editor::DataStorage
 		if (CleanUpInfo.BatchWentOverTime.load() >= FCleanUpInfo::BatchShrinkThreshold)
 		{
 			// The batch size was too big to complete in the allotted time so take a sizable chunk off.
-			CleanUpInfo.BatchSize *= FCleanUpInfo::BatchShrinkFactor;
+			CleanUpInfo.BatchSize = FMath::RoundToInt32(static_cast<float>(CleanUpInfo.BatchSize) * FCleanUpInfo::BatchShrinkFactor);
 		}
 		// Check if more than the minimum percentage of jobs took less time than currently allocated.
 		else if (CleanUpInfo.BatchWentUnderTime.load() >= 
-			static_cast<uint32>(FMath::RoundToInt32(JobCount * FCleanUpInfo::BatchIncreaseThreshold)))
+			static_cast<uint32>(FMath::RoundToInt32(static_cast<float>(JobCount) * FCleanUpInfo::BatchIncreaseThreshold)))
 		{
 			// The batch size fitted within the allotted time so slightly increase it.
-			CleanUpInfo.BatchSize *= FCleanUpInfo::BatchGrowthFactor;
+			CleanUpInfo.BatchSize = FMath::RoundToInt32(static_cast<float>(CleanUpInfo.BatchSize) * FCleanUpInfo::BatchGrowthFactor);
 		}
 		// Clamp the job batch size within reasonable sizes to avoid extremes.
 		CleanUpInfo.BatchSize = FMath::Clamp(CleanUpInfo.BatchSize, FCleanUpInfo::MinBatchSize, FCleanUpInfo::MaxBatchSize);
diff --git a/Engine/Plugins/Experimental/EditorDataStorage/Source/TedsCore/Private/TypedElementDatabaseScratchBuffer.cpp b/Engine/Plugins/Experimental/EditorDataStorage/Source/TedsCore/Private/TypedElementDatabaseScratchBuffer.cpp
--- a/Engine/Plugins/Experimental/EditorDataStorage/Source/TedsCore/Private/TypedElementDatabaseScratchBuffer.cpp
+++ b/Engine/Plugins/Experimental/EditorDataStorage/Source/TedsCore/Private/TypedElementDatabaseScratchBuffer.cpp
@@ -11,7 +11,7 @@ namespace UE::Editor::DataStorage
 	{
 	}
 
-	void* FScratchBuffer::Allocate(size_t Size, size_t Alignment)
+	void* FScratchBuffer::Allocate(SIZE_T Size, uint32 Alignment)
 	{
 		return CurrentAllocator.load()->Malloc(Size, Alignment);
 	}
diff --git a/Engine/Plugins/Experimental/EditorDataStorage/Source/TedsCore/Private/TypedElementDatabaseScratchBuffer.h b/Engine/Plugins/Experimental/EditorDataStorage/Source/TedsCore/Private/TypedElementDatabaseScratchBuffer.h
--- a/Engine/Plugins/Experimental/EditorDataStorage/Source/TedsCore/Private/TypedElementDatabaseScratchBuffer.h
+++ b/Engine/Plugins/Experimental/EditorDataStorage/Source/TedsCore/Private/TypedElementDatabaseScratchBuffer.h
@@ -27,7 +27,7 @@ namespace UE::Editor::DataStorage
 	public:
 		FScratchBuffer();
 	
-		void* Allocate(size_t Size, size_t Alignment);
+		void* Allocate(SIZE_T Size, uint32 Alignment);
 		template<typename T, typename... ArgTypes>
 		T* Emplace(ArgTypes... Args);
 		template<typename T, typename... ArgTypes>
diff --git a/Engine/Plugins/Experimental/EditorDataStorage/Source/TedsCore/TedsCore.Build.cs b/Engine/Plugins/Experimental/EditorDataStorage/Source/TedsCore/TedsCore.Build.cs
--- a/Engine/Plugins/Experimental/EditorDataStorage/Source/TedsCore/TedsCore.Build.cs
+++ b/Engine/Plugins/Experimental/EditorDataStorage/Source/TedsCore/TedsCore.Build.cs
@@ -6,6 +6,9 @@ public class TedsCore : ModuleRules
 {
 	public TedsCore(ReadOnlyTargetRules Target) : base(Target)
 	{
+		// Enable truncation warnings in this module
+		CppCompileWarningSettings.UnsafeTypeCastWarningLevel = WarningLevel.Warning;
+
 		PCHUsage = ModuleRules.PCHUsageMode.UseExplicitOrSharedPCHs;
 
 		if (Target.bBuildEditor)
