TEDS UI: Add support for dynamic column template matching

diff --git a/Engine/Plugins/Experimental/EditorDataStorage/Source/TedsCore/Private/TypedElementDatabase.cpp b/Engine/Plugins/Experimental/EditorDataStorage/Source/TedsCore/Private/TypedElementDatabase.cpp
--- a/Engine/Plugins/Experimental/EditorDataStorage/Source/TedsCore/Private/TypedElementDatabase.cpp
+++ b/Engine/Plugins/Experimental/EditorDataStorage/Source/TedsCore/Private/TypedElementDatabase.cpp
@@ -996,43 +996,15 @@ bool UEditorDataStorage::MatchesColumns(RowHandle Row, const UE::Editor::DataSto
 	if (ActiveEditorEntityManager)
 	{
 		checkf(Conditions.IsCompiled(), TEXT("Query Conditions must be compiled before they can be used"));
-		
-		FMassEntityHandle Entity = FMassEntityHandle::FromNumber(Row);
-		if (ActiveEditorEntityManager->IsEntityActive(Entity))
-		{
-			FMassArchetypeHandle Archetype = ActiveEditorEntityManager->GetArchetypeForEntity(Entity);
-			const FMassArchetypeCompositionDescriptor& Composition = ActiveEditorEntityManager->GetArchetypeComposition(Archetype);
 
-			auto Callback = [&Composition](uint8_t ColumnIndex, TWeakObjectPtr<const UScriptStruct> ColumnType)
-			{
-				if (ColumnType.IsValid())
-				{
-					if (UE::Mass::IsA<FMassFragment>(ColumnType.Get()))
-					{
-						return Composition.Fragments.Contains(*ColumnType);
-					}
-					else if (UE::Mass::IsA<FMassTag>(ColumnType.Get()))
-					{
-						return Composition.Tags.Contains(*ColumnType);
-					}
-				}
-				return false;
-			};
-			return Conditions.Verify(Callback);
-		}
-		else
+		TArray<TWeakObjectPtr<const UScriptStruct>> Columns;
+		
+		ListColumns(Row, [&Columns](const UScriptStruct& InColumn)
 		{
-			auto Callback = [Row, &CommandBuffer = Environment->GetDirectDeferredCommands()]
-				(uint8_t ColumnIndex, TWeakObjectPtr<const UScriptStruct> ColumnType)
-			{
-				if (ColumnType.IsValid())
-				{
-					return CommandBuffer.HasColumn(Row, ColumnType.Get());
-				}
-				return false;
-			};
-			return Conditions.Verify(Callback);
-		}
+			Columns.Add(&InColumn);
+		});
+		
+		return Conditions.Verify(Columns);
 	}
 	return false;
 }
@@ -1047,6 +1019,11 @@ const UScriptStruct* UEditorDataStorage::GenerateDynamicColumn(const UE::Editor:
 	return Environment->GenerateDynamicColumn(*Description.TemplateType, Description.Identifier);
 }
 
+void UEditorDataStorage::ForEachDynamicColumn(const UScriptStruct& Template, TFunctionRef<void(const UScriptStruct& Type)> Callback) const
+{
+	return Environment->ForEachDynamicColumn(Template, Callback);
+}
+
 void UEditorDataStorage::RegisterTickGroup(
 	FName GroupName, EQueryTickPhase Phase, FName BeforeGroup, FName AfterGroup, UE::Editor::DataStorage::EExecutionMode ExecutionMode)
 {
diff --git a/Engine/Plugins/Experimental/EditorDataStorage/Source/TedsCore/Private/TypedElementDatabase.h b/Engine/Plugins/Experimental/EditorDataStorage/Source/TedsCore/Private/TypedElementDatabase.h
--- a/Engine/Plugins/Experimental/EditorDataStorage/Source/TedsCore/Private/TypedElementDatabase.h
+++ b/Engine/Plugins/Experimental/EditorDataStorage/Source/TedsCore/Private/TypedElementDatabase.h
@@ -144,6 +144,7 @@ public:
 
 	virtual const UScriptStruct* FindDynamicColumn(const UE::Editor::DataStorage::FDynamicColumnDescription& Description) const override;
 	virtual const UScriptStruct* GenerateDynamicColumn(const UE::Editor::DataStorage::FDynamicColumnDescription& Description) override;
+	virtual void ForEachDynamicColumn(const UScriptStruct& Template, TFunctionRef<void(const UScriptStruct& Type)> Callback) const override;
 
 	void RegisterTickGroup(FName GroupName, EQueryTickPhase Phase, FName BeforeGroup, FName AfterGroup, UE::Editor::DataStorage::EExecutionMode ExecutionMode);
 	void UnregisterTickGroup(FName GroupName, EQueryTickPhase Phase);
diff --git a/Engine/Plugins/Experimental/EditorDataStorage/Source/TedsCore/Private/TypedElementDatabaseEnvironment.cpp b/Engine/Plugins/Experimental/EditorDataStorage/Source/TedsCore/Private/TypedElementDatabaseEnvironment.cpp
--- a/Engine/Plugins/Experimental/EditorDataStorage/Source/TedsCore/Private/TypedElementDatabaseEnvironment.cpp
+++ b/Engine/Plugins/Experimental/EditorDataStorage/Source/TedsCore/Private/TypedElementDatabaseEnvironment.cpp
@@ -114,6 +114,17 @@ namespace UE::Editor::DataStorage
 		return ValueTagManager.GenerateColumnType(Tag);
 	}
 
+	void FEnvironment::ForEachDynamicColumn(const UScriptStruct& Template, TFunctionRef<void(const UScriptStruct& Type)> Callback) const
+	{
+		DynamicColumnGenerator.ForEachDynamicColumn(Template, [Callback](const FDynamicColumnGeneratorInfo& Info)
+    		{
+    			if (Info.Type)
+    			{
+    				Callback(*Info.Type);
+    			}
+    		});
+	}
+
 	void FEnvironment::NextUpdateCycle()
 	{
 		Queries.UpdateActivatableQueries();
diff --git a/Engine/Plugins/Experimental/EditorDataStorage/Source/TedsCore/Private/TypedElementDatabaseEnvironment.h b/Engine/Plugins/Experimental/EditorDataStorage/Source/TedsCore/Private/TypedElementDatabaseEnvironment.h
--- a/Engine/Plugins/Experimental/EditorDataStorage/Source/TedsCore/Private/TypedElementDatabaseEnvironment.h
+++ b/Engine/Plugins/Experimental/EditorDataStorage/Source/TedsCore/Private/TypedElementDatabaseEnvironment.h
@@ -54,6 +54,9 @@ namespace UE::Editor::DataStorage
 		
 		// Creates or Finds the column type associated with the value tag
 		const UScriptStruct* GenerateColumnType(const FValueTag& Tag);
+
+		// Executes the given callback for each known dynamic column that derives from the base template provided
+		void ForEachDynamicColumn(const UScriptStruct& Template, TFunctionRef<void(const UScriptStruct& Type)> Callback) const;
 		
 		// Creates an instance of a value tag
 		FConstSharedStruct GenerateValueTag(const FValueTag& Tag, const FName& Value);
diff --git a/Engine/Plugins/Experimental/EditorDataStorage/Source/TedsCore/Private/TypedElementDatabaseUI.cpp b/Engine/Plugins/Experimental/EditorDataStorage/Source/TedsCore/Private/TypedElementDatabaseUI.cpp
--- a/Engine/Plugins/Experimental/EditorDataStorage/Source/TedsCore/Private/TypedElementDatabaseUI.cpp
+++ b/Engine/Plugins/Experimental/EditorDataStorage/Source/TedsCore/Private/TypedElementDatabaseUI.cpp
@@ -8,6 +8,7 @@
 #include "Elements/Columns/TypedElementMiscColumns.h"
 #include "Elements/Columns/TypedElementSlateWidgetColumns.h"
 #include "Elements/Columns/TypedElementTypeInfoColumns.h"
+#include "Elements/Framework/TypedElementColumnUtils.h"
 #include "Elements/Interfaces/TypedElementDataStorageInterface.h"
 #include "Elements/Interfaces/TypedElementDataStorageCompatibilityInterface.h"
 #include "GenericPlatform/GenericPlatformMemory.h"
@@ -26,6 +27,23 @@ namespace Internal
 	};
 
 	template<class... Ts> TOverloaded(Ts...) -> TOverloaded<Ts...>;
+
+	// Check if the two columns are equal, or if InRequestedColumn is a dynamic specialization of InMatchedColumn
+	bool CheckSingleColumnMatch(const TWeakObjectPtr<const UScriptStruct>& InMatchedColumn, const TWeakObjectPtr<const UScriptStruct>& InRequestedColumn)
+	{
+		if (InMatchedColumn == InRequestedColumn)
+		{
+			return true;
+		}
+		
+		else if (UE::Editor::DataStorage::ColumnUtils::IsDynamicTemplate(InMatchedColumn.Get()) &&
+			UE::Editor::DataStorage::ColumnUtils::IsDerivedFromDynamicTemplate(InRequestedColumn.Get()))
+		{
+			return InRequestedColumn->IsChildOf(InMatchedColumn.Get());
+		}
+
+		return false;
+	}
 }
 
 void UEditorDataStorageUi::Initialize(
@@ -285,13 +303,12 @@ void UEditorDataStorageUi::CreateWidgetConstructors(FName Purpose, EMatchApproac
 {
 	if (FPurposeInfo* PurposeInfo = WidgetPurposes.Find(Purpose))
 	{
-		// Sort so searching can be done in a single pass. This would also allow for binary searching, but the number of columns
-		// is typically small enough for a binary search to end up being more expensive than a linear search. This may change
-		// if/when there are a sufficient enough number of widgets that are bound to a large number of columns.
+		// Sort by name so that removing the matched columns can be done in a single pass.
+		// Sorting by ptr does not work because dynamic column ptrs are different than their base template
 		Columns.Sort(
 			[](const TWeakObjectPtr<const UScriptStruct>& Lhs, const TWeakObjectPtr<const UScriptStruct>& Rhs)
 			{
-				return Lhs.Get() < Rhs.Get();
+				return Lhs->GetName() < Rhs->GetName();
 			});
 
 		if (!PurposeInfo->bIsSorted)
@@ -301,8 +318,20 @@ void UEditorDataStorageUi::CreateWidgetConstructors(FName Purpose, EMatchApproac
 			PurposeInfo->Factories.StableSort(
 				[this](const FWidgetFactory& Lhs, const FWidgetFactory& Rhs)
 				{
-					int32 LeftSize = Lhs.GetConditions(Storage).MinimumColumnMatchRequired();
-					int32 RightSize = Rhs.GetConditions(Storage).MinimumColumnMatchRequired();
+					const UE::Editor::DataStorage::Queries::FConditions& LhsConditions = Lhs.GetConditions(Storage);
+					const UE::Editor::DataStorage::Queries::FConditions& RhsConditions = Rhs.GetConditions(Storage);
+					
+					int32 LeftSize = LhsConditions.MinimumColumnMatchRequired();
+					int32 RightSize = RhsConditions.MinimumColumnMatchRequired();
+
+					// If two factories are the same size, we want factories containing dynamic templates to be at the end so
+					// they are de-prioritized when matching and factories with dynamic specializations are matched first.
+					// e.g A widget factory for ColumnA("Apple") or ColumnA("Orange") should be considered before a generic one for ColumnA.
+					if (LeftSize == RightSize)
+					{
+						return !LhsConditions.UsesDynamicTemplates();
+					}
+					
 					return LeftSize > RightSize;
 				});
 			PurposeInfo->bIsSorted = true;
@@ -498,11 +527,20 @@ void UEditorDataStorageUi::CreateWidgetConstructors_LongestMatch(const TArray<FW
 
 		MatchedColumns.Reset();
 		
-		if (Conditions.Verify(MatchedColumns, Columns, true))
+		if (Conditions.Verify(MatchedColumns, Columns))
 		{
 			// Remove the found columns from the requested list.
-			Algo::SortBy(MatchedColumns, [](const TWeakObjectPtr<const UScriptStruct>& Column) { return Column.Get(); });
+			MatchedColumns.Sort(
+				[](const TWeakObjectPtr<const UScriptStruct>& Lhs, const TWeakObjectPtr<const UScriptStruct>& Rhs)
+				{
+					return Lhs->GetName() < Rhs->GetName();
+				});
+			
 			MatchedColumns.SetNum(Algo::Unique(MatchedColumns), EAllowShrinking::No);
+
+			// We need to keep track of the columns the user requested that ended up matching separately because MatchedColumns could contain the
+			// base template for a dynamic column in the requested columns that actually matched - and the widget constructor wants the latter
+			TArray<TWeakObjectPtr<const UScriptStruct>> RequestedColumnsThatMatched;
 			
 			TWeakObjectPtr<const UScriptStruct>* ColumnsIt = Columns.GetData();
 			TWeakObjectPtr<const UScriptStruct>* ColumnsEnd = ColumnsIt + Columns.Num();
@@ -510,7 +548,7 @@ void UEditorDataStorageUi::CreateWidgetConstructors_LongestMatch(const TArray<FW
 			for (const TWeakObjectPtr<const UScriptStruct>& MatchedColumn : MatchedColumns)
 			{
 				// Remove all the columns that were matched from the provided column list.
-				while (*ColumnsIt != MatchedColumn)
+				while (!Internal::CheckSingleColumnMatch(MatchedColumn, *ColumnsIt))
 				{
 					++ColumnIndex;
 					++ColumnsIt;
@@ -520,11 +558,12 @@ void UEditorDataStorageUi::CreateWidgetConstructors_LongestMatch(const TArray<FW
 						return;
 					}
 				}
+				RequestedColumnsThatMatched.Add(*ColumnsIt);
 				Columns.RemoveAt(ColumnIndex, EAllowShrinking::No);
 				--ColumnsEnd;
 			}
 			
-			if (!CreateSingleWidgetConstructor(FactoryIt->Constructor, Arguments, MoveTemp(MatchedColumns), Conditions, Callback))
+			if (!CreateSingleWidgetConstructor(FactoryIt->Constructor, Arguments, MoveTemp(RequestedColumnsThatMatched), Conditions, Callback))
 			{
 				return;
 			}
@@ -551,7 +590,7 @@ void UEditorDataStorageUi::CreateWidgetConstructors_ExactMatch(const TArray<FWid
 
 		MatchedColumns.Reset();
 
-		if (Conditions.Verify(MatchedColumns, Columns, true))
+		if (Conditions.Verify(MatchedColumns, Columns))
 		{
 			Algo::SortBy(MatchedColumns, [](const TWeakObjectPtr<const UScriptStruct>& Column) { return Column.Get(); });
 			MatchedColumns.SetNum(Algo::Unique(MatchedColumns), EAllowShrinking::No);
@@ -591,20 +630,22 @@ void UEditorDataStorageUi::CreateWidgetConstructors_SingleMatch(const TArray<FWi
 				continue;
 			}
 
-			if (ColumnData[0] == Columns[ColumnIndex])
+			if (Internal::CheckSingleColumnMatch(ColumnData[0], Columns[ColumnIndex]))
 			{
+				// We need to keep a copy of the actually requested column because the matched column could be the base template for the actually matched
+				// dynamic column and the widget constructor wants the latter
+				TWeakObjectPtr<const UScriptStruct> RequestedColumn = Columns[ColumnIndex];
 				Columns.RemoveAt(ColumnIndex);
-				CreateSingleWidgetConstructor((*FactoryIt).Constructor, Arguments, 
-					TArray<TWeakObjectPtr<const UScriptStruct>>(ColumnData), Conditions, Callback);
+				CreateSingleWidgetConstructor((*FactoryIt).Constructor, Arguments, {RequestedColumn}, Conditions, Callback);
+				
+				
 				// Match was found so move on to the next column in the column.
 				break;
-			}
+			} 
 		}
 	}
 }
 
-
-
 //
 // FWidgetFactory
 //
diff --git a/Engine/Source/Runtime/TypedElementFramework/Private/Elements/Interfaces/TypedElementDataStorageUiInterface.cpp b/Engine/Source/Runtime/TypedElementFramework/Private/Elements/Interfaces/TypedElementDataStorageUiInterface.cpp
--- a/Engine/Source/Runtime/TypedElementFramework/Private/Elements/Interfaces/TypedElementDataStorageUiInterface.cpp
+++ b/Engine/Source/Runtime/TypedElementFramework/Private/Elements/Interfaces/TypedElementDataStorageUiInterface.cpp
@@ -19,6 +19,17 @@ bool FTypedElementWidgetConstructor::Initialize(const UE::Editor::DataStorage::F
 {
 	MatchedColumnTypes = MoveTemp(InMatchedColumnTypes);
 	QueryConditions = &InQueryConditions;
+
+	// If we matched with any dynamic columns, store a mapping from the base template to the identifier so the widget can look it up later
+	for (const TWeakObjectPtr<const UScriptStruct>& Column : MatchedColumnTypes)
+	{
+		FName Identifier = UE::Editor::DataStorage::ColumnUtils::GetDynamicColumnIdentifier(Column.Get());
+		if (!Identifier.IsNone())
+		{
+			MatchedDynamicTemplates.Emplace(Cast<UScriptStruct>(Column->GetSuperStruct()), Identifier);
+		}
+	}
+	
 	return true;
 }
 
diff --git a/Engine/Source/Runtime/TypedElementFramework/Public/Elements/Common/TypedElementQueryConditions.h b/Engine/Source/Runtime/TypedElementFramework/Public/Elements/Common/TypedElementQueryConditions.h
--- a/Engine/Source/Runtime/TypedElementFramework/Public/Elements/Common/TypedElementQueryConditions.h
+++ b/Engine/Source/Runtime/TypedElementFramework/Public/Elements/Common/TypedElementQueryConditions.h
@@ -126,6 +126,8 @@ namespace UE::Editor::DataStorage::Queries
 		TYPEDELEMENTFRAMEWORK_API TConstArrayView<TWeakObjectPtr<const UScriptStruct>> GetColumns() const;
 		/** Whether or not there are any columns registered for operation.*/
 		TYPEDELEMENTFRAMEWORK_API bool IsEmpty() const;
+		/** Whether the conditions contain any columns that are dynamic templates. */
+		TYPEDELEMENTFRAMEWORK_API bool UsesDynamicTemplates() const;
 
 	private:
 		template<typename AvailableColumnType, typename ProjectionFunction>
@@ -144,8 +146,6 @@ namespace UE::Editor::DataStorage::Queries
 
 		uint8_t MinimumColumnMatchRequiredRange(uint8_t& Front) const;
 
-		bool UsesDynamicTemplates() const;
-
 		static void AppendQuery(FConditions& Target, const FConditions& Source);
 
 		static constexpr SIZE_T MaxColumnCount = 32;
diff --git a/Engine/Source/Runtime/TypedElementFramework/Public/Elements/Interfaces/TypedElementDataStorageInterface.h b/Engine/Source/Runtime/TypedElementFramework/Public/Elements/Interfaces/TypedElementDataStorageInterface.h
--- a/Engine/Source/Runtime/TypedElementFramework/Public/Elements/Interfaces/TypedElementDataStorageInterface.h
+++ b/Engine/Source/Runtime/TypedElementFramework/Public/Elements/Interfaces/TypedElementDataStorageInterface.h
@@ -338,6 +338,11 @@ public:
 	 */
 	virtual const UScriptStruct* GenerateDynamicColumn(const UE::Editor::DataStorage::FDynamicColumnDescription& Description) = 0;
 
+	/**
+	 * Executes the given callback for each known dynamic column that derives from the base template provided
+	 */
+	virtual void ForEachDynamicColumn(const UScriptStruct& Template, TFunctionRef<void(const UScriptStruct& Type)> Callback) const = 0;
+
 	/**
 	 * Outputs the registered query callbacks to the given output device for debugging purposes.
 	 */
diff --git a/Engine/Source/Runtime/TypedElementFramework/Public/Elements/Interfaces/TypedElementDataStorageUiInterface.h b/Engine/Source/Runtime/TypedElementFramework/Public/Elements/Interfaces/TypedElementDataStorageUiInterface.h
--- a/Engine/Source/Runtime/TypedElementFramework/Public/Elements/Interfaces/TypedElementDataStorageUiInterface.h
+++ b/Engine/Source/Runtime/TypedElementFramework/Public/Elements/Interfaces/TypedElementDataStorageUiInterface.h
@@ -128,6 +128,7 @@ protected:
 protected:
 
 	TArray<TWeakObjectPtr<const UScriptStruct>> MatchedColumnTypes;
+	TMap<TWeakObjectPtr<const UScriptStruct>, FName> MatchedDynamicTemplates;
 	const UE::Editor::DataStorage::Queries::FConditions* QueryConditions = nullptr;
 	const UScriptStruct* TypeInfo = nullptr;
 };
