Added support for retrieving delta time in the TEDS Query Callbacks.

diff --git a/Engine/Plugins/Experimental/EditorDataStorage/Source/TedsCore/Private/Processors/TypedElementProcessorAdaptors.cpp b/Engine/Plugins/Experimental/EditorDataStorage/Source/TedsCore/Private/Processors/TypedElementProcessorAdaptors.cpp
--- a/Engine/Plugins/Experimental/EditorDataStorage/Source/TedsCore/Private/Processors/TypedElementProcessorAdaptors.cpp
+++ b/Engine/Plugins/Experimental/EditorDataStorage/Source/TedsCore/Private/Processors/TypedElementProcessorAdaptors.cpp
@@ -117,6 +117,11 @@ namespace UE::Editor::DataStorage
 				const UScriptStruct* DynamicColumnType = Environment.FindDynamicColumn(*Description.TemplateType, Description.Identifier);
 				return DynamicColumnType;
 			}
+
+			float GetDeltaTimeSeconds() const 
+			{
+				return Context.GetDeltaTimeSeconds();
+			}
 		};
 
 		struct FMassWithEnvironmentContextCommon : public FMassContextCommon
@@ -661,7 +666,8 @@ namespace UE::Editor::DataStorage
 			void GetColumnsUnguarded(int32 TypeCount, char** RetrievedAddresses, const TWeakObjectPtr<const UScriptStruct>* ColumnTypes, const EQueryAccessType* AccessTypes) override { return Implementation.GetColumnsUnguarded(TypeCount, RetrievedAddresses, ColumnTypes, AccessTypes); }
 			bool HasColumn(const UScriptStruct* ColumnType) const override { return Implementation.HasColumn(ColumnType); }
 			bool HasColumn(RowHandle Row, const UScriptStruct* ColumnType) const override { return Implementation.HasColumn(Row, ColumnType); }
-			virtual const UScriptStruct* FindDynamicColumnType(const UE::Editor::DataStorage::FDynamicColumnDescription& Description) const override { return Implementation.FindDynamicColumnType(Description); }
+			const UScriptStruct* FindDynamicColumnType(const UE::Editor::DataStorage::FDynamicColumnDescription& Description) const override { return Implementation.FindDynamicColumnType(Description); }
+			float GetDeltaTimeSeconds() const override { return Implementation.GetDeltaTimeSeconds(); }
 
 			FMassContextCommon Implementation;
 		};
@@ -698,6 +704,7 @@ namespace UE::Editor::DataStorage
 			void RemoveColumns(RowHandle Row, TConstArrayView<const UScriptStruct*> ColumnTypes) override { Implementation.RemoveColumns(Row, ColumnTypes); }
 			void RemoveColumns(TConstArrayView<RowHandle> Rows, TConstArrayView<const UScriptStruct*> ColumnTypes) override { Implementation.RemoveColumns(Rows, ColumnTypes); }
 			const UScriptStruct* FindDynamicColumnType(const UE::Editor::DataStorage::FDynamicColumnDescription& Description) const override { return Implementation.FindDynamicColumnType(Description); }
+			float GetDeltaTimeSeconds() const override { return Implementation.GetDeltaTimeSeconds(); }
 			void PushCommand(void (*CommandFunction)(void*), void* CommandData) override { return Implementation.PushCommand(CommandFunction, CommandData); }
 		protected:
 			void* EmplaceObjectInScratch(const FEmplaceObjectParams& Params) override { return Implementation.EmplaceObjectInScratch(Params.ObjectSize, Params.Alignment, Params.Construct, Params.Destroy, Params.SourceObject); }
@@ -846,6 +853,7 @@ namespace UE::Editor::DataStorage
 			void RemoveColumns(RowHandle Row, TConstArrayView<const UScriptStruct*> ColumnTypes) override { Implementation.RemoveColumns(Row, ColumnTypes); }
 			void RemoveColumns(TConstArrayView<RowHandle> Rows, TConstArrayView<const UScriptStruct*> ColumnTypes) override { Implementation.RemoveColumns(Rows, ColumnTypes); }
 			const UScriptStruct* FindDynamicColumnType(const UE::Editor::DataStorage::FDynamicColumnDescription& Description) const override { return Implementation.FindDynamicColumnType(Description); }
+			float GetDeltaTimeSeconds() const override { return Implementation.GetDeltaTimeSeconds(); }
 			void PushCommand(void (*CommandFunction)(void*), void* Context) override { return Implementation.PushCommand(CommandFunction, Context); }
 
 			const UObject* GetDependency(const UClass* DependencyClass) override { return Implementation.GetDependency(DependencyClass); }
diff --git a/Engine/Source/Runtime/TypedElementFramework/Public/Elements/Interfaces/TypedElementQueryStorageInterfaces.h b/Engine/Source/Runtime/TypedElementFramework/Public/Elements/Interfaces/TypedElementQueryStorageInterfaces.h
--- a/Engine/Source/Runtime/TypedElementFramework/Public/Elements/Interfaces/TypedElementQueryStorageInterfaces.h
+++ b/Engine/Source/Runtime/TypedElementFramework/Public/Elements/Interfaces/TypedElementQueryStorageInterfaces.h
@@ -103,6 +103,9 @@ namespace UE::Editor::DataStorage
 		virtual const UScriptStruct* FindDynamicColumnType(const FDynamicColumnDescription& Description) const = 0;
 		template<TDynamicColumnTemplate TemplateType>
 		const UScriptStruct* FindDynamicColumnType(const FName& Identifier) const;
+
+		/** Returns the amount of time now and the last update in seconds. */
+		virtual float GetDeltaTimeSeconds() const = 0;
 	};
 
 	struct ICommonQueryWithEnvironmentContext : public ICommonQueryContext
