TEDS - follow up changes to add util functions for dynamic columns matching

diff --git a/Engine/Plugins/Experimental/EditorDataStorage/Source/TedsCore/Private/TypedElementDatabaseUI.cpp b/Engine/Plugins/Experimental/EditorDataStorage/Source/TedsCore/Private/TypedElementDatabaseUI.cpp
--- a/Engine/Plugins/Experimental/EditorDataStorage/Source/TedsCore/Private/TypedElementDatabaseUI.cpp
+++ b/Engine/Plugins/Experimental/EditorDataStorage/Source/TedsCore/Private/TypedElementDatabaseUI.cpp
@@ -29,17 +29,17 @@ namespace Internal
 	template<class... Ts> TOverloaded(Ts...) -> TOverloaded<Ts...>;
 
 	// Check if the two columns are equal, or if InRequestedColumn is a dynamic specialization of InMatchedColumn
-	bool CheckSingleColumnMatch(const TWeakObjectPtr<const UScriptStruct>& InMatchedColumn, const TWeakObjectPtr<const UScriptStruct>& InRequestedColumn)
+	bool CheckSingleColumnMatch(const UScriptStruct* InMatchedColumn, const UScriptStruct* InRequestedColumn)
 	{
 		if (InMatchedColumn == InRequestedColumn)
 		{
 			return true;
 		}
 		
-		else if (UE::Editor::DataStorage::ColumnUtils::IsDynamicTemplate(InMatchedColumn.Get()) &&
-			UE::Editor::DataStorage::ColumnUtils::IsDerivedFromDynamicTemplate(InRequestedColumn.Get()))
+		else if (UE::Editor::DataStorage::ColumnUtils::IsDynamicTemplate(InMatchedColumn) &&
+			UE::Editor::DataStorage::ColumnUtils::IsDerivedFromDynamicTemplate(InRequestedColumn))
 		{
-			return InRequestedColumn->IsChildOf(InMatchedColumn.Get());
+			return InRequestedColumn->IsChildOf(InMatchedColumn);
 		}
 
 		return false;
@@ -548,7 +548,7 @@ void UEditorDataStorageUi::CreateWidgetConstructors_LongestMatch(const TArray<FW
 			for (const TWeakObjectPtr<const UScriptStruct>& MatchedColumn : MatchedColumns)
 			{
 				// Remove all the columns that were matched from the provided column list.
-				while (!Internal::CheckSingleColumnMatch(MatchedColumn, *ColumnsIt))
+				while (!Internal::CheckSingleColumnMatch(MatchedColumn.Get(), ColumnsIt->Get()))
 				{
 					++ColumnIndex;
 					++ColumnsIt;
@@ -630,7 +630,7 @@ void UEditorDataStorageUi::CreateWidgetConstructors_SingleMatch(const TArray<FWi
 				continue;
 			}
 
-			if (Internal::CheckSingleColumnMatch(ColumnData[0], Columns[ColumnIndex]))
+			if (Internal::CheckSingleColumnMatch(ColumnData[0].Get(), Columns[ColumnIndex].Get()))
 			{
 				// We need to keep a copy of the actually requested column because the matched column could be the base template for the actually matched
 				// dynamic column and the widget constructor wants the latter
diff --git a/Engine/Source/Runtime/TypedElementFramework/Public/Elements/Interfaces/TypedElementDataStorageInterface.h b/Engine/Source/Runtime/TypedElementFramework/Public/Elements/Interfaces/TypedElementDataStorageInterface.h
--- a/Engine/Source/Runtime/TypedElementFramework/Public/Elements/Interfaces/TypedElementDataStorageInterface.h
+++ b/Engine/Source/Runtime/TypedElementFramework/Public/Elements/Interfaces/TypedElementDataStorageInterface.h
@@ -343,6 +343,12 @@ public:
 	 */
 	virtual void ForEachDynamicColumn(const UScriptStruct& Template, TFunctionRef<void(const UScriptStruct& Type)> Callback) const = 0;
 
+	/**
+	 * Executes the given callback for each known dynamic column that derives from the base template provided
+	 */
+	template<UE::Editor::DataStorage::TDynamicColumnTemplate DynamicColumnTemplateType>
+	void ForEachDynamicColumn(TFunctionRef<void(const UScriptStruct& Type)> Callback) const;
+
 	/**
 	 * Outputs the registered query callbacks to the given output device for debugging purposes.
 	 */
@@ -750,4 +756,10 @@ void IEditorDataStorageProvider::Reindex(UE::Editor::DataStorage::IndexHash Orig
 void IEditorDataStorageProvider::RemoveIndex(UE::Editor::DataStorage::IndexHash Index)
 {
 	RemoveRowMapping(UE::Editor::DataStorage::FMapKeyView(Index));
+}
+
+template<UE::Editor::DataStorage::TDynamicColumnTemplate DynamicColumnTemplateType>
+void IEditorDataStorageProvider::ForEachDynamicColumn(TFunctionRef<void(const UScriptStruct& Type)> Callback) const
+{
+	ForEachDynamicColumn(DynamicColumnTemplateType::StaticStruct(), Callback);
 }
\ No newline at end of file
