Replace temporary disabled UHT plugin as UBT build bug this uncovered has been fixed

diff --git a/Engine/Plugins/Experimental/EditorDataStorage/Source/TedsCore/Private/Debug/TypedElementDatabaseDebugTypes.h b/Engine/Plugins/Experimental/EditorDataStorage/Source/TedsCore/Private/Debug/TypedElementDatabaseDebugTypes.h
--- a/Engine/Plugins/Experimental/EditorDataStorage/Source/TedsCore/Private/Debug/TypedElementDatabaseDebugTypes.h
+++ b/Engine/Plugins/Experimental/EditorDataStorage/Source/TedsCore/Private/Debug/TypedElementDatabaseDebugTypes.h
@@ -27,8 +27,6 @@ struct FTestDynamicColumn : public FEditorDataStorageColumn
 {
 	GENERATED_BODY()
 
-	struct EditorDataStorage_DynamicColumnTemplate{};
-
 	// Note: Not trivial type to check that non-default move/copying works
 	// Should work regardless of UPROPERTY
 	TArray<int32> IntArray;
@@ -38,6 +36,4 @@ USTRUCT(meta=(EditorDataStorage_DynamicColumnTemplate))
 struct FTestDynamicTag : public FEditorDataStorageTag
 {
 	GENERATED_BODY()
-	
-	struct EditorDataStorage_DynamicColumnTemplate{};
 };
diff --git a/Engine/Plugins/Experimental/EditorDataStorageFeatures/Source/TedsAssetData/Public/TedsAssetDataColumns.h b/Engine/Plugins/Experimental/EditorDataStorageFeatures/Source/TedsAssetData/Public/TedsAssetDataColumns.h
--- a/Engine/Plugins/Experimental/EditorDataStorageFeatures/Source/TedsAssetData/Public/TedsAssetDataColumns.h
+++ b/Engine/Plugins/Experimental/EditorDataStorageFeatures/Source/TedsAssetData/Public/TedsAssetDataColumns.h
@@ -181,8 +181,6 @@ struct FItemStringAttributeColumn_Experimental : public FItemAttributeBaseColumn
 {
 	GENERATED_BODY()
 
-	struct EditorDataStorage_DynamicColumnTemplate{};
-
 	UPROPERTY()
 	FString Value;
 };
\ No newline at end of file
diff --git a/Engine/Plugins/Fab/Source/Fab/Private/Teds/FabMyFolderIntegration.h b/Engine/Plugins/Fab/Source/Fab/Private/Teds/FabMyFolderIntegration.h
--- a/Engine/Plugins/Fab/Source/Fab/Private/Teds/FabMyFolderIntegration.h
+++ b/Engine/Plugins/Fab/Source/Fab/Private/Teds/FabMyFolderIntegration.h
@@ -30,8 +30,6 @@ USTRUCT(meta = (DisplayName = "Fab distribution method", EditorDataStorage_Dynam
 struct FFabDistributionMethodTag : public FEditorDataStorageTag
 {
 	GENERATED_BODY()
-
-	struct EditorDataStorage_DynamicColumnTemplate{};
 };
 
 USTRUCT(meta = (DisplayName = "Fab object"))
diff --git a/Engine/Source/Runtime/TypedElementFramework/Public/Elements/Columns/TypedElementWebColumns.h b/Engine/Source/Runtime/TypedElementFramework/Public/Elements/Columns/TypedElementWebColumns.h
--- a/Engine/Source/Runtime/TypedElementFramework/Public/Elements/Columns/TypedElementWebColumns.h
+++ b/Engine/Source/Runtime/TypedElementFramework/Public/Elements/Columns/TypedElementWebColumns.h
@@ -28,8 +28,6 @@ struct FEditorDataStorageWebImageColumn final : public FEditorDataStorageColumn
 {
 	GENERATED_BODY()
 
-	struct EditorDataStorage_DynamicColumnTemplate{};
-
 	/** URL to the image in string form. Can be used with for instance with WebImage(Cache) to create a local image. */
 	UPROPERTY()
 	FString UrlString;
