TEDS: Disable the query that syncs actor transforms from TEDS back to the world as it causes sync issues

diff --git a/Engine/Plugins/Experimental/EditorDataStorage/Source/TedsCore/Private/Compatibility/TypedElementActorTransformQueries.cpp b/Engine/Plugins/Experimental/EditorDataStorage/Source/TedsCore/Private/Compatibility/TypedElementActorTransformQueries.cpp
--- a/Engine/Plugins/Experimental/EditorDataStorage/Source/TedsCore/Private/Compatibility/TypedElementActorTransformQueries.cpp
+++ b/Engine/Plugins/Experimental/EditorDataStorage/Source/TedsCore/Private/Compatibility/TypedElementActorTransformQueries.cpp
@@ -8,6 +8,17 @@
 #include "Elements/Framework/TypedElementQueryBuilder.h"
 #include "GameFramework/Actor.h"
 
+namespace UE::Editor::DataStorage::ActorTransformQueries::Private
+{
+	static bool bSyncActorTransformsToWorld = false;
+	
+	// Cvar to allow the TEDS-Outliner to automatically take over the level editor's 4 Outliners instead of appearing as a separate tab
+	static FAutoConsoleVariableRef CvarSyncActorTransforms(
+		TEXT("TEDS.Feature.SyncActorTransformsToWorld"),
+		bSyncActorTransformsToWorld,
+		TEXT("Sync actor transform changes from TEDS back to the world (only works when set on startup)"));
+}
+
 void UActorTransformDataStorageFactory::RegisterQueries(IEditorDataStorageProvider& DataStorage)
 {
 	RegisterActorAddTransformColumn(DataStorage);
@@ -68,8 +79,10 @@ void UActorTransformDataStorageFactory::RegisterActorLocalTransformToColumn(IEdi
 void UActorTransformDataStorageFactory::RegisterLocalTransformColumnToActor(IEditorDataStorageProvider& DataStorage) const
 {
 	using namespace UE::Editor::DataStorage::Queries;
-	
-	DataStorage.RegisterQuery(
+
+	if (ActorTransformQueries::Private::bSyncActorTransformsToWorld)
+	{
+		DataStorage.RegisterQuery(
 		Select(
 			TEXT("Sync transform column to actor"),
 			FProcessor(EQueryTickPhase::FrameEnd, DataStorage.GetQueryTickGroupName(EQueryTickGroups::SyncDataStorageToExternal))
@@ -84,4 +97,6 @@ void UActorTransformDataStorageFactory::RegisterLocalTransformColumnToActor(IEdi
 		.Where()
 			.All<FTypedElementActorTag, FTypedElementSyncBackToWorldTag>()
 		.Compile());
+	}
+	
 }
