ExportedTextWidget: Update to use attribute binding instead of the old TEDS UI API
TEDS UI Attribute Bindings: Add methods to bind a slate attribute to a whole TEDS column directly (instead of a specific member variable)

diff --git a/Engine/Plugins/Experimental/EditorDataStorage/Source/TedsCore/Tests/TypedElementAttributeBindingTests.cpp b/Engine/Plugins/Experimental/EditorDataStorage/Source/TedsCore/Tests/TypedElementAttributeBindingTests.cpp
--- a/Engine/Plugins/Experimental/EditorDataStorage/Source/TedsCore/Tests/TypedElementAttributeBindingTests.cpp
+++ b/Engine/Plugins/Experimental/EditorDataStorage/Source/TedsCore/Tests/TypedElementAttributeBindingTests.cpp
@@ -212,6 +212,68 @@ namespace UE::Editor::DataStorage::Tests
 
 				});
 			});
+
+			Describe("Bind Column", [this]()
+			{
+				It("Attribute can be bound to a TEDS column", [this]()
+				{
+					constexpr int InitialValue = 10;
+					constexpr int UpdatedValue = 20;
+					
+					FAttributeBinder Binder(TestRowHandle);
+					
+					// Add the test int column to the test row
+					TedsInterface->AddColumn(TestRowHandle, FTestColumnInt{.TestInt =  InitialValue});
+					FTestColumnInt* TestColumnInt = TedsInterface->GetColumn<FTestColumnInt>(TestRowHandle);
+
+					TestNotNull("Expecting valid column", TestColumnInt);
+
+					// Create an int attribute and bind it to the whole column
+					const TAttribute TestIntAttribute(Binder.BindColumn<FTestColumnInt>([](const FTestColumnInt& Column)
+					{
+						return Column.TestInt;
+					}));
+
+					// Update the value in the column
+					TestColumnInt->TestInt = UpdatedValue;
+					
+					TestEqual("Expecting int attribute value to match column value", TestIntAttribute.Get(), TestColumnInt->TestInt);
+				});
+			});
+			
+			Describe("Bind Column Data", [this]()
+			{
+				It("Attribute can be bound to a TEDS column data", [this]()
+				{
+					constexpr int InitialValue = 10;
+					constexpr int UpdatedValue = 20;
+					
+					FAttributeBinder Binder(TestRowHandle);
+					
+					// Add the test int column to the test row
+					TedsInterface->AddColumn(TestRowHandle, FTestColumnInt{.TestInt =  InitialValue});
+					FTestColumnInt* TestColumnInt = TedsInterface->GetColumn<FTestColumnInt>(TestRowHandle);
+
+					TestNotNull("Expecting valid column", TestColumnInt);
+
+					// Create an int attribute and bind it to a whole column 
+					const TAttribute TestIntAttribute(Binder.BindColumnData(FTestColumnInt::StaticStruct(),
+						[](const TWeakObjectPtr<const UScriptStruct>& ColumnType, const void* Data)
+					{
+						if (const FTestColumnInt* TestColumn = static_cast<const FTestColumnInt*>(Data))
+						{
+							return TestColumn->TestInt;
+						}
+
+						return 0;
+					}));
+
+					// Update the value in the column
+					TestColumnInt->TestInt = UpdatedValue;
+					
+					TestEqual("Expecting int attribute value to match column value", TestIntAttribute.Get(), TestColumnInt->TestInt);
+				});
+			});
 		});
 
 		AfterEach([this]()
diff --git a/Engine/Plugins/Experimental/EditorDataStorage/Source/TedsUI/Private/Widgets/ExportedTextWidget.cpp b/Engine/Plugins/Experimental/EditorDataStorage/Source/TedsUI/Private/Widgets/ExportedTextWidget.cpp
--- a/Engine/Plugins/Experimental/EditorDataStorage/Source/TedsUI/Private/Widgets/ExportedTextWidget.cpp
+++ b/Engine/Plugins/Experimental/EditorDataStorage/Source/TedsUI/Private/Widgets/ExportedTextWidget.cpp
@@ -5,6 +5,7 @@
 #include "Elements/Columns/TypedElementMiscColumns.h"
 #include "Elements/Columns/TypedElementSlateWidgetColumns.h"
 #include "Elements/Columns/TypedElementTypeInfoColumns.h"
+#include "Elements/Framework/TypedElementAttributeBinding.h"
 #include "Elements/Framework/TypedElementQueryBuilder.h"
 #include "Widgets/Text/STextBlock.h"
 
@@ -14,107 +15,6 @@
 // UExportedTextWidgetFactory
 //
 
-namespace UE::Editor
-{
-	namespace DataStorage::UI::Private
-	{
-		static void UpdateExportedTextWidget(FText Text, FTypedElementSlateWidgetReferenceColumn& Widget)
-		{
-			TSharedPtr<SWidget> WidgetPointer = Widget.Widget.Pin();
-			checkf(WidgetPointer, TEXT("Referenced widget is not valid. A constructed widget may not have been cleaned up. This can "
-				"also happen if this processor is running in the same phase as the processors responsible for cleaning up old "
-				"references."));
-			checkf(WidgetPointer->GetType() == STextBlock::StaticWidgetClass().GetWidgetType(),
-				TEXT("Stored widget with FExportedTextWidgetTag doesn't match type %s, but was a %s."),
-				*(STextBlock::StaticWidgetClass().GetWidgetType().ToString()),
-				*(WidgetPointer->GetTypeAsString()));
-		
-			STextBlock* TextWidget = static_cast<STextBlock*>(WidgetPointer.Get());
-			TextWidget->SetToolTipText(Text);
-			TextWidget->SetText(MoveTemp(Text));
-		}
-		
-		static void UpdateExportedTextWidget(const void* Data, FTypedElementSlateWidgetReferenceColumn& Widget, 
-			const UScriptStruct* StructType)
-		{
-			FString Label;
-			StructType->ExportText(Label, Data, Data, nullptr, PPF_None, nullptr);
-			UpdateExportedTextWidget(FText::FromString(MoveTemp(Label)), Widget);
-		}
-		
-		static void UpdateExportedTextWidget(IEditorDataStorageProvider& DataStorage, FTypedElementSlateWidgetReferenceColumn& Widget,
-			const FTypedElementScriptStructTypeInfoColumn& TypeInfo, const FTypedElementRowReferenceColumn& ReferencedRow)
-		{
-			const UScriptStruct* StructType = TypeInfo.TypeInfo.Get();
-			if (void* Data = DataStorage.GetColumnData(ReferencedRow.Row, StructType))
-			{
-				UpdateExportedTextWidget(Data, Widget, StructType);
-			}
-		}
-		
-		static DataStorage::RowHandle RegisterUpdateCallback(IEditorDataStorageProvider& DataStorage, const UScriptStruct* Target)
-		{
-			using namespace UE::Editor::DataStorage::Queries;
-
-			
-			QueryHandle TypeDataQuery = DataStorage.RegisterQuery(
-				Select()
-					.ReadOnly(Target)
-				.Where()
-					.Any<FTypedElementSyncFromWorldTag, FTypedElementSyncBackToWorldTag>()
-				.Compile());
-		
-			FString Name = TEXT("Sync exported text widgets (");
-			Target->AppendName(Name);
-			Name += ')';
-		
-			FName ProcessorName(Name);
-		
-			return DataStorage.RegisterQuery(
-				Select(ProcessorName,
-					FProcessor(EQueryTickPhase::FrameEnd, DataStorage.GetQueryTickGroupName(EQueryTickGroups::SyncWidgets))
-						.SetExecutionMode(EExecutionMode::GameThread),
-					[Target, ProcessorName](
-						IQueryContext& Context, 
-						FTypedElementSlateWidgetReferenceColumn& Widget,
-						const FTypedElementScriptStructTypeInfoColumn& TypeInfo,
-						const FTypedElementRowReferenceColumn& ReferencedRow)
-					{
-						const UScriptStruct* StructType = TypeInfo.TypeInfo.Get();
-						if (!ensureMsgf(StructType != nullptr, TEXT("WeakPtr TypeInfo is null in query '%s'"), *ProcessorName.ToString()))
-						{
-							return;
-						}
-		
-						/* This query will grab all ExportedText widgets that were created for rows with the "Target" column, but we want to make sure we are
-						 * only updating the widgets that are actually displaying the "Target" column.
-						 *
-						 * E.g a row could have ColumnA and ColumnB which are both using the exported text widget to display - but there is no way for the
-						 * two widgets to be differentiated from a TEDS query. So if the widget for ColumnB wants to update: Both the widgets for ColumnA
-						 * and ColumnB would match the query condition (WidgetRow has FExportedTextWidgetTag && TargetRow has ColumnB), but this
-						 * query (+ subquery) only have access to ColumnB and only want to update the widget that's displaying ColumnB.
-						 *
-						 * To work around this we check to make sure the widget this query is trying to update is for the column this query is targeting.
-						 */
-						if(StructType == Target)
-						{
-							Context.RunSubquery(0, ReferencedRow.Row, 
-							[&Widget, StructType](const FQueryDescription&, ISubqueryContext& SubqueryContext)
-							{
-								const void* ColumnData = SubqueryContext.GetColumn(StructType);
-								UpdateExportedTextWidget(ColumnData, Widget, StructType);
-							});
-						}
-					})
-				.Where()
-					.All<FExportedTextWidgetTag>()
-				.DependsOn()
-					.SubQuery(TypeDataQuery)
-				.Compile());
-		}
-	} // namespace DataStorage::UI::Private
-} // namespace UE::Editor
-
 void UExportedTextWidgetFactory::RegisterWidgetConstructors(IEditorDataStorageProvider& DataStorage,
 	IEditorDataStorageUiProvider& DataStorageUi) const
 {
@@ -133,10 +33,10 @@ FExportedTextWidgetConstructor::FExportedTextWidgetConstructor()
 TConstArrayView<const UScriptStruct*> FExportedTextWidgetConstructor::GetAdditionalColumnsList() const
 {
 	static TTypedElementColumnTypeList<
-		FTypedElementRowReferenceColumn,
-		FTypedElementScriptStructTypeInfoColumn,
-		FExportedTextWidgetTag> Columns;
-	return Columns;
+    		FTypedElementRowReferenceColumn,
+    		FTypedElementScriptStructTypeInfoColumn,
+    		FExportedTextWidgetTag> Columns;
+    	return Columns;
 }
 
 const UE::Editor::DataStorage::Queries::FConditions* FExportedTextWidgetConstructor::GetQueryConditions() const
@@ -163,59 +63,47 @@ FText FExportedTextWidgetConstructor::CreateWidgetDisplayNameText(IEditorDataSto
 	}
 }
 
-TSharedPtr<SWidget> FExportedTextWidgetConstructor::ConstructFinalWidget(UE::Editor::DataStorage::RowHandle Row, 
-	IEditorDataStorageProvider* DataStorage, IEditorDataStorageUiProvider* DataStorageUi, const UE::Editor::DataStorage::FMetaDataView& Arguments)
+TSharedPtr<SWidget> FExportedTextWidgetConstructor::CreateWidget(IEditorDataStorageProvider* DataStorage, IEditorDataStorageUiProvider* DataStorageUi,
+	UE::Editor::DataStorage::RowHandle TargetRow, UE::Editor::DataStorage::RowHandle WidgetRow,
+	const UE::Editor::DataStorage::FMetaDataView& Arguments)
 {
-	FTypedElementScriptStructTypeInfoColumn& TypeInfoColumn = *DataStorage->GetColumn<FTypedElementScriptStructTypeInfoColumn>(Row);
-
 	using namespace UE::Editor::DataStorage::Queries;
 	
-	// NOTE: We are currently assuming that an instance of FExportedTextWidgetConstructor will only be used to show the same type info for all rows
-	// matched with it. This isn't ideal but it's better than nothing since we need some sort of matched conditions for column based virtualization to work.
-	// TEDS UI TODO: We should work around it by refactoring this into an STedsWidget in the future so it can store the column conditions per instance
-	MatchedColumn = FConditions(TColumn(TypeInfoColumn.TypeInfo)).Compile(FEditorStorageQueryConditionCompileContext(DataStorage));
-
-	return FTypedElementWidgetConstructor::ConstructFinalWidget(Row, DataStorage, DataStorageUi, Arguments);
-}
-
-TSharedPtr<SWidget> FExportedTextWidgetConstructor::CreateWidget(const UE::Editor::DataStorage::FMetaDataView& Arguments)
-{
-	return SNew(STextBlock);
-}
-
-bool FExportedTextWidgetConstructor::FinalizeWidget(
-	IEditorDataStorageProvider* DataStorage,
-	IEditorDataStorageUiProvider* DataStorageUi,
-	UE::Editor::DataStorage::RowHandle Row,
-	const TSharedPtr<SWidget>& Widget)
-{
-	using namespace UE::Editor;
-	FTypedElementScriptStructTypeInfoColumn& TypeInfoColumn = *DataStorage->GetColumn<FTypedElementScriptStructTypeInfoColumn>(Row);
-
-	if (TypeInfoColumn.TypeInfo->IsChildOf(DataStorage::FTag::StaticStruct()))
-	{
-		DataStorage::UI::Private::UpdateExportedTextWidget(
-			LOCTEXT("ExportedTextWidgetTag", "<Tag>"),
-			*DataStorage->GetColumn<FTypedElementSlateWidgetReferenceColumn>(Row));
-	}
-	else
+	FAttributeBinder Binder(TargetRow, DataStorage);
+	
+	if (FTypedElementScriptStructTypeInfoColumn* TypeInfoColumn = DataStorage->GetColumn<FTypedElementScriptStructTypeInfoColumn>(WidgetRow))
 	{
-		DataStorage::UI::Private::UpdateExportedTextWidget(
-			*DataStorage,
-			*DataStorage->GetColumn<FTypedElementSlateWidgetReferenceColumn>(Row),
-			TypeInfoColumn,
-			*DataStorage->GetColumn<FTypedElementRowReferenceColumn>(Row));
+		TWeakObjectPtr<const UScriptStruct> ColumnTypeInfo = TypeInfoColumn->TypeInfo;
 
-		UExportedTextWidgetFactory* Factory =
-			UExportedTextWidgetFactory::StaticClass()->GetDefaultObject<UExportedTextWidgetFactory>();
-		if (Factory && !Factory->RegisteredTypes.Contains(TypeInfoColumn.TypeInfo))
+		if(ColumnTypeInfo.IsValid())
 		{
-			DataStorage::UI::Private::RegisterUpdateCallback(*DataStorage, TypeInfoColumn.TypeInfo.Get());
-			Factory->RegisteredTypes.Add(TypeInfoColumn.TypeInfo);
+			// NOTE: We are currently assuming that an instance of FExportedTextWidgetConstructor will only be used to show the same type info for all rows
+			// matched with it. This isn't ideal but it's better than nothing since we need some sort of matched conditions for column based virtualization to work.
+			// TEDS UI TODO: We should work around it by refactoring this into an STedsWidget in the future so it can store the column conditions per instance
+			MatchedColumn = FConditions(TColumn(ColumnTypeInfo)).Compile(FEditorStorageQueryConditionCompileContext(DataStorage));
+	
+			return SNew(STextBlock)
+					.Text(Binder.BindColumnData(ColumnTypeInfo, [](const TWeakObjectPtr<const UScriptStruct>& InTypeInfo, const void* InData)
+					{
+						if (InTypeInfo.IsValid())
+						{
+							if (InData)
+							{
+								FString Label;
+								InTypeInfo->ExportText(Label, InData, InData, nullptr, PPF_None, nullptr);
+								return FText::FromString(Label);
+							}
+							else
+							{
+								return FText::Format(LOCTEXT("ColumnNotFoundText", "Column {0} not found on row"), InTypeInfo->GetDisplayNameText());
+							}
+						}
+						return LOCTEXT("MissingTypeInfoText", "Missing type info for column");
+					}));
 		}
 	}
 	
-	return true;
+	return SNullWidget::NullWidget;
 }
 
 #undef LOCTEXT_NAMESPACE
diff --git a/Engine/Plugins/Experimental/EditorDataStorage/Source/TedsUI/Public/Widgets/ExportedTextWidget.h b/Engine/Plugins/Experimental/EditorDataStorage/Source/TedsUI/Public/Widgets/ExportedTextWidget.h
--- a/Engine/Plugins/Experimental/EditorDataStorage/Source/TedsUI/Public/Widgets/ExportedTextWidget.h
+++ b/Engine/Plugins/Experimental/EditorDataStorage/Source/TedsUI/Public/Widgets/ExportedTextWidget.h
@@ -21,12 +21,10 @@ public:
 
 	virtual void RegisterWidgetConstructors(IEditorDataStorageProvider& DataStorage,
 		IEditorDataStorageUiProvider& DataStorageUi) const override;
-
-	TSet<TWeakObjectPtr<const UScriptStruct>> RegisteredTypes;
 };
 
 USTRUCT()
-struct TEDSUI_API FExportedTextWidgetConstructor : public FTypedElementWidgetConstructor
+struct TEDSUI_API FExportedTextWidgetConstructor : public FSimpleWidgetConstructor
 {
 	GENERATED_BODY()
 
@@ -39,20 +37,13 @@ public:
 
 	virtual FText CreateWidgetDisplayNameText(
 		IEditorDataStorageProvider* DataStorage, UE::Editor::DataStorage::RowHandle Row) const override;
-
-	virtual TSharedPtr<SWidget> ConstructFinalWidget(
-		UE::Editor::DataStorage::RowHandle Row, /** The row the widget will be stored in. */
-		IEditorDataStorageProvider* DataStorage,
-		IEditorDataStorageUiProvider* DataStorageUi,
-		const UE::Editor::DataStorage::FMetaDataView& Arguments) override;
 	
-protected:
-	virtual TSharedPtr<SWidget> CreateWidget(const UE::Editor::DataStorage::FMetaDataView& Arguments) override;
-	virtual bool FinalizeWidget(
+	virtual TSharedPtr<SWidget> CreateWidget(
 		IEditorDataStorageProvider* DataStorage,
 		IEditorDataStorageUiProvider* DataStorageUi,
-		UE::Editor::DataStorage::RowHandle Row,
-		const TSharedPtr<SWidget>& Widget) override;
+		UE::Editor::DataStorage::RowHandle TargetRow,
+		UE::Editor::DataStorage::RowHandle WidgetRow,
+		const UE::Editor::DataStorage::FMetaDataView& Arguments) override;
 
 protected:
 	// The column this exported text widget is operating on
diff --git a/Engine/Source/Runtime/TypedElementFramework/Public/Elements/Framework/TypedElementAttributeBinding.h b/Engine/Source/Runtime/TypedElementFramework/Public/Elements/Framework/TypedElementAttributeBinding.h
--- a/Engine/Source/Runtime/TypedElementFramework/Public/Elements/Framework/TypedElementAttributeBinding.h
+++ b/Engine/Source/Runtime/TypedElementFramework/Public/Elements/Framework/TypedElementAttributeBinding.h
@@ -66,7 +66,6 @@ namespace UE::Editor::DataStorage
 		 * @param InVariable The data member inside a column to be bound
 		 * @param InConverter Conversion function to convert from DataType -> AttributeType
 		 * @param InDefaultValue The default value to be used when the column isn't present on a row
-		 * @param InIdentifier The identifier for this column if it is a dynamic column, NAME_None if it is not a dynamic column
 		 * @return A TAttribute bound to the row, column pair specified
 		 *
 		 * Example:
@@ -112,7 +111,7 @@ namespace UE::Editor::DataStorage
 		 */
 		template <typename DataType, TDataColumnType ColumnType, typename FunctionType>
 			requires Private::AttributeBinderInvocable<FunctionType, DataType>
-		auto BindData(DataType ColumnType::* InVariable, FunctionType InConverter, const DataType& InDefaultValue = DataType());
+		auto BindData(DataType ColumnType::* InVariable, FunctionType&& InConverter, const DataType& InDefaultValue = DataType());
 
 		/**
 		 * Overload for the conversion binder to accept lambdas instead of TFunctions
@@ -125,7 +124,59 @@ namespace UE::Editor::DataStorage
 		 */
 		template <typename DataType, TDynamicColumnTemplate ColumnType, typename FunctionType>
 			requires Private::AttributeBinderInvocable<FunctionType, DataType>
-		auto BindData(const FName& InIdentifier, DataType ColumnType::* InVariable, FunctionType InConverter, const DataType& InDefaultValue = DataType());
+		auto BindData(const FName& InIdentifier, DataType ColumnType::* InVariable, FunctionType&& InConverter, const DataType& InDefaultValue = DataType());
+
+		/**
+		 * Bind a whole TEDS column to a slate attribute (instead of a single member variable). This allows you to derive an attribute from multiple
+		 * members of a TEDS column
+		 *
+		 * @param InConverter Conversion function to convert from ColumnType -> AttributeType
+		 * @return A TAttribute bound to the row, column pair specified
+		 *
+		 * Example:
+		 *			const TAttribute TestIntAttribute(Binder.BindColumn<FTestColumnInt>([](const FTestColumnInt& Column)
+		 *			{
+		 *				return Column.TestInt1 + Column.TestInt2;
+		 *			}
+		 */
+		template <typename AttributeType, TDataColumnType ColumnType>
+		TAttribute<AttributeType> BindColumn(TFunction<AttributeType(const ColumnType&)> InConverter);
+
+		/**
+		 * Overload for the BindColumn to accept lambdas instead of TFunctions
+		 * @param InConverter Conversion function to convert from DataType -> AttributeType
+		 * @return A TAttribute bound to the row, column pair specified
+		 */
+		template <TDataColumnType ColumnType, typename FunctionType>
+			requires Private::AttributeBinderColumnInvocable<FunctionType, ColumnType>
+		auto BindColumn(FunctionType&& InConverter);
+		
+		/**
+		 * Bind a whole TEDS column to a slate attribute (instead of a single member variable) using the column's typeinfo.
+		 *
+		 * @param InColumnType The typeinfo of the column (e.g FTestIntColumn::StaticStruct)
+		 * @param InConverter Conversion function to convert from const void* ColumnData -> AttributeType
+		 * @return A TAttribute bound to the row, column pair specified
+		 *
+		 * Example:
+		 *			const TAttribute TestIntAttribute(Binder.BindColumnData(FTestColumnInt::StaticStruct(),
+		 *				[](const TWeakObjectPtr<const UScriptStruct>& ColumnType, const void* ColumnData)
+		 *			{
+		 *				return static_cast<const FTestColumnInt*>(ColumnData)->TestInt;
+		 *			}
+		 */
+		template <typename AttributeType>
+		TAttribute<AttributeType> BindColumnData(const TWeakObjectPtr<const UScriptStruct>& InColumnType,
+			const TFunction<AttributeType(const TWeakObjectPtr<const UScriptStruct>&, const void*)>& InConverter);
+
+		/**
+		 * Overload for the BindColumnData to accept lambdas instead of TFunctions
+		 * @param InConverter Conversion function to convert from const void* ColumnData -> AttributeType
+		 * @return A TAttribute bound to the row, column pair specified
+		 */
+		template <typename FunctionType>
+			requires Private::AttributeBinderColumnDataInvocable<FunctionType>
+		auto BindColumnData(const TWeakObjectPtr<const UScriptStruct>& InColumnType, FunctionType&& InConverter);
 
 		/**
 		 * Bind a delegate inside a TEDS column to a SLATE_EVENT macro on a widget
diff --git a/Engine/Source/Runtime/TypedElementFramework/Public/Elements/Framework/TypedElementAttributeBinding.inl b/Engine/Source/Runtime/TypedElementFramework/Public/Elements/Framework/TypedElementAttributeBinding.inl
--- a/Engine/Source/Runtime/TypedElementFramework/Public/Elements/Framework/TypedElementAttributeBinding.inl
+++ b/Engine/Source/Runtime/TypedElementFramework/Public/Elements/Framework/TypedElementAttributeBinding.inl
@@ -100,7 +100,8 @@ namespace UE::Editor::DataStorage
 	}
 
 	template <typename AttributeType, typename DataType, TDynamicColumnTemplate ColumnType>
-	TAttribute<AttributeType> FAttributeBinder::BindData(const FName& InIdentifier, DataType ColumnType::* InVariable, TFunction<AttributeType(const DataType&)> InConverter, const DataType& InDefaultValue)
+	TAttribute<AttributeType> FAttributeBinder::BindData(const FName& InIdentifier, DataType ColumnType::* InVariable,
+		TFunction<AttributeType(const DataType&)> InConverter, const DataType& InDefaultValue)
 	{
 		using namespace Private;
 
@@ -136,7 +137,7 @@ namespace UE::Editor::DataStorage
 	
 	template <typename DataType, TDataColumnType ColumnType, typename FunctionType>
 		requires Private::AttributeBinderInvocable<FunctionType, DataType>
-	auto FAttributeBinder::BindData(DataType ColumnType::* InVariable, FunctionType InConverter, const DataType& InDefaultValue)
+	auto FAttributeBinder::BindData(DataType ColumnType::* InVariable, FunctionType&& InConverter, const DataType& InDefaultValue)
 	{
 		// Deduce the attribute type from the return value of the converter function
 		using AttributeType = decltype(InConverter(std::declval<DataType>()));
@@ -146,13 +147,74 @@ namespace UE::Editor::DataStorage
 	
 	template <typename DataType, TDynamicColumnTemplate ColumnType, typename FunctionType>
 		requires Private::AttributeBinderInvocable<FunctionType, DataType>
-	auto FAttributeBinder::BindData(const FName& InIdentifier, DataType ColumnType::* InVariable, FunctionType InConverter, const DataType& InDefaultValue)
+	auto FAttributeBinder::BindData(const FName& InIdentifier, DataType ColumnType::* InVariable, FunctionType&& InConverter, const DataType& InDefaultValue)
 	{
 		// Deduce the attribute type from the return value of the converter function
 		using AttributeType = decltype(InConverter(std::declval<DataType>()));
 		
 		return BindData<AttributeType, DataType>(InIdentifier, InVariable, TFunction<AttributeType(const DataType&)>(MoveTemp(InConverter)), InDefaultValue);
 	}
+	
+	template <typename AttributeType, TDataColumnType ColumnType>
+	TAttribute<AttributeType> FAttributeBinder::BindColumn(TFunction<AttributeType(const ColumnType&)> InConverter)
+	{
+		return TAttribute<AttributeType>::CreateLambda(
+			[
+				Convertor = InConverter, 
+				Storage = DataStorage, 
+				Row = TargetRow 
+			]()
+			{
+				if(ColumnType* Column = Storage->GetColumn<ColumnType>(Row))
+				{
+					return Convertor(*Column);
+				}
+
+				return AttributeType();
+			});
+	}
+		
+	template <TDataColumnType ColumnType, typename FunctionType>
+		requires Private::AttributeBinderColumnInvocable<FunctionType, ColumnType>
+	auto FAttributeBinder::BindColumn(FunctionType&& InConverter)
+	{
+		// Deduce the attribute type from the return value of the converter function
+		using AttributeType = decltype(InConverter(std::declval<const ColumnType&>()));
+		
+		return BindColumn<AttributeType, ColumnType>(TFunction<AttributeType(const ColumnType&)>(InConverter));
+	}
+	
+	template <typename AttributeType>
+	TAttribute<AttributeType> FAttributeBinder::BindColumnData(const TWeakObjectPtr<const UScriptStruct>& InColumnType,
+		const TFunction<AttributeType(const TWeakObjectPtr<const UScriptStruct>&, const void*)>& InConverter)
+	{
+		return TAttribute<AttributeType>::CreateLambda(
+			[
+				ColumnType = InColumnType,
+				Convertor = InConverter, 
+				Storage = DataStorage, 
+				Row = TargetRow 
+			]()
+			{
+				if(const void* ColumnData = Storage->GetColumnData(Row, ColumnType.Get()))
+				{
+					return Convertor(ColumnType, ColumnData);
+				}
+
+				return AttributeType();
+			});
+	}
+
+	template <typename FunctionType>
+			requires Private::AttributeBinderColumnDataInvocable<FunctionType>
+	auto FAttributeBinder::BindColumnData(const TWeakObjectPtr<const UScriptStruct>& InColumnType, FunctionType&& InConverter)
+	{
+		// Deduce the attribute type from the return value of the converter function
+		using AttributeType = decltype(InConverter(std::declval<const TWeakObjectPtr<const UScriptStruct>&>(), std::declval<const void*>()));
+		
+		return BindColumnData<AttributeType>(InColumnType,
+			TFunction<AttributeType(const TWeakObjectPtr<const UScriptStruct>&, const void*)>(InConverter));
+	}
 
 	template <typename InRetValType, typename... ParamTypes, TDataColumnType ColumnType>
 		TDelegate<InRetValType(ParamTypes...)> FAttributeBinder::BindEvent(TDelegate<InRetValType(ParamTypes...)> ColumnType::* InVariable)
diff --git a/Engine/Source/Runtime/TypedElementFramework/Public/Elements/Framework/TypedElementAttributeBindingProperty.h b/Engine/Source/Runtime/TypedElementFramework/Public/Elements/Framework/TypedElementAttributeBindingProperty.h
--- a/Engine/Source/Runtime/TypedElementFramework/Public/Elements/Framework/TypedElementAttributeBindingProperty.h
+++ b/Engine/Source/Runtime/TypedElementFramework/Public/Elements/Framework/TypedElementAttributeBindingProperty.h
@@ -15,6 +15,14 @@ namespace UE::Editor::DataStorage
 		// Concept for a conversion function used by the attribute binder to bind a column data member to an attribute of a different type
 		template <typename FunctionType, typename ArgumentType>
 		concept AttributeBinderInvocable = std::is_invocable_v<std::decay_t<FunctionType>, const ArgumentType&>;
+		
+		// Concept for a conversion function used by the attribute binder to bind a lambda to a TFunction accepting column arguments
+		template <typename FunctionType, typename ColumnType>
+		concept AttributeBinderColumnInvocable = std::is_invocable_v<std::decay_t<FunctionType>, const ColumnType&>;
+
+		// Concept for a conversion function used by the attribute binder to bind a lambda to a TFunction accepting column data arguments
+		template <typename FunctionType>
+		concept AttributeBinderColumnDataInvocable = std::is_invocable_v<std::decay_t<FunctionType>, const TWeakObjectPtr<const UScriptStruct>&, const void*>;
 
 		template <TColumnType ColumnType>
 		ColumnType* GetColumn(IEditorDataStorageProvider* DataStorage, RowHandle Row)
