Added the possibility for a teds processor to opt-in into a batched mode for all its deferred commands (Adding/Removing of columns and rows)
On very large queries this can drasticly improve the performance. On a large project I observed a perf improvement of  21 seconds on the GT on the first mass tick after initializing the TEDS.AssetDataStorage. Before: 30.1 sec, With this change:8.1 sec. Note that some experimental editor feature where disabled for this test.

Also changed the signature of the function BatchBuildEntities that used a FMassArchetypeCompositionDescriptor&& argurment to used a const FMassArchetypeCompositionDescriptor& instead as there is no performence gain in the current implementation by owning the composition descriptor.

Please note that I still need to add some automated test for this.

diff --git a/Engine/Plugins/Experimental/EditorDataStorage/Source/TedsCore/Private/Processors/TypedElementProcessorAdaptors.cpp b/Engine/Plugins/Experimental/EditorDataStorage/Source/TedsCore/Private/Processors/TypedElementProcessorAdaptors.cpp
--- a/Engine/Plugins/Experimental/EditorDataStorage/Source/TedsCore/Private/Processors/TypedElementProcessorAdaptors.cpp
+++ b/Engine/Plugins/Experimental/EditorDataStorage/Source/TedsCore/Private/Processors/TypedElementProcessorAdaptors.cpp
@@ -5,12 +5,15 @@
 #include <utility>
 #include "Algo/Copy.h"
 #include "Elements/Common/TypedElementQueryTypes.h"
+#include "MassEntityTypes.h"
 #include "MassEntityView.h"
 #include "MassExecutionContext.h"
 #include "Queries/TypedElementExtendedQueryStore.h"
 #include "TypedElementDatabase.h"
 #include "TypedElementDatabaseEnvironment.h"
 #include "TypedElementUtils.h"
+#include "Algo/Copy.h"
+#include "StructUtils/StructTypeBitSet.h"
 
 namespace UE::Editor::DataStorage
 {
@@ -124,6 +127,34 @@ namespace UE::Editor::DataStorage
 			}
 		};
 
+		struct FRowColumnModifications
+		{
+			FMassTagBitSet AddedTags;
+			FMassFragmentBitSet AddedFragments;
+			TSet<FDynamicColumnDescription> AddedDynamicColumns;
+			FMassTagBitSet RemovedTags;
+			FMassFragmentBitSet RemovedFragments;
+
+			bool operator==(const FRowColumnModifications& Other) const
+			{
+				return AddedTags == Other.AddedTags
+					&& AddedFragments == Other.AddedFragments
+					&& AddedDynamicColumns.Num() == Other.AddedDynamicColumns.Num()
+					&& AddedDynamicColumns.Includes(Other.AddedDynamicColumns)
+					&& RemovedTags == Other.RemovedTags
+					&& RemovedFragments == Other.RemovedFragments;
+			}
+		};
+
+		uint32 GetTypeHash(const FRowColumnModifications& InModification)
+		{
+			uint32 Hash = GetTypeHash(InModification.AddedTags);
+			Hash = HashCombine(Hash, GetTypeHash(InModification.AddedFragments));
+			Hash = HashCombine(Hash, uint32(InModification.AddedDynamicColumns.Num()));
+			Hash = HashCombine(Hash, GetTypeHash(InModification.RemovedTags));
+			return HashCombine(Hash, GetTypeHash(InModification.RemovedFragments));
+		}
+
 		struct FMassWithEnvironmentContextCommon : public FMassContextCommon
 		{
 			using Parent = FMassContextCommon;
@@ -176,10 +207,289 @@ namespace UE::Editor::DataStorage
 		public:
 			using ObjectCopyOrMove = void (*)(const UScriptStruct& TypeInfo, void* Destination, void* Source);
 	
-			FMassWithEnvironmentContextCommon(FMassExecutionContext& InContext, FEnvironment& InEnvironment)
+			FMassWithEnvironmentContextCommon(FMassExecutionContext& InContext, FEnvironment& InEnvironment, bool bInBatchDeferredCommands)
 				: FMassContextCommon(InContext, InEnvironment)
+				, bBatchDeferredCommands(bInBatchDeferredCommands)
 			{}
 
+			~FMassWithEnvironmentContextCommon()
+			{
+				struct FRelocator
+				{
+					ObjectCopyOrMove CopyOrMove = nullptr;
+					const UScriptStruct* FragmentType = nullptr;
+					FName Identifier;
+					void* Object = nullptr;
+				};
+
+				struct FRowOperations
+				{
+					FRowColumnModifications ColumnsChange;
+					TArray<FRelocator> Relocators;
+				};
+
+				if (bBatchDeferredCommands)
+				{
+					TMap<FMassEntityHandle, FRowOperations> RowOperations;
+
+					// Process the per row operations
+					{
+						int32 Count = 0;
+						for (FAddedColumnsBatch* const AddColumns : AddedColumnsBatchQueue)
+						{
+							Count += AddColumns->EntityCount;
+						}
+						const int32 AddedColumnsBathRowsCount = Count;
+
+						Count = 0;
+						for (FAddDynamicColumnsBatch* const AddedDynamicColumns : AddedDynamicColumnsBatchQueue)
+						{
+							Count += AddedDynamicColumns->Rows.Num();
+						}
+						const int32 AddedDynamicColumnsRowCount = Count;
+
+						Count = 0;
+						for (FRemovedColumnsBatch* const RemovedColumns : RemovedColumnsBatchQueue)
+						{
+							Count += RemovedColumns->EntityCount;
+						}
+						const int32 RemovedColumnsBatchRowCount = Count;
+
+
+						RowOperations.Reserve(FMath::Max(AddedColumnsQueue.Num(),
+							AddedColumnsBathRowsCount,
+							AddedColumnWithRelocationQueue.Num(),
+							AddedDynamicColumnWithRelocationQueue.Num(),
+							AddedDynamicColumnsRowCount,
+							RemovedColumnsQueue.Num(),
+							RemovedColumnsBatchRowCount
+						));
+
+						for (FAddedColumns* const AddColumns : AddedColumnsQueue)
+						{
+							FRowOperations& Operations = RowOperations.FindOrAdd(AddColumns->Entity);
+							Operations.ColumnsChange.AddedTags = AddColumns->AddDescriptor.Tags;
+							Operations.ColumnsChange.AddedFragments = AddColumns->AddDescriptor.Fragments;
+						}
+						AddedColumnsQueue.Empty();
+
+						for (FAddedColumnsBatch* const AddColumns : AddedColumnsBatchQueue)
+						{
+							TArrayView<FMassEntityHandle> Rows(AddColumns->Entities, AddColumns->EntityCount);
+							for (const FMassEntityHandle& Row : Rows)
+							{
+								FRowOperations& Operations = RowOperations.FindOrAdd(Row);
+								Operations.ColumnsChange.AddedTags += AddColumns->AddDescriptor.Tags;
+								Operations.ColumnsChange.AddedFragments += AddColumns->AddDescriptor.Fragments;
+							}
+						}
+						AddedColumnsBatchQueue.Empty();
+
+						for (FAddValueColumn* const AddColumn : AddedColumnWithRelocationQueue)
+						{
+							FRowOperations& Operations = RowOperations.FindOrAdd(AddColumn->Entity);
+							Operations.ColumnsChange.AddedFragments.Add(*AddColumn->FragmentType);
+							Operations.Relocators.Add(FRelocator{
+									.CopyOrMove = AddColumn->Relocator,
+									.FragmentType = AddColumn->FragmentType,
+									.Object = AddColumn->Object,
+								});
+						}
+						AddedColumnWithRelocationQueue.Empty();
+
+						for (FAddDynamicColumn* const AddDynamicColumn : AddedDynamicColumnWithRelocationQueue)
+						{
+							FRowOperations& Operations = RowOperations.FindOrAdd(AddDynamicColumn->Entity);
+							Operations.ColumnsChange.AddedDynamicColumns.Add(AddDynamicColumn->Description);
+							Operations.Relocators.Add(FRelocator{
+									.CopyOrMove = AddDynamicColumn->Relocator,
+									.FragmentType = AddDynamicColumn->Description.TemplateType,
+									.Identifier = AddDynamicColumn->Description.Identifier,
+									.Object = AddDynamicColumn->Object,
+								});
+						}
+						AddedDynamicColumnWithRelocationQueue.Empty();
+
+						for (FAddDynamicColumnsBatch* const AddedDynamicColumns : AddedDynamicColumnsBatchQueue)
+						{
+							for (const RowHandle& Row : AddedDynamicColumns->Rows)
+							{
+								FRowOperations& Operations = RowOperations.FindOrAdd(FMassEntityHandle::FromNumber(Row));
+								Operations.ColumnsChange.AddedDynamicColumns.Append(AddedDynamicColumns->Descriptions);
+							}
+						}
+						AddedDynamicColumnsBatchQueue.Empty();
+
+						for (FRemovedColumns* const RemovedColumns : RemovedColumnsQueue)
+						{
+							FRowOperations& Operations = RowOperations.FindOrAdd(RemovedColumns->Entity);
+							Operations.ColumnsChange.RemovedTags += RemovedColumns->RemoveDescriptor.Tags;
+							Operations.ColumnsChange.RemovedFragments += RemovedColumns->RemoveDescriptor.Fragments;
+						}
+						RemovedColumnsQueue.Empty();
+
+						for (FRemovedColumnsBatch* const RemovedColumns : RemovedColumnsBatchQueue)
+						{
+							TArrayView<FMassEntityHandle> Rows(RemovedColumns->Entities, RemovedColumns->EntityCount);
+							for (const FMassEntityHandle& Row : Rows)
+							{
+								FRowOperations& Operations = RowOperations.FindOrAdd(Row);
+								Operations.ColumnsChange.RemovedTags += RemovedColumns->RemoveDescriptor.Tags;
+								Operations.ColumnsChange.RemovedFragments += RemovedColumns->RemoveDescriptor.Fragments;
+							}
+						}
+						RemovedColumnsBatchQueue.Empty();
+					}
+
+
+					// Convert the per row operations to batched changes
+					TMap<FRowColumnModifications, TArray<FMassEntityHandle>> ChangesAndRows;
+					for (const TPair<FMassEntityHandle, FRowOperations>& RowOperation : RowOperations)
+					{
+						ChangesAndRows.FindOrAdd(RowOperation.Value.ColumnsChange).Add(RowOperation.Key);
+					}
+
+					// Convert the Add row to batched changes
+					TMap<FMassArchetypeHandle, TArray<FMassEntityHandle>> RowsToAdd;
+					for (const FAddRow& AddRow : AddedRowsQueue)
+					{
+						RowsToAdd.FindOrAdd(AddRow.Archetype).Add(AddRow.Entity);
+					}
+
+					// Prepare the remove row array for Mass
+					TArray<FMassEntityHandle> EntitiesToDestroy;
+					EntitiesToDestroy.Reserve(RemovedRowsQueue.Num());
+					for (const RowHandle& Row : RemovedRowsQueue)
+					{
+						EntitiesToDestroy.Add(FMassEntityHandle::FromNumber(Row));
+					}
+					RemovedRowsQueue.Empty();
+
+					this->Context.Defer().template PushCommand<FMassDeferredChangeCompositionCommand>(
+						[InChangesAndRows = MoveTemp(ChangesAndRows),
+						InRowOperations = MoveTemp(RowOperations),
+						PtrToEnvironment = &Environment,
+						InEntitiesToDestroy = MoveTemp(EntitiesToDestroy),
+						InRowsToAdd = MoveTemp(RowsToAdd)
+						](FMassEntityManager& System) mutable
+						{
+							TSharedRef<FMassEntityManager::FEntityCreationContext> ObtainedContext = System.GetOrMakeCreationContext();
+
+							// Adds Rows
+							{
+								TArray<FMassArchetypeEntityCollectionWithPayload> EntityCollections;
+								for (TPair<FMassArchetypeHandle, TArray<FMassEntityHandle>>& ArchetypeAndRows : InRowsToAdd)
+								{
+									FMassArchetypeEntityCollectionWithPayload::CreateEntityRangesWithPayload(System, ArchetypeAndRows.Value, FMassArchetypeEntityCollection::NoDuplicates
+										, FMassGenericPayloadView(), EntityCollections);
+									const FMassArchetypeCompositionDescriptor& ArchetypeComposition = System.GetArchetypeComposition(ArchetypeAndRows.Key);
+
+									for (const FMassArchetypeEntityCollectionWithPayload& Collection : EntityCollections)
+									{
+										System.BatchBuildEntities(Collection, ArchetypeComposition);
+									}
+
+									EntityCollections.Reset();
+								}
+							}
+
+							using EntityHandleArray = TArray<FMassEntityHandle, TInlineAllocator<32>>;
+							using EntityArchetypeLookup = TMap<FMassArchetypeHandle, EntityHandleArray, TInlineSetAllocator<32>>;
+							using ArchetypeEntityArray = TArray<FMassArchetypeEntityCollection, TInlineAllocator<32>>;
+
+							EntityArchetypeLookup LookupTable;
+							ArchetypeEntityArray EntityCollections;
+
+							// Change the entity types
+							for (TPair<FRowColumnModifications, TArray<FMassEntityHandle>>& BatchedModification : InChangesAndRows)
+							{
+								for (const FDynamicColumnDescription& DynamicColumnAdded : BatchedModification.Key.AddedDynamicColumns)
+								{
+									BatchedModification.Key.AddedFragments.Add(*PtrToEnvironment->GenerateDynamicColumn(*DynamicColumnAdded.TemplateType, DynamicColumnAdded.Identifier));
+								}
+
+								const FRowColumnModifications& ColumnModification = BatchedModification.Key;
+								if (!ColumnModification.AddedFragments.IsEmpty() || !ColumnModification.RemovedFragments.IsEmpty())
+								{
+									LookupTable.Reset();
+									EntityCollections.Reset();
+
+									// Sort rows (entities) into to matching table (archetype) bucket.
+									for (FMassEntityHandle EntityHandle : BatchedModification.Value)
+									{
+										if (System.IsEntityValid(EntityHandle))
+										{
+											FMassArchetypeHandle Archetype = System.GetArchetypeForEntity(EntityHandle);
+											EntityHandleArray& EntityCollection = LookupTable.FindOrAdd(Archetype);
+											EntityCollection.Add(EntityHandle);
+										}
+									}
+
+									// Construct table (archetype) specific row (entity) collections.
+									EntityCollections.Reserve(LookupTable.Num());
+									for (auto It = LookupTable.CreateConstIterator(); It; ++It)
+									{
+										// Since we use a map to combine all the operation on a row in one via a map go, we already know there won't be any duplicate
+										EntityCollections.Emplace(It.Key(), It.Value(), FMassArchetypeEntityCollection::EDuplicatesHandling::NoDuplicates);
+									}
+
+									// This could be improved by adding an operation that would both combine the Fragments and Tags change in one bath operation.
+									System.BatchChangeFragmentCompositionForEntities(EntityCollections, ColumnModification.AddedFragments,  ColumnModification.RemovedFragments);
+								}
+
+								if (!ColumnModification.AddedTags.IsEmpty() || !ColumnModification.RemovedTags.IsEmpty())
+								{
+									LookupTable.Reset();
+									EntityCollections.Reset();
+
+									// Sort rows (entities) into to matching table (archetype) bucket.
+									for (FMassEntityHandle EntityHandle : BatchedModification.Value)
+									{
+										if (System.IsEntityValid(EntityHandle))
+										{
+											FMassArchetypeHandle Archetype = System.GetArchetypeForEntity(EntityHandle);
+											EntityHandleArray& EntityCollection = LookupTable.FindOrAdd(Archetype);
+											EntityCollection.Add(EntityHandle);
+										}
+									}
+
+									// Construct table (archetype) specific row (entity) collections.
+									EntityCollections.Reserve(LookupTable.Num());
+									for (auto It = LookupTable.CreateConstIterator(); It; ++It)
+									{
+										// Since we use a map to combine all the operation on a row in one via a map go, we already know there won't be any duplicate
+										EntityCollections.Emplace(It.Key(), It.Value(), FMassArchetypeEntityCollection::EDuplicatesHandling::NoDuplicates);
+									}
+
+									System.BatchChangeTagsForEntities(EntityCollections, ColumnModification.AddedTags, ColumnModification.RemovedTags);
+								}
+							}
+
+							// Do the relocation
+							for (TPair<FMassEntityHandle, FRowOperations>& PerRowOperations : InRowOperations)
+							{
+								const FMassEntityHandle Handle = PerRowOperations.Key;
+								for (FRelocator& Relocator : PerRowOperations.Value.Relocators)
+								{
+									if (!Relocator.Identifier.IsNone())
+									{
+										Relocator.FragmentType = PtrToEnvironment->GenerateDynamicColumn(*Relocator.FragmentType, Relocator.Identifier);
+									}
+
+									FStructView Fragment = System.GetFragmentDataStruct(Handle, Relocator.FragmentType);
+									Relocator.CopyOrMove(*Relocator.FragmentType, Fragment.GetMemory(), Relocator.Object);
+								}
+							}
+
+							// Remove Rows
+							if (!InEntitiesToDestroy.IsEmpty())
+							{
+								System.BatchDestroyEntities(InEntitiesToDestroy);
+							}
+						});
+				}
+			}
+
 			uint64 GetUpdateCycleId() const 
 			{
 				return Environment.GetUpdateCycleId();
@@ -215,16 +525,9 @@ namespace UE::Editor::DataStorage
 
 			void AddColumns(TConstArrayView<RowHandle> Rows, TConstArrayView<UE::Editor::DataStorage::FDynamicColumnDescription> DynamicColumnDescriptions)
 			{
-				struct FAddDynamicColumns
-				{
-					TConstArrayView<RowHandle> Rows;
-					TConstArrayView<FDynamicColumnDescription> Descriptions;
-					TArrayView<const UScriptStruct*> ResolvedTypes;
-				};
-
 				FScratchBuffer& ScratchBuffer = Environment.GetScratchBuffer();
 				
-				FAddDynamicColumns* CommandData = ScratchBuffer.Emplace<FAddDynamicColumns>();
+				FAddDynamicColumnsBatch* CommandData = ScratchBuffer.Emplace<FAddDynamicColumnsBatch>();
 				TArrayView<RowHandle> ScratchRows = MakeArrayView(
 					ScratchBuffer.EmplaceArray<RowHandle>(Rows.Num()),
 					DynamicColumnDescriptions.Num());
@@ -239,35 +542,42 @@ namespace UE::Editor::DataStorage
 				CopyArrayViews(Rows, ScratchRows);
 				CopyArrayViews(DynamicColumnDescriptions, ScratchDescriptions);
 
-				*CommandData = FAddDynamicColumns
+				*CommandData = FAddDynamicColumnsBatch
 				{
 					.Rows = ScratchRows,
 					.Descriptions = ScratchDescriptions,
 					.ResolvedTypes = ScratchTypes
 				};
-				
-				this->Context.Defer().template PushCommand<FMassDeferredAddCommand>(
-					[CommandData, this](FMassEntityManager& System)
-					{
-						for (int32 DynamicTypeIndex = 0, DynamicTypeIndexEnd = CommandData->Descriptions.Num(); DynamicTypeIndex < DynamicTypeIndexEnd; ++DynamicTypeIndex)
+
+				if (bBatchDeferredCommands)
+				{
+					AddedDynamicColumnsBatchQueue.AddElement(CommandData);
+				}
+				else
+				{ 
+					this->Context.Defer().template PushCommand<FMassDeferredAddCommand>(
+						[CommandData, this](FMassEntityManager& System)
 						{
-							const FDynamicColumnDescription& Description = CommandData->Descriptions[DynamicTypeIndex];
-							const UScriptStruct* DynamicColumnType = Environment.GenerateDynamicColumn(*Description.TemplateType, Description.Identifier);
-							CommandData->ResolvedTypes[DynamicTypeIndex] = DynamicColumnType;
-						}
+							for (int32 DynamicTypeIndex = 0, DynamicTypeIndexEnd = CommandData->Descriptions.Num(); DynamicTypeIndex < DynamicTypeIndexEnd; ++DynamicTypeIndex)
+							{
+								const FDynamicColumnDescription& Description = CommandData->Descriptions[DynamicTypeIndex];
+								const UScriptStruct* DynamicColumnType = Environment.GenerateDynamicColumn(*Description.TemplateType, Description.Identifier);
+								CommandData->ResolvedTypes[DynamicTypeIndex] = DynamicColumnType;
+							}
 
-						FMassArchetypeCompositionDescriptor AddDescriptor;
-						TedsColumnsToMassDescriptor(AddDescriptor, CommandData->ResolvedTypes);
+							FMassArchetypeCompositionDescriptor AddDescriptor;
+							TedsColumnsToMassDescriptor(AddDescriptor, CommandData->ResolvedTypes);
 						
-						for (RowHandle Row : CommandData->Rows)
-						{
-							FMassEntityHandle Entity = FMassEntityHandle::FromNumber(Row);
-							if (System.IsEntityValid(Entity))
+							for (RowHandle Row : CommandData->Rows)
 							{
-								System.AddCompositionToEntity_GetDelta(Entity, AddDescriptor);
+								FMassEntityHandle Entity = FMassEntityHandle::FromNumber(Row);
+								if (System.IsEntityValid(Entity))
+								{
+									System.AddCompositionToEntity_GetDelta(Entity, AddDescriptor);
+								}
 							}
-						}
-					});
+						});
+				}
 			}
 
 			void* AddColumnUninitialized(RowHandle Row, const UScriptStruct* ObjectType) 
@@ -282,59 +592,42 @@ namespace UE::Editor::DataStorage
 			void* AddColumnUninitialized(RowHandle Row, const UScriptStruct* ObjectType, ObjectCopyOrMove Relocator)
 			{
 				checkf(UE::Mass::IsA<FMassFragment>(ObjectType), TEXT("Column [%s] can not be a tag"), *ObjectType->GetName());
-		
-				struct FAddValueColumn
-				{
-					ObjectCopyOrMove Relocator;
-					const UScriptStruct* FragmentType;
-					FMassEntityHandle Entity;
-					void* Object;
-
-					FAddValueColumn() = default;
-					FAddValueColumn(ObjectCopyOrMove InRelocator, const UScriptStruct* InFragmentType, FMassEntityHandle InEntity, void* InObject)
-						: Relocator(InRelocator)
-						, FragmentType(InFragmentType)
-						, Entity(InEntity)
-						, Object(InObject)
-					{}
-
-					~FAddValueColumn()
-					{
-						if ((this->FragmentType->StructFlags & (STRUCT_IsPlainOldData | STRUCT_NoDestructor)) == 0)
-						{
-							this->FragmentType->DestroyStruct(this->Object);
-						}
-					}
-				};
-
+	
 				FScratchBuffer& ScratchBuffer = Environment.GetScratchBuffer();
 				void* ColumnData = ScratchBuffer.Allocate(ObjectType->GetStructureSize(), ObjectType->GetMinAlignment());
 				FAddValueColumn* AddedColumn =
 					ScratchBuffer.Emplace<FAddValueColumn>(Relocator, ObjectType, FMassEntityHandle::FromNumber(Row), ColumnData);
-		
-				this->Context.Defer().template PushCommand<FMassDeferredAddCommand>(
-					[AddedColumn](FMassEntityManager& System)
-					{
-						// Check entity before proceeding. It's possible it may have been invalidated before this deferred call fired.
-						if (System.IsEntityActive(AddedColumn->Entity))
+
+				if (bBatchDeferredCommands)
+				{
+					AddedColumnWithRelocationQueue.AddElement(AddedColumn);
+				}
+				else
+				{
+					this->Context.Defer().template PushCommand<FMassDeferredAddCommand>(
+						[AddedColumn](FMassEntityManager& System)
 						{
-							// Check before adding.  Mass's AddFragmentToEntity is not idempotent and will assert if adding
-							// column to a row that already has one
-							FStructView Fragment = System.GetFragmentDataStruct(AddedColumn->Entity, AddedColumn->FragmentType);
-							if (!Fragment.IsValid())
+							// Check entity before proceeding. It's possible it may have been invalidated before this deferred call fired.
+							if (System.IsEntityActive(AddedColumn->Entity))
 							{
-								System.AddFragmentToEntity(AddedColumn->Entity, AddedColumn->FragmentType, 
-									[AddedColumn](void* Fragment, const UScriptStruct& FragmentType)
-									{
-										AddedColumn->Relocator(FragmentType, Fragment, AddedColumn->Object);
-									});
+								// Check before adding.  Mass's AddFragmentToEntity is not idempotent and will assert if adding
+								// column to a row that already has one
+								FStructView Fragment = System.GetFragmentDataStruct(AddedColumn->Entity, AddedColumn->FragmentType);
+								if (!Fragment.IsValid())
+								{
+									System.AddFragmentToEntity(AddedColumn->Entity, AddedColumn->FragmentType, 
+										[AddedColumn](void* Fragment, const UScriptStruct& FragmentType)
+										{
+											AddedColumn->Relocator(FragmentType, Fragment, AddedColumn->Object);
+										});
+								}
+								else
+								{
+									AddedColumn->Relocator(*AddedColumn->FragmentType, Fragment.GetMemory(), AddedColumn->Object);
+								}
 							}
-							else
-							{
-								AddedColumn->Relocator(*AddedColumn->FragmentType, Fragment.GetMemory(), AddedColumn->Object);
-							}
-						}
-					});
+						});
+					}
 		
 				return ColumnData;
 			}
@@ -349,103 +642,75 @@ namespace UE::Editor::DataStorage
 
 			void* AddColumnUninitialized(RowHandle Row, const FDynamicColumnDescription& Description, ObjectCopyOrMove Relocator)
 			{
-				struct FAddDynamicColumn
-				{
-					ObjectCopyOrMove Relocator;
-					FDynamicColumnDescription Description;
-					FMassEntityHandle Entity;
-					void* Object;
-					bool bNeedsDestruction;
-
-					FAddDynamicColumn() = default;
-					FAddDynamicColumn(ObjectCopyOrMove InRelocator, const FDynamicColumnDescription& InDescription, FMassEntityHandle InEntity, void* InObject)
-						: Relocator(InRelocator)
-						, Description(InDescription)
-						, Entity(InEntity)
-						, Object(InObject)
-					{
-						// Check here and cache off the result to avoid command buffer needing to dereference UScriptStruct to check if anything
-						// needs to be done. In many cases, this is expected to be false
-						bNeedsDestruction = (this->Description.TemplateType->StructFlags & (STRUCT_IsPlainOldData | STRUCT_NoDestructor)) == 0;
-					}
-
-					~FAddDynamicColumn()
-					{
-						if (bNeedsDestruction)
-						{
-							this->Description.TemplateType->DestroyStruct(this->Object);
-						}
-					}
-				};
-
 				FScratchBuffer& ScratchBuffer = Environment.GetScratchBuffer();
 				// DynamicColumn types are derivations from their template that add no new members.  The size and alignment will be the same
 				void* ColumnData = ScratchBuffer.Allocate(Description.TemplateType->GetStructureSize(), Description.TemplateType->GetMinAlignment());
 				FAddDynamicColumn* AddedColumn =
 					ScratchBuffer.Emplace<FAddDynamicColumn>(Relocator, Description, FMassEntityHandle::FromNumber(Row), ColumnData);
 		
-				this->Context.Defer().template PushCommand<FMassDeferredAddCommand>(
-					[AddedColumn, PtrToEnvironment = &Environment](FMassEntityManager& System)
-					{
-						// Check entity before proceeding. It's possible it may have been invalidated before this deferred call fired.
-						if (System.IsEntityActive(AddedColumn->Entity))
+				if (bBatchDeferredCommands)
+				{
+					AddedDynamicColumnWithRelocationQueue.AddElement(AddedColumn);
+				}
+				else
+				{
+					this->Context.Defer().template PushCommand<FMassDeferredAddCommand>(
+						[AddedColumn, PtrToEnvironment = &Environment](FMassEntityManager& System)
 						{
-							const UScriptStruct* DynamicStructType = PtrToEnvironment->GenerateDynamicColumn(*AddedColumn->Description.TemplateType, AddedColumn->Description.Identifier);
-
-							FStructView Fragment = System.GetFragmentDataStruct(AddedColumn->Entity, DynamicStructType);
-							// Check before adding.  Mass's AddFragmentToEntity is not idempotent and will assert if adding
-							// column to a row that already has one
-							if (!Fragment.IsValid())
+							// Check entity before proceeding. It's possible it may have been invalidated before this deferred call fired.
+							if (System.IsEntityActive(AddedColumn->Entity))
 							{
-								System.AddFragmentToEntity(AddedColumn->Entity, DynamicStructType, 
-									[AddedColumn](void* Fragment, const UScriptStruct& FragmentType)
-									{
-										AddedColumn->Relocator(FragmentType, Fragment, AddedColumn->Object);
-									});
+								const UScriptStruct* DynamicStructType = PtrToEnvironment->GenerateDynamicColumn(*AddedColumn->Description.TemplateType, AddedColumn->Description.Identifier);
+
+								FStructView Fragment = System.GetFragmentDataStruct(AddedColumn->Entity, DynamicStructType);
+								// Check before adding.  Mass's AddFragmentToEntity is not idempotent and will assert if adding
+								// column to a row that already has one
+								if (!Fragment.IsValid())
+								{
+									System.AddFragmentToEntity(AddedColumn->Entity, DynamicStructType, 
+										[AddedColumn](void* Fragment, const UScriptStruct& FragmentType)
+										{
+											AddedColumn->Relocator(FragmentType, Fragment, AddedColumn->Object);
+										});
+								}
+								else
+								{
+									AddedColumn->Relocator(*DynamicStructType, Fragment.GetMemory(), AddedColumn->Object);
+								}
 							}
-							else
-							{
-								AddedColumn->Relocator(*DynamicStructType, Fragment.GetMemory(), AddedColumn->Object);
-							}
-						}
-					});
+						});
+				}
 		
 				return ColumnData;
 			}
 	
 			void AddColumns(RowHandle Row, TConstArrayView<const UScriptStruct*> ColumnTypes) 
 			{
-				struct FAddedColumns
-				{
-					FMassArchetypeCompositionDescriptor AddDescriptor;
-					FMassEntityHandle Entity;
-				};
-
 				FAddedColumns* AddedColumns = Environment.GetScratchBuffer().template Emplace<FAddedColumns>();
 				TedsColumnsToMassDescriptor(AddedColumns->AddDescriptor, ColumnTypes);
 				AddedColumns->Entity = FMassEntityHandle::FromNumber(Row);
 
-				this->Context.Defer().template PushCommand<FMassDeferredAddCommand>(
-					[AddedColumns](FMassEntityManager& System)
-					{
-						if (System.IsEntityValid(AddedColumns->Entity))
+				if (bBatchDeferredCommands)
+				{ 
+					AddedColumnsQueue.AddElement(AddedColumns);
+				}
+				else
+				{ 
+					this->Context.Defer().template PushCommand<FMassDeferredAddCommand>(
+						[AddedColumns](FMassEntityManager& System)
 						{
-							System.AddCompositionToEntity_GetDelta(AddedColumns->Entity, AddedColumns->AddDescriptor);
-						}
-					});
+							if (System.IsEntityValid(AddedColumns->Entity))
+							{
+								System.AddCompositionToEntity_GetDelta(AddedColumns->Entity, AddedColumns->AddDescriptor);
+							}
+						});
+				}
 			}
 
 			void AddColumns(TConstArrayView<RowHandle> Rows, TConstArrayView<const UScriptStruct*> ColumnTypes) 
 			{
-				struct FAddedColumns
-				{
-					FMassArchetypeCompositionDescriptor AddDescriptor;
-					FMassEntityHandle* Entities;
-					int32 EntityCount;
-				};
-
 				FScratchBuffer& ScratchBuffer = Environment.GetScratchBuffer();
-				FAddedColumns* AddedColumns = ScratchBuffer.Emplace<FAddedColumns>();
+				FAddedColumnsBatch* AddedColumns = ScratchBuffer.Emplace<FAddedColumnsBatch>();
 				TedsColumnsToMassDescriptor(AddedColumns->AddDescriptor, ColumnTypes);
 		
 				FMassEntityHandle* Entities = ScratchBuffer.EmplaceArray<FMassEntityHandle>(Rows.Num());
@@ -457,6 +722,12 @@ namespace UE::Editor::DataStorage
 				}
 				AddedColumns->EntityCount = Rows.Num();
 
+				if (bBatchDeferredCommands)
+				{
+					AddedColumnsBatchQueue.AddElement(AddedColumns);
+				}
+				else
+				{
 				this->Context.Defer().template PushCommand<FMassDeferredAddCommand>(
 					[AddedColumns](FMassEntityManager& System)
 					{
@@ -470,44 +741,40 @@ namespace UE::Editor::DataStorage
 							}
 						}
 					});
+				}
 			}
 
 			void RemoveColumns(RowHandle Row, TConstArrayView<const UScriptStruct*> ColumnTypes) 
 			{
-				struct FRemovedColumns
-				{
-					FMassArchetypeCompositionDescriptor RemoveDescriptor;
-					FMassEntityHandle Entity;
-				};
-
 				FRemovedColumns* RemovedColumns = Environment.GetScratchBuffer().template Emplace<FRemovedColumns>();
 				TedsColumnsToMassDescriptorIfActiveTable(RemovedColumns->RemoveDescriptor, ColumnTypes);
 				if (!RemovedColumns->RemoveDescriptor.IsEmpty())
 				{
 					RemovedColumns->Entity = FMassEntityHandle::FromNumber(Row);
 
-					this->Context.Defer().template PushCommand<FMassDeferredAddCommand>(
-						[RemovedColumns](FMassEntityManager& System)
-						{
-							if (System.IsEntityValid(RemovedColumns->Entity))
+
+					if (bBatchDeferredCommands)
+					{
+						RemovedColumnsQueue.AddElement(RemovedColumns);
+					}
+					else
+					{ 
+						this->Context.Defer().template PushCommand<FMassDeferredAddCommand>(
+							[RemovedColumns](FMassEntityManager& System)
 							{
-								System.RemoveCompositionFromEntity(RemovedColumns->Entity, RemovedColumns->RemoveDescriptor);
-							}
-						});
+								if (System.IsEntityValid(RemovedColumns->Entity))
+								{
+									System.RemoveCompositionFromEntity(RemovedColumns->Entity, RemovedColumns->RemoveDescriptor);
+								}
+							});
+					}
 				}
 			}
 
 			void RemoveColumns(TConstArrayView<RowHandle> Rows, TConstArrayView<const UScriptStruct*> ColumnTypes) 
 			{
-				struct FRemovedColumns
-				{
-					FMassArchetypeCompositionDescriptor RemoveDescriptor;
-					FMassEntityHandle* Entities;
-					int32 EntityCount;
-				};
-
 				FScratchBuffer& ScratchBuffer = Environment.GetScratchBuffer();
-				FRemovedColumns* RemovedColumns = ScratchBuffer.Emplace<FRemovedColumns>();
+				FRemovedColumnsBatch* RemovedColumns = ScratchBuffer.Emplace<FRemovedColumnsBatch>();
 				TedsColumnsToMassDescriptorIfActiveTable(RemovedColumns->RemoveDescriptor, ColumnTypes);
 
 				FMassEntityHandle* Entities = ScratchBuffer.EmplaceArray<FMassEntityHandle>(Rows.Num());
@@ -519,49 +786,79 @@ namespace UE::Editor::DataStorage
 				}
 				RemovedColumns->EntityCount = Rows.Num();
 
-				this->Context.Defer().template PushCommand<FMassDeferredAddCommand>(
-					[RemovedColumns](FMassEntityManager& System)
-					{
-						FMassEntityHandle* Entities = RemovedColumns->Entities;
-						int32 Count = RemovedColumns->EntityCount;
+				if (bBatchDeferredCommands)
+				{
+					RemovedColumnsBatchQueue.AddElement(RemovedColumns);
+				}
+				else
+				{
+					this->Context.Defer().template PushCommand<FMassDeferredRemoveCommand>(
+						[RemovedColumns](FMassEntityManager& System)
+						{
+							FMassEntityHandle* Entities = RemovedColumns->Entities;
+							int32 Count = RemovedColumns->EntityCount;
 
-						using EntityHandleArray = TArray<FMassEntityHandle, TInlineAllocator<32>>;
-						using EntityArchetypeLookup = TMap<FMassArchetypeHandle, EntityHandleArray, TInlineSetAllocator<32>>;
-						using ArchetypeEntityArray = TArray<FMassArchetypeEntityCollection, TInlineAllocator<32>>;
+							using EntityHandleArray = TArray<FMassEntityHandle, TInlineAllocator<32>>;
+							using EntityArchetypeLookup = TMap<FMassArchetypeHandle, EntityHandleArray, TInlineSetAllocator<32>>;
+							using ArchetypeEntityArray = TArray<FMassArchetypeEntityCollection, TInlineAllocator<32>>;
 
-						// Sort rows (entities) into to matching table (archetype) bucket.
-						EntityArchetypeLookup LookupTable;
-						for (int32 Counter = 0; Counter < Count; ++Counter)
-						{
-							if (System.IsEntityValid(*Entities))
-							{
-								FMassArchetypeHandle Archetype = System.GetArchetypeForEntity(*Entities);
-								EntityHandleArray& EntityCollection = LookupTable.FindOrAdd(Archetype);
-								EntityCollection.Add(*Entities);
-							}
-							Entities++;
-						}
+							// Sort rows (entities) into to matching table (archetype) bucket.
+							EntityArchetypeLookup LookupTable;
+							ArchetypeEntityArray EntityCollections;
 
-						// Construct table (archetype) specific row (entity) collections.
-						ArchetypeEntityArray EntityCollections;
-						EntityCollections.Reserve(LookupTable.Num());
-						for (auto It = LookupTable.CreateConstIterator(); It; ++It)
-						{
-							// Could be more effective but the previous implementation was robust when called with duplicate rows.
-							EntityCollections.Emplace(It.Key(), It.Value(), FMassArchetypeEntityCollection::EDuplicatesHandling::FoldDuplicates);
-						}
 
+							// This could be improved by adding an operation that would both combine the Fragments and Tags change in one bath operation.
+							if (!RemovedColumns->RemoveDescriptor.Fragments.IsEmpty())
+							{
+								for (int32 Counter = 0; Counter < Count; ++Counter)
+								{
+									if (System.IsEntityValid(*Entities))
+									{
+										FMassArchetypeHandle Archetype = System.GetArchetypeForEntity(*Entities);
+										EntityHandleArray& EntityCollection = LookupTable.FindOrAdd(Archetype);
+										EntityCollection.Add(*Entities);
+									}
+									Entities++;
+								}
+
+								// Construct table (archetype) specific row (entity) collections.
+								EntityCollections.Reserve(LookupTable.Num());
+								for (auto It = LookupTable.CreateConstIterator(); It; ++It)
+								{
+									// Could be more effective but the previous implementation was robust when called with duplicate rows.
+									EntityCollections.Emplace(It.Key(), It.Value(), FMassArchetypeEntityCollection::EDuplicatesHandling::FoldDuplicates);
+								}
+
+								System.BatchChangeFragmentCompositionForEntities(EntityCollections, FMassFragmentBitSet(), RemovedColumns->RemoveDescriptor.Fragments);
+							}
+							if (!RemovedColumns->RemoveDescriptor.Tags.IsEmpty())
+							{
+								LookupTable.Reset();
+								EntityCollections.Reset();
 
-						// This could be improved by adding an operation that would both combine the Fragments and Tags change in one bath operation.
-						if (!RemovedColumns->RemoveDescriptor.Fragments.IsEmpty())
-						{
-							System.BatchChangeFragmentCompositionForEntities(EntityCollections, FMassFragmentBitSet(), RemovedColumns->RemoveDescriptor.Fragments);
-						}
-						if (!RemovedColumns->RemoveDescriptor.Tags.IsEmpty())
-						{
-							System.BatchChangeTagsForEntities(EntityCollections, FMassTagBitSet(), RemovedColumns->RemoveDescriptor.Tags);
-						}
-					});
+								for (int32 Counter = 0; Counter < Count; ++Counter)
+								{
+									if (System.IsEntityValid(*Entities))
+									{
+										FMassArchetypeHandle Archetype = System.GetArchetypeForEntity(*Entities);
+										EntityHandleArray& EntityCollection = LookupTable.FindOrAdd(Archetype);
+										EntityCollection.Add(*Entities);
+									}
+									Entities++;
+								}
+
+								// Construct table (archetype) specific row (entity) collections.
+								EntityCollections.Reserve(LookupTable.Num());
+								for (auto It = LookupTable.CreateConstIterator(); It; ++It)
+								{
+									// Could be more effective but the previous implementation was robust when called with duplicate rows.
+									EntityCollections.Emplace(It.Key(), It.Value(), FMassArchetypeEntityCollection::EDuplicatesHandling::FoldDuplicates);
+								}
+
+								System.BatchChangeTagsForEntities(EntityCollections, FMassTagBitSet(), RemovedColumns->RemoveDescriptor.Tags);
+							}
+						});
+					}
 			}
 	
 			RowHandle AddRow(TableHandle Table) 
@@ -573,24 +870,25 @@ namespace UE::Editor::DataStorage
 				{
 					return InvalidRowHandle;
 				}
-		
-				struct CommandInfo
-				{
-					FMassEntityHandle Entity;
-					FMassArchetypeHandle Archetype;
-				};
 
-				CommandInfo CommandDataTmp{
-					.Entity = EntityHandle,
-					.Archetype = MoveTemp(ArchetypeHandle)
-				};
+				FAddRow AddRowTmp{
+						.Entity = EntityHandle,
+						.Archetype = MoveTemp(ArchetypeHandle)
+					};
 
-				this->Context.Defer().template PushCommand<FMassDeferredCreateCommand>(
-					[CommandData = MoveTemp(CommandDataTmp)](FMassEntityManager& System)
-					{
-						const FMassArchetypeSharedFragmentValues SharedFragmentValues;
-						System.BuildEntity(CommandData.Entity, CommandData.Archetype, SharedFragmentValues);
-					});
+				if (bBatchDeferredCommands)
+				{
+					AddedRowsQueue.AddElement(MoveTemp(AddRowTmp));
+				}
+				else
+				{
+					this->Context.Defer().template PushCommand<FMassDeferredCreateCommand>(
+						[AddRow = MoveTemp(AddRowTmp)](FMassEntityManager& System)
+						{
+							const FMassArchetypeSharedFragmentValues SharedFragmentValues;
+							System.BuildEntity(AddRow.Entity, AddRow.Archetype, SharedFragmentValues);
+						});
+				}
 		
 				const RowHandle TedsRowHandle = EntityHandle.AsNumber();
 				return TedsRowHandle;
@@ -598,12 +896,29 @@ namespace UE::Editor::DataStorage
 
 			void RemoveRow(RowHandle Row) 
 			{
-				this->Context.Defer().DestroyEntity(FMassEntityHandle::FromNumber(Row));
+				if (bBatchDeferredCommands)
+				{
+					RemovedRowsQueue.AddElement(Row);
+				}
+				else
+				{
+					this->Context.Defer().DestroyEntity(FMassEntityHandle::FromNumber(Row));
+				}
 			}
 
 			void RemoveRows(TConstArrayView<RowHandle> Rows) 
 			{
-				this->Context.Defer().DestroyEntities(RowsToMassEntitiesConversion(Rows));
+				if (bBatchDeferredCommands)
+				{
+					for (const RowHandle& Row : Rows)
+					{
+						RemovedRowsQueue.AddElement(Row);
+					}
+				}
+				else
+				{
+					this->Context.Defer().DestroyEntities(RowsToMassEntitiesConversion(Rows));
+				}
 			}
 
 			void PushCommand(void (*CommandFunction)(void*), void* InCommandData)
@@ -650,6 +965,109 @@ namespace UE::Editor::DataStorage
 				}
 				return ObjectMemory;
 			}
+
+			struct FAddedColumns
+			{
+				FMassArchetypeCompositionDescriptor AddDescriptor;
+				FMassEntityHandle Entity;
+			};
+
+			struct FAddedColumnsBatch
+			{
+				FMassArchetypeCompositionDescriptor AddDescriptor;
+				FMassEntityHandle* Entities;
+				int32 EntityCount;
+			};
+
+			struct FAddValueColumn
+			{
+				ObjectCopyOrMove Relocator;
+				const UScriptStruct* FragmentType;
+				FMassEntityHandle Entity;
+				void* Object;
+
+				FAddValueColumn() = default;
+				FAddValueColumn(ObjectCopyOrMove InRelocator, const UScriptStruct* InFragmentType, FMassEntityHandle InEntity, void* InObject)
+					: Relocator(InRelocator)
+					, FragmentType(InFragmentType)
+					, Entity(InEntity)
+					, Object(InObject)
+				{}
+
+				~FAddValueColumn()
+				{
+					if (this->FragmentType && (this->FragmentType->StructFlags & (STRUCT_IsPlainOldData | STRUCT_NoDestructor)) == 0)
+					{
+						this->FragmentType->DestroyStruct(this->Object);
+					}
+				}
+			};
+
+			struct FAddDynamicColumn
+			{
+				ObjectCopyOrMove Relocator;
+				FDynamicColumnDescription Description;
+				FMassEntityHandle Entity;
+				void* Object;
+				bool bNeedsDestruction;
+
+				FAddDynamicColumn() = default;
+				FAddDynamicColumn(ObjectCopyOrMove InRelocator, const FDynamicColumnDescription& InDescription, FMassEntityHandle InEntity, void* InObject)
+					: Relocator(InRelocator)
+					, Description(InDescription)
+					, Entity(InEntity)
+					, Object(InObject)
+				{
+					// Check here and cache off the result to avoid command buffer needing to dereference UScriptStruct to check if anything
+					// needs to be done. In many cases, this is expected to be false
+					bNeedsDestruction = (this->Description.TemplateType->StructFlags & (STRUCT_IsPlainOldData | STRUCT_NoDestructor)) == 0;
+				}
+
+				~FAddDynamicColumn()
+				{
+					if (bNeedsDestruction)
+					{
+						this->Description.TemplateType->DestroyStruct(this->Object);
+					}
+				}
+			};
+
+			struct FAddDynamicColumnsBatch
+			{
+				TConstArrayView<RowHandle> Rows;
+				TConstArrayView<FDynamicColumnDescription> Descriptions;
+				TArrayView<const UScriptStruct*> ResolvedTypes;
+			};
+
+			struct FRemovedColumns
+			{
+				FMassArchetypeCompositionDescriptor RemoveDescriptor;
+				FMassEntityHandle Entity;
+			};
+
+			struct FRemovedColumnsBatch
+			{
+				FMassArchetypeCompositionDescriptor RemoveDescriptor;
+				FMassEntityHandle* Entities;
+				int32 EntityCount;
+			};
+
+			struct FAddRow
+			{
+				FMassEntityHandle Entity;
+				FMassArchetypeHandle Archetype;
+			};
+
+			bool bBatchDeferredCommands = false;
+			TChunkedArray<FAddedColumns*> AddedColumnsQueue;
+			TChunkedArray<FAddedColumnsBatch*> AddedColumnsBatchQueue;
+			TChunkedArray<FAddValueColumn*> AddedColumnWithRelocationQueue;
+			TChunkedArray<FAddDynamicColumn*> AddedDynamicColumnWithRelocationQueue;
+			TChunkedArray<FAddDynamicColumnsBatch*> AddedDynamicColumnsBatchQueue;
+			TChunkedArray<FRemovedColumns*> RemovedColumnsQueue;
+			TChunkedArray<FRemovedColumnsBatch*> RemovedColumnsBatchQueue;
+			TChunkedArray<FAddRow> AddedRowsQueue;
+			TChunkedArray<RowHandle> RemovedRowsQueue;
 		};
 
 		struct FMassDirectContextForwarder final : public IEditorDataStorageProvider::IDirectQueryContext
@@ -674,8 +1092,8 @@ namespace UE::Editor::DataStorage
 
 		struct FMassSubqueryContextForwarder  : public IEditorDataStorageProvider::ISubqueryContext
 		{
-			FMassSubqueryContextForwarder(FMassExecutionContext& InContext, FEnvironment& InEnvironment)
-				: Implementation(InContext, InEnvironment)
+			FMassSubqueryContextForwarder(FMassExecutionContext& InContext, FEnvironment& InEnvironment, bool bInBatchDeferredCommands)
+				: Implementation(InContext, InEnvironment, bInBatchDeferredCommands)
 			{}
 
 			~FMassSubqueryContextForwarder() override = default;
@@ -719,7 +1137,7 @@ namespace UE::Editor::DataStorage
 				FMassExecutionContext& InContext, 
 				FExtendedQueryStore& InQueryStore,
 				FEnvironment& InEnvironment)
-				: FMassWithEnvironmentContextCommon(InContext, InEnvironment)
+				: FMassWithEnvironmentContextCommon(InContext, InEnvironment, InQueryDescription.bShouldBatchModifications)
 				, QueryDescription(InQueryDescription)
 				, QueryStore(InQueryStore)
 			{
@@ -1164,7 +1582,7 @@ UE::Editor::DataStorage::FQueryResult FTypedElementQueryProcessorData::Execute(
 			[&Result, &Callback, &Description, &Environment](FMassExecutionContext& Context)
 			{
 				// No need to cache any subsystem dependencies as these are not accessible from a subquery.
-				UE::Editor::DataStorage::Processors::Private::FMassSubqueryContextForwarder QueryContext(Context, Environment);
+				UE::Editor::DataStorage::Processors::Private::FMassSubqueryContextForwarder QueryContext(Context, Environment, Description.bShouldBatchModifications);
 				Callback(Description, QueryContext);
 				Result.Count += Context.GetNumEntities();
 			}
@@ -1203,7 +1621,7 @@ UE::Editor::DataStorage::FQueryResult FTypedElementQueryProcessorData::Execute(
 			[&Result, &Callback, &Description, &Environment, &EntityManager, RowHandle](FMassExecutionContext& Context)
 			{
 				// No need to cache any subsystem dependencies as these are not accessible from a subquery.
-				Processors::Private::FMassSubqueryContextForwarder QueryContext(Context, Environment);
+				Processors::Private::FMassSubqueryContextForwarder QueryContext(Context, Environment, Description.bShouldBatchModifications);
 				Callback(Description, QueryContext);
 				Result.Count += Context.GetNumEntities();
 			}
diff --git a/Engine/Plugins/Experimental/EditorDataStorageFeatures/Source/TedsAssetData/Private/CB/TedsAssetDataCBDataSource.cpp b/Engine/Plugins/Experimental/EditorDataStorageFeatures/Source/TedsAssetData/Private/CB/TedsAssetDataCBDataSource.cpp
--- a/Engine/Plugins/Experimental/EditorDataStorageFeatures/Source/TedsAssetData/Private/CB/TedsAssetDataCBDataSource.cpp
+++ b/Engine/Plugins/Experimental/EditorDataStorageFeatures/Source/TedsAssetData/Private/CB/TedsAssetDataCBDataSource.cpp
@@ -68,7 +68,8 @@ FTedsAssetDataCBDataSource::FTedsAssetDataCBDataSource(IEditorDataStorageProvide
 	ProcessPathQuery = Database.RegisterQuery(
 		Select(
 			TEXT("FTedsAssetDataCBDataSource: Process Path updates"),
-			FProcessor(EQueryTickPhase::DuringPhysics, Database.GetQueryTickGroupName(EQueryTickGroups::Update)),
+			FProcessor(EQueryTickPhase::DuringPhysics, Database.GetQueryTickGroupName(EQueryTickGroups::Update))
+				.BatchModifications(true),
 			[DataSource = const_cast<const FTedsAssetDataCBDataSource*>(this)](IQueryContext& Context, const RowHandle* Rows, const FAssetPathColumn_Experimental* PathColumn)
 			{
 				DataSource->ProcessPathQueryCallback(Context, Rows, PathColumn);
@@ -80,7 +81,8 @@ FTedsAssetDataCBDataSource::FTedsAssetDataCBDataSource(IEditorDataStorageProvide
 	ProcessAssetDataPathUpdateQuery = Database.RegisterQuery(
 		Select(
 			TEXT("FTedsAssetDataCBDataSource:: Process Asset Data Path Update"),
-			FProcessor(EQueryTickPhase::DuringPhysics, Database.GetQueryTickGroupName(EQueryTickGroups::Update)),
+			FProcessor(EQueryTickPhase::DuringPhysics, Database.GetQueryTickGroupName(EQueryTickGroups::Update))
+				.BatchModifications(true),
 			[DataSource = const_cast<const FTedsAssetDataCBDataSource*>(this)](IQueryContext& Context, const RowHandle* Row, const FAssetDataColumn_Experimental* AssetDataColumn)
 			{
 				DataSource->ProcessAssetDataPathUpdateQueryCallback(Context, Row, AssetDataColumn);
@@ -93,7 +95,8 @@ FTedsAssetDataCBDataSource::FTedsAssetDataCBDataSource(IEditorDataStorageProvide
 	ProcessAssetDataAndPathUpdateQuery = Database.RegisterQuery(
 		Select(
 			TEXT("FTedsAssetDataCBDataSource: Process Asset Data and Path Updates"),
-			FProcessor(EQueryTickPhase::DuringPhysics, Database.GetQueryTickGroupName(EQueryTickGroups::Update)),
+			FProcessor(EQueryTickPhase::DuringPhysics, Database.GetQueryTickGroupName(EQueryTickGroups::Update))
+				.BatchModifications(true),
 			[DataSource = const_cast<const FTedsAssetDataCBDataSource*>(this)](IQueryContext& Context, const RowHandle* Rows, const FAssetDataColumn_Experimental* AssetDataColumn)
 			{
 				DataSource->ProcessAssetDataAndPathUpdateQueryCallback(Context, Rows, AssetDataColumn);
@@ -105,7 +108,8 @@ FTedsAssetDataCBDataSource::FTedsAssetDataCBDataSource(IEditorDataStorageProvide
 	ProcessAssetDataUpdateQuery = Database.RegisterQuery(
 		Select(
 			TEXT("FTedsAssetDataCBDataSource: Process Asset Data updates"),
-			FProcessor(EQueryTickPhase::DuringPhysics, Database.GetQueryTickGroupName(EQueryTickGroups::Update)),
+			FProcessor(EQueryTickPhase::DuringPhysics, Database.GetQueryTickGroupName(EQueryTickGroups::Update))
+				.BatchModifications(true),
 			[DataSource = const_cast<const FTedsAssetDataCBDataSource*>(this)](IQueryContext& Context, const RowHandle* Rows, const FAssetDataColumn_Experimental* AssetDataColumn)
 			{
 				DataSource->ProcessAssetDataUpdateQueryCallback(Context, Rows, AssetDataColumn);
@@ -249,6 +253,10 @@ void FTedsAssetDataCBDataSource::AddAssetDataColumns(UE::Editor::DataStorage::IQ
 		DiskSizeColumn.DiskSize = OptionalPackageData->DiskSize;
 		Context.AddColumn(Row, MoveTemp(DiskSizeColumn));
 	}
+	else
+	{
+		Context.RemoveColumns<FDiskSizeColumn>(Row);
+	}
 
 	FAssetClassColumn AssetClassColumn;
 	AssetClassColumn.ClassPath = AssetData.AssetClassPath;
diff --git a/Engine/Source/Runtime/MassEntity/Private/MassEntityManager.cpp b/Engine/Source/Runtime/MassEntity/Private/MassEntityManager.cpp
--- a/Engine/Source/Runtime/MassEntity/Private/MassEntityManager.cpp
+++ b/Engine/Source/Runtime/MassEntity/Private/MassEntityManager.cpp
@@ -784,7 +784,7 @@ TSharedRef<FMassEntityManager::FEntityCreationContext> FMassEntityManager::Batch
 }
 
 TSharedRef<FMassEntityManager::FEntityCreationContext> FMassEntityManager::BatchBuildEntities(const FMassArchetypeEntityCollectionWithPayload& EncodedEntitiesWithPayload
-	, FMassArchetypeCompositionDescriptor&& Composition
+	, const FMassArchetypeCompositionDescriptor& Composition
 	, const FMassArchetypeSharedFragmentValues& SharedFragmentValues, const FMassArchetypeCreationParams& CreationParams)
 {
 	CHECK_SYNC_API_RETURN(MakeShareable(new FEntityCreationContext));
diff --git a/Engine/Source/Runtime/MassEntity/Public/MassEntityManager.h b/Engine/Source/Runtime/MassEntity/Public/MassEntityManager.h
--- a/Engine/Source/Runtime/MassEntity/Public/MassEntityManager.h
+++ b/Engine/Source/Runtime/MassEntity/Public/MassEntityManager.h
@@ -462,7 +462,7 @@ public:
 
 	TSharedRef<FEntityCreationContext> BatchBuildEntities(const FMassArchetypeEntityCollectionWithPayload& EncodedEntitiesWithPayload, const FMassFragmentBitSet& FragmentsAffected
 		, const FMassArchetypeSharedFragmentValues& SharedFragmentValues = {}, const FMassArchetypeCreationParams& CreationParams = FMassArchetypeCreationParams());
-	TSharedRef<FEntityCreationContext> BatchBuildEntities(const FMassArchetypeEntityCollectionWithPayload& EncodedEntitiesWithPayload, FMassArchetypeCompositionDescriptor&& Composition
+	TSharedRef<FEntityCreationContext> BatchBuildEntities(const FMassArchetypeEntityCollectionWithPayload& EncodedEntitiesWithPayload, const FMassArchetypeCompositionDescriptor& Composition
 		, const FMassArchetypeSharedFragmentValues& SharedFragmentValues = {}, const FMassArchetypeCreationParams& CreationParams = FMassArchetypeCreationParams());
 	void BatchChangeTagsForEntities(TConstArrayView<FMassArchetypeEntityCollection> EntityCollections, const FMassTagBitSet& TagsToAdd, const FMassTagBitSet& TagsToRemove);
 	void BatchChangeFragmentCompositionForEntities(TConstArrayView<FMassArchetypeEntityCollection> EntityCollections, const FMassFragmentBitSet& FragmentsToAdd, const FMassFragmentBitSet& FragmentsToRemove);
diff --git a/Engine/Source/Runtime/TypedElementFramework/Private/Elements/Framework/TypedElementQueryBuilder.cpp b/Engine/Source/Runtime/TypedElementFramework/Private/Elements/Framework/TypedElementQueryBuilder.cpp
--- a/Engine/Source/Runtime/TypedElementFramework/Private/Elements/Framework/TypedElementQueryBuilder.cpp
+++ b/Engine/Source/Runtime/TypedElementFramework/Private/Elements/Framework/TypedElementQueryBuilder.cpp
@@ -320,6 +320,11 @@ namespace UE::Editor::DataStorage::Queries
 		return *this;
 	}
 
+	FProcessor& FProcessor::BatchModifications(bool bBatch)
+	{
+		bBatchModifications = bBatch;
+		return *this;
+	}
 
 	/**
 	 * FObserver
diff --git a/Engine/Source/Runtime/TypedElementFramework/Public/Elements/Common/TypedElementQueryDescription.h b/Engine/Source/Runtime/TypedElementFramework/Public/Elements/Common/TypedElementQueryDescription.h
--- a/Engine/Source/Runtime/TypedElementFramework/Public/Elements/Common/TypedElementQueryDescription.h
+++ b/Engine/Source/Runtime/TypedElementFramework/Public/Elements/Common/TypedElementQueryDescription.h
@@ -119,5 +119,6 @@ namespace UE::Editor::DataStorage
 		FMetaData MetaData;
 
 		EActionType Action;
+		bool bShouldBatchModifications = false;
 	};
 } // namespace UE::Editor::DataStorage
diff --git a/Engine/Source/Runtime/TypedElementFramework/Public/Elements/Framework/TypedElementQueryBuilder.h b/Engine/Source/Runtime/TypedElementFramework/Public/Elements/Framework/TypedElementQueryBuilder.h
--- a/Engine/Source/Runtime/TypedElementFramework/Public/Elements/Framework/TypedElementQueryBuilder.h
+++ b/Engine/Source/Runtime/TypedElementFramework/Public/Elements/Framework/TypedElementQueryBuilder.h
@@ -240,13 +240,15 @@ namespace UE::Editor::DataStorage::Queries
 		TYPEDELEMENTFRAMEWORK_API FProcessor& SetAfterGroup(FName GroupName);
 		TYPEDELEMENTFRAMEWORK_API FProcessor& SetExecutionMode(EExecutionMode Mode);
 		TYPEDELEMENTFRAMEWORK_API FProcessor& MakeActivatable(FName Name);
-		
+		TYPEDELEMENTFRAMEWORK_API FProcessor& BatchModifications(bool bBatch);
+
 		EQueryTickPhase Phase;
 		FName Group;
 		FName BeforeGroup;
 		FName AfterGroup;
 		FName ActivationName;
 		EExecutionMode ExecutionMode = EExecutionMode::Default;
+		bool bBatchModifications = false;
 	};
 
 	struct FObserver final : public FQueryCallbackType
diff --git a/Engine/Source/Runtime/TypedElementFramework/Public/Elements/Framework/TypedElementQueryBuilder.inl b/Engine/Source/Runtime/TypedElementFramework/Public/Elements/Framework/TypedElementQueryBuilder.inl
--- a/Engine/Source/Runtime/TypedElementFramework/Public/Elements/Framework/TypedElementQueryBuilder.inl
+++ b/Engine/Source/Runtime/TypedElementFramework/Public/Elements/Framework/TypedElementQueryBuilder.inl
@@ -1045,6 +1045,8 @@ namespace UE::Editor::DataStorage::Queries
 				Query.Callback.AfterGroups.Add(Processor.AfterGroup);
 			}
 			Query.Callback.ExecutionMode = Processor.ExecutionMode;
+
+			Query.bShouldBatchModifications = Processor.bBatchModifications;
 		}
 
 		inline void PrepareForQueryBinding(FQueryDescription& Query, const FObserver& Observer)
