Fixing warning caused by duplicate test names by using the preferred test registration of TEST_CASE_NAMED to provide unique class names.

#rnx

diff --git a/Engine/Plugins/Experimental/EditorDataStorage/Source/TedsCore/Tests/TypedElementDatabaseScratchBufferTests.cpp b/Engine/Plugins/Experimental/EditorDataStorage/Source/TedsCore/Tests/TypedElementDatabaseScratchBufferTests.cpp
--- a/Engine/Plugins/Experimental/EditorDataStorage/Source/TedsCore/Tests/TypedElementDatabaseScratchBufferTests.cpp
+++ b/Engine/Plugins/Experimental/EditorDataStorage/Source/TedsCore/Tests/TypedElementDatabaseScratchBufferTests.cpp
@@ -24,7 +24,7 @@ static void WaitForFinalization(std::condition_variable& WaitVariable, std::mute
 	REQUIRE(CompletedThreadCount == ThreadCount);
 }
 
-TEST_CASE("Editor::DataStorage::Scratch Buffer - MT (FScratchBuffer)", "[ApplicationContextMask][EngineFilter]")
+TEST_CASE_NAMED(Editor_DataStorage_ScratchBuffer_MultiThreaded_Tests, "Editor::DataStorage::Scratch Buffer - MT (FScratchBuffer)", "[ApplicationContextMask][EngineFilter]")
 {
 	using namespace UE::Editor::DataStorage;
 	SECTION("Stress test without recycling.")
@@ -138,7 +138,7 @@ TEST_CASE("Editor::DataStorage::Scratch Buffer - MT (FScratchBuffer)", "[Applica
 	}
 }
 
-TEST_CASE("Editor::DataStorage::Scratch Buffer (FScratchBuffer)", "[ApplicationContextMask][EngineFilter]")
+TEST_CASE_NAMED(Editor_DataStorage_ScratchBuffer_Tests, "Editor::DataStorage::Scratch Buffer (FScratchBuffer)", "[ApplicationContextMask][EngineFilter]")
 {
 	using namespace UE::Editor::DataStorage;
 	SECTION("Create and destroy buffer")
diff --git a/Engine/Source/Runtime/Engine/Private/Tests/KismetArrayTests.cpp b/Engine/Source/Runtime/Engine/Private/Tests/KismetArrayTests.cpp
--- a/Engine/Source/Runtime/Engine/Private/Tests/KismetArrayTests.cpp
+++ b/Engine/Source/Runtime/Engine/Private/Tests/KismetArrayTests.cpp
@@ -46,7 +46,8 @@ namespace UE::Kismet::Test
 		}
 	}
 
-	TEST_CASE(
+	TEST_CASE_NAMED(
+		KismetArray_SortString_Tests,
 		"UE::KismetArray::SortStrings",
 		"[KismetArray][ClientContext][EditorContext][EngineFilter]")
 	{
@@ -86,7 +87,8 @@ namespace UE::Kismet::Test
 		}
 	}
 
-	TEST_CASE(
+	TEST_CASE_NAMED(
+		KismetArray_SortName_Tests,
 		"UE::KismetArray::SortNames",
 		"[KismetArray][ClientContext][EditorContext][EngineFilter]")
 	{
@@ -117,7 +119,8 @@ namespace UE::Kismet::Test
 		}
 	}
 
-	TEST_CASE(
+	TEST_CASE_NAMED(
+		KismetArray_SortIntegerType_Tests,
 		"UE::KismetArray::SortIntegerTypes",
 		"[KismetArray][ClientContext][EditorContext][EngineFilter]")
 	{
