[TEDS] Cleanup aliases, move core interface classes into namespace, and fix up one class name that wasn't following standards

diff --git a/Engine/Plugins/ChaosVD/Source/ChaosVD/Private/ChaosVDScene.cpp b/Engine/Plugins/ChaosVD/Source/ChaosVD/Private/ChaosVDScene.cpp
--- a/Engine/Plugins/ChaosVD/Source/ChaosVD/Private/ChaosVDScene.cpp
+++ b/Engine/Plugins/ChaosVD/Source/ChaosVD/Private/ChaosVDScene.cpp
@@ -166,12 +166,12 @@ void FChaosVDScene::AddFromCVDWorldTagToActor(AActor* Actor)
 {
 	using namespace UE::Editor::DataStorage;
 	// Add a selection column in TEDS
-	if (const IEditorDataStorageCompatibilityProvider* Compatibility = GetDataStorageFeature<IEditorDataStorageCompatibilityProvider>(CompatibilityFeatureName))
+	if (const ICompatibilityProvider* Compatibility = GetDataStorageFeature<ICompatibilityProvider>(CompatibilityFeatureName))
 	{
 		RowHandle Row = Compatibility->FindRowWithCompatibleObject(Actor);
 		if (Row != InvalidRowHandle)
 		{
-			if (IEditorDataStorageProvider* DataStorage = GetMutableDataStorageFeature<IEditorDataStorageProvider>(StorageFeatureName))
+			if (ICoreProvider* DataStorage = GetMutableDataStorageFeature<ICoreProvider>(StorageFeatureName))
 			{
 				DataStorage->AddColumn<FTypedElementFromCVDWorldTag>(Row);
 				DataStorage->AddColumn<FChaosVDActiveObjectTag>(Row);
diff --git a/Engine/Plugins/ChaosVD/Source/ChaosVD/Private/ChaosVDWorldOutlinerMode.cpp b/Engine/Plugins/ChaosVD/Source/ChaosVD/Private/ChaosVDWorldOutlinerMode.cpp
--- a/Engine/Plugins/ChaosVD/Source/ChaosVD/Private/ChaosVDWorldOutlinerMode.cpp
+++ b/Engine/Plugins/ChaosVD/Source/ChaosVD/Private/ChaosVDWorldOutlinerMode.cpp
@@ -30,7 +30,7 @@ namespace Chaos::VD::SceneOutlinerUtils::Private
 		using namespace UE::Editor::DataStorage;
 		if (FTedsOutlinerTreeItem* TEDSItem = TreeItemRef->CastTo<FTedsOutlinerTreeItem>())
 		{
-			IEditorDataStorageProvider* Storage = GetMutableDataStorageFeature<IEditorDataStorageProvider>(StorageFeatureName);
+			ICoreProvider* Storage = GetMutableDataStorageFeature<ICoreProvider>(StorageFeatureName);
 			if (TEDSItem->IsValid() && Storage)
 			{
 				const FTypedElementExternalObjectColumn* RawObjectColumn = Storage->GetColumn<FTypedElementExternalObjectColumn>(TEDSItem->GetRowHandle());
@@ -55,7 +55,7 @@ namespace Chaos::VD::SceneOutlinerUtils::Private
 
 		if (FTedsOutlinerTreeItem* TEDSItem = TreeItemRef->CastTo<FTedsOutlinerTreeItem>())
 		{
-			IEditorDataStorageProvider* Storage = GetMutableDataStorageFeature<IEditorDataStorageProvider>(StorageFeatureName);
+			ICoreProvider* Storage = GetMutableDataStorageFeature<ICoreProvider>(StorageFeatureName);
 			if (TEDSItem->IsValid() && Storage)
 			{
 				const FTypedElementUObjectColumn* RawObjectColumn = Storage->GetColumn<FTypedElementUObjectColumn>(TEDSItem->GetRowHandle());
@@ -131,7 +131,7 @@ FChaosVDWorldOutlinerMode::FChaosVDWorldOutlinerMode(const UE::Editor::Outliner:
 	
 	RegisterSelectionSetObject(ScenePtr->GetElementSelectionSet());
 
-	if (IEditorDataStorageProvider* DataStorage = TedsOutlinerImpl ? TedsOutlinerImpl->GetStorage() : nullptr)
+	if (UE::Editor::DataStorage::ICoreProvider* DataStorage = TedsOutlinerImpl ? TedsOutlinerImpl->GetStorage() : nullptr)
 	{
 		using namespace UE::Editor::DataStorage;
 		using namespace UE::Editor::DataStorage::Queries;
@@ -175,7 +175,7 @@ FChaosVDWorldOutlinerMode::~FChaosVDWorldOutlinerMode()
 		SceneOutliner->SetEnabled(EmptyConditionalEnabledAttribute);
 	}
 
-	if (IEditorDataStorageProvider* DataStorage = TedsOutlinerImpl ? TedsOutlinerImpl->GetStorage() : nullptr)
+	if (UE::Editor::DataStorage::ICoreProvider* DataStorage = TedsOutlinerImpl ? TedsOutlinerImpl->GetStorage() : nullptr)
 	{
 		DataStorage->UnregisterQuery(CustomRowAdditionQueryHandle);
 		DataStorage->UnregisterQuery(CustomRowRemovalQueryHandle);
@@ -270,7 +270,7 @@ void FChaosVDWorldOutlinerMode::HandlePostSelectionChange(const UTypedElementSel
 		{
 			RowHandle = Particle->GetTedsRowHandle();
 		}
-		else if (const IEditorDataStorageCompatibilityProvider* Compatibility = GetDataStorageFeature<IEditorDataStorageCompatibilityProvider>(CompatibilityFeatureName))
+		else if (const ICompatibilityProvider* Compatibility = GetDataStorageFeature<ICompatibilityProvider>(CompatibilityFeatureName))
 		{
 			RowHandle = Compatibility->FindRowWithCompatibleObject(ActorElementDataUtil::GetActorFromHandle(SelectedParticlesHandles[0]));
 		}
diff --git a/Engine/Plugins/ChaosVD/Source/ChaosVD/Private/ChaosVDWorldOutlinerTab.cpp b/Engine/Plugins/ChaosVD/Source/ChaosVD/Private/ChaosVDWorldOutlinerTab.cpp
--- a/Engine/Plugins/ChaosVD/Source/ChaosVD/Private/ChaosVDWorldOutlinerTab.cpp
+++ b/Engine/Plugins/ChaosVD/Source/ChaosVD/Private/ChaosVDWorldOutlinerTab.cpp
@@ -107,7 +107,7 @@ void FChaosVDWorldOutlinerTab::CreateWorldOutlinerWidget()
 	TSharedRef<ISceneOutliner> TedsOutliner = SNew(SSceneOutliner, InitOptions);
 
 	using namespace UE::Editor::DataStorage;
-	IEditorDataStorageProvider* Storage = GetMutableDataStorageFeature<IEditorDataStorageProvider>(StorageFeatureName);
+	ICoreProvider* Storage = GetMutableDataStorageFeature<ICoreProvider>(StorageFeatureName);
 
 	static QueryHandle ColumnQuery = Storage->RegisterQuery(
 	Select()
diff --git a/Engine/Plugins/ChaosVD/Source/ChaosVD/Private/Components/ChaosVDParticleDataComponent.cpp b/Engine/Plugins/ChaosVD/Source/ChaosVD/Private/Components/ChaosVDParticleDataComponent.cpp
--- a/Engine/Plugins/ChaosVD/Source/ChaosVD/Private/Components/ChaosVDParticleDataComponent.cpp
+++ b/Engine/Plugins/ChaosVD/Source/ChaosVD/Private/Components/ChaosVDParticleDataComponent.cpp
@@ -219,7 +219,8 @@ void UChaosVDParticleDataComponent::ProcessRemovedParticles_AssumessLocked(const
 }
 
 void UChaosVDParticleDataComponent::BatchTEDSSyncFromWorldTag(FTedsHandlesForBatchArray& ParticlesPendingTEDSTagSync,
-	const IEditorDataStorageCompatibilityProvider* CompatibilityStorage, IEditorDataStorageProvider* DataStorage)
+	const UE::Editor::DataStorage::ICompatibilityProvider* CompatibilityStorage,
+	UE::Editor::DataStorage::ICoreProvider* DataStorage)
 {
 	if (CompatibilityStorage && DataStorage)
 	{
@@ -236,7 +237,7 @@ void UChaosVDParticleDataComponent::BatchTEDSSyncFromWorldTag(FTedsHandlesForBat
 	}
 }
 
-void UChaosVDParticleDataComponent::BatchTEDSAddRemoveTags(IEditorDataStorageProvider* EditorDataStorage, FTedsHandlesForBatchArray& ParticlesPendingTEDSTagSync,
+void UChaosVDParticleDataComponent::BatchTEDSAddRemoveTags(UE::Editor::DataStorage::ICoreProvider* EditorDataStorage, FTedsHandlesForBatchArray& ParticlesPendingTEDSTagSync,
 	TConstArrayView<const UScriptStruct*> ColumnsToAdd, TConstArrayView<const UScriptStruct*> ColumnsToRemove)
 {
 	if (!EditorDataStorage)
@@ -268,8 +269,8 @@ void UChaosVDParticleDataComponent::UpdateFromNewSolverStageData(const FChaosVDS
 
 	{
 		using namespace UE::Editor::DataStorage;
-		IEditorDataStorageCompatibilityProvider* Compatibility = GetMutableDataStorageFeature<IEditorDataStorageCompatibilityProvider>(CompatibilityFeatureName);
-		IEditorDataStorageProvider* DataStorage = GetMutableDataStorageFeature<IEditorDataStorageProvider>(StorageFeatureName);
+		ICompatibilityProvider* Compatibility = GetMutableDataStorageFeature<ICompatibilityProvider>(CompatibilityFeatureName);
+		ICoreProvider* DataStorage = GetMutableDataStorageFeature<ICoreProvider>(StorageFeatureName);
 
 		// Make a pre-pass in parallel and process all the data that can be updated in parallel. This step will mostly copy data, and set flags of what needs to be updated in the Game Thread
 		ParallelFor(InSolverFrameStageData.RecordedParticlesData.Num(),[this, &DataStorage, &Compatibility, &SceneParticleToProcessInGT, &InSolverFrameStageData, &InSolverFrameData](int32 ParticleIndex)
@@ -388,8 +389,8 @@ void UChaosVDParticleDataComponent::UpdateFromNewSolverStageData(const FChaosVDS
 	FTedsHandlesForBatchArray ParticlesPendingTEDSActiveStateTagRemoval;
 	ProcessRemovedParticles_AssumessLocked(ScenePtr.ToSharedRef(), InSolverFrameData, InSolverFrameStageData, ParticlesUpdatedIDsForKeyFrameDiff, ParticlesPendingTEDSActiveStateTagRemoval);
 
-	const IEditorDataStorageCompatibilityProvider* EditorDataStorageCompatibility = GetDataStorageFeature<IEditorDataStorageCompatibilityProvider>(CompatibilityFeatureName);
-	IEditorDataStorageProvider* EditorDataStorage = GetMutableDataStorageFeature<IEditorDataStorageProvider>(StorageFeatureName);
+	const ICompatibilityProvider* EditorDataStorageCompatibility = GetDataStorageFeature<ICompatibilityProvider>(CompatibilityFeatureName);
+	ICoreProvider* EditorDataStorage = GetMutableDataStorageFeature<ICoreProvider>(StorageFeatureName);
 	
 	BatchTEDSSyncFromWorldTag(ParticlesPendingTEDSSyncFromWorldTagSync, EditorDataStorageCompatibility, EditorDataStorage);
 	BatchTEDSAddRemoveTags(EditorDataStorage, ParticlesPendingTEDSActiveStateTagAddition, { FChaosVDActiveObjectTag::StaticStruct() } , {});
@@ -493,8 +494,8 @@ void UChaosVDParticleDataComponent::BatchApplySolverVisibilityToParticle(bool bN
 		ParticleWithIDPair.Value->UpdateGeometryComponentsVisibility();
 	}
 
-	const IEditorDataStorageCompatibilityProvider* EditorDataStorageCompatibility = GetDataStorageFeature<IEditorDataStorageCompatibilityProvider>(CompatibilityFeatureName);
-	IEditorDataStorageProvider* EditorDataStorage = GetMutableDataStorageFeature<IEditorDataStorageProvider>(StorageFeatureName);
+	const ICompatibilityProvider* EditorDataStorageCompatibility = GetDataStorageFeature<ICompatibilityProvider>(CompatibilityFeatureName);
+	ICoreProvider* EditorDataStorage = GetMutableDataStorageFeature<ICoreProvider>(StorageFeatureName);
 	BatchTEDSSyncFromWorldTag(ParticlesPendingTEDSTagSync, EditorDataStorageCompatibility, EditorDataStorage);
 	
 	if (TSharedPtr<FChaosVDScene> ScenePtr = SceneWeakPtr.Pin())
diff --git a/Engine/Plugins/ChaosVD/Source/ChaosVD/Private/Components/ChaosVDParticleDataComponent.h b/Engine/Plugins/ChaosVD/Source/ChaosVD/Private/Components/ChaosVDParticleDataComponent.h
--- a/Engine/Plugins/ChaosVD/Source/ChaosVD/Private/Components/ChaosVDParticleDataComponent.h
+++ b/Engine/Plugins/ChaosVD/Source/ChaosVD/Private/Components/ChaosVDParticleDataComponent.h
@@ -82,9 +82,12 @@ protected:
 
 	int32 GetMaxElementCountForBatchTEDSUpdate_AssumesLocked() const;
 
-	void BatchTEDSSyncFromWorldTag(FTedsHandlesForBatchArray& ParticlesPendingTEDSTagSync, const IEditorDataStorageCompatibilityProvider* CompatibilityStorage, IEditorDataStorageProvider* DataStorage);
+	void BatchTEDSSyncFromWorldTag(FTedsHandlesForBatchArray& ParticlesPendingTEDSTagSync,
+		const UE::Editor::DataStorage::ICompatibilityProvider* CompatibilityStorage,
+		UE::Editor::DataStorage::ICoreProvider* DataStorage);
 
-	void BatchTEDSAddRemoveTags(IEditorDataStorageProvider* EditorDataStorage, FTedsHandlesForBatchArray& ParticlesPendingTEDSTagSync, TConstArrayView<const UScriptStruct*> ColumnsToAdd, TConstArrayView<const UScriptStruct*> ColumnsToRemove);
+	void BatchTEDSAddRemoveTags(UE::Editor::DataStorage::ICoreProvider* EditorDataStorage, FTedsHandlesForBatchArray& ParticlesPendingTEDSTagSync,
+		TConstArrayView<const UScriptStruct*> ColumnsToAdd, TConstArrayView<const UScriptStruct*> ColumnsToRemove);
 
 	void ProcessRemovedParticles_AssumessLocked(const TSharedRef<FChaosVDScene>& SceneRef, const FChaosVDSolverFrameData& InSolverFrameData,
 		const FChaosVDFrameStageData& InSolverFrameStageData, const TSet<int32>& ParticlesUpdatedIDsForKeyFrameDiff, FTedsHandlesForBatchArray& OutRemovedTedsHandles);
diff --git a/Engine/Plugins/ChaosVD/Source/ChaosVD/Private/TEDS/ChaosVDCustomIconDataStorageFactory.cpp b/Engine/Plugins/ChaosVD/Source/ChaosVD/Private/TEDS/ChaosVDCustomIconDataStorageFactory.cpp
--- a/Engine/Plugins/ChaosVD/Source/ChaosVD/Private/TEDS/ChaosVDCustomIconDataStorageFactory.cpp
+++ b/Engine/Plugins/ChaosVD/Source/ChaosVD/Private/TEDS/ChaosVDCustomIconDataStorageFactory.cpp
@@ -10,7 +10,7 @@
 #include "Elements/Columns/TypedElementTypeInfoColumns.h"
 #include "Elements/Framework/TypedElementQueryBuilder.h"
 
-void UChaosVDCustomIconDataStorageFactory::RegisterQueries(IEditorDataStorageProvider& DataStorage)
+void UChaosVDCustomIconDataStorageFactory::RegisterQueries(UE::Editor::DataStorage::ICoreProvider& DataStorage)
 {
 	Super::RegisterQueries(DataStorage);
 
diff --git a/Engine/Plugins/ChaosVD/Source/ChaosVD/Private/TEDS/ChaosVDCustomIconDataStorageFactory.h b/Engine/Plugins/ChaosVD/Source/ChaosVD/Private/TEDS/ChaosVDCustomIconDataStorageFactory.h
--- a/Engine/Plugins/ChaosVD/Source/ChaosVD/Private/TEDS/ChaosVDCustomIconDataStorageFactory.h
+++ b/Engine/Plugins/ChaosVD/Source/ChaosVD/Private/TEDS/ChaosVDCustomIconDataStorageFactory.h
@@ -15,5 +15,5 @@ class UChaosVDCustomIconDataStorageFactory : public UEditorDataStorageFactory
 	
 	virtual ~UChaosVDCustomIconDataStorageFactory() override = default;
 
-	virtual void RegisterQueries(IEditorDataStorageProvider& DataStorage) override;
+	virtual void RegisterQueries(UE::Editor::DataStorage::ICoreProvider& DataStorage) override;
 };
diff --git a/Engine/Plugins/ChaosVD/Source/ChaosVD/Private/TEDS/ChaosVDDataStorageVisibilityQueries.cpp b/Engine/Plugins/ChaosVD/Source/ChaosVD/Private/TEDS/ChaosVDDataStorageVisibilityQueries.cpp
--- a/Engine/Plugins/ChaosVD/Source/ChaosVD/Private/TEDS/ChaosVDDataStorageVisibilityQueries.cpp
+++ b/Engine/Plugins/ChaosVD/Source/ChaosVD/Private/TEDS/ChaosVDDataStorageVisibilityQueries.cpp
@@ -15,7 +15,7 @@
 struct FTypedElementScriptStructTypeInfoColumn;
 
 
-void UChaosVDDataStorageVisibilityQueries::RegisterQueries(IEditorDataStorageProvider& DataStorage)
+void UChaosVDDataStorageVisibilityQueries::RegisterQueries(UE::Editor::DataStorage::ICoreProvider& DataStorage)
 {
 	Super::RegisterQueries(DataStorage);
 	
@@ -24,7 +24,7 @@ void UChaosVDDataStorageVisibilityQueries::RegisterQueries(IEditorDataStoragePro
 	RegisterVisibilityColumnToParticleQuery(DataStorage);
 }
 
-void UChaosVDDataStorageVisibilityQueries::RegisterParticleAddVisibilityColumn(IEditorDataStorageProvider& DataStorage) const
+void UChaosVDDataStorageVisibilityQueries::RegisterParticleAddVisibilityColumn(UE::Editor::DataStorage::ICoreProvider& DataStorage) const
 {
 	using namespace UE::Editor::DataStorage::Queries;
 
@@ -48,7 +48,7 @@ void UChaosVDDataStorageVisibilityQueries::RegisterParticleAddVisibilityColumn(I
 	);
 }
 
-void UChaosVDDataStorageVisibilityQueries::RegisterParticleVisibilityToColumnQuery(IEditorDataStorageProvider& DataStorage) const
+void UChaosVDDataStorageVisibilityQueries::RegisterParticleVisibilityToColumnQuery(UE::Editor::DataStorage::ICoreProvider& DataStorage) const
 {
 	using namespace UE::Editor::DataStorage::Queries;
 
@@ -71,7 +71,7 @@ void UChaosVDDataStorageVisibilityQueries::RegisterParticleVisibilityToColumnQue
 	);
 }
 
-void UChaosVDDataStorageVisibilityQueries::RegisterVisibilityColumnToParticleQuery(IEditorDataStorageProvider& DataStorage) const
+void UChaosVDDataStorageVisibilityQueries::RegisterVisibilityColumnToParticleQuery(UE::Editor::DataStorage::ICoreProvider& DataStorage) const
 {
 	using namespace UE::Editor::DataStorage::Queries;
 
diff --git a/Engine/Plugins/ChaosVD/Source/ChaosVD/Private/TEDS/ChaosVDDataStorageVisibilityQueries.h b/Engine/Plugins/ChaosVD/Source/ChaosVD/Private/TEDS/ChaosVDDataStorageVisibilityQueries.h
--- a/Engine/Plugins/ChaosVD/Source/ChaosVD/Private/TEDS/ChaosVDDataStorageVisibilityQueries.h
+++ b/Engine/Plugins/ChaosVD/Source/ChaosVD/Private/TEDS/ChaosVDDataStorageVisibilityQueries.h
@@ -13,22 +13,22 @@ class UChaosVDDataStorageVisibilityQueries : public UEditorDataStorageFactory
 public:
 	~UChaosVDDataStorageVisibilityQueries() override = default;
 
-	void RegisterQueries(IEditorDataStorageProvider& DataStorage) override;
+	void RegisterQueries(UE::Editor::DataStorage::ICoreProvider& DataStorage) override;
 
 private:
 
 	/**
 	 * Adds the Visibility columns to new actors that do not have one already.
 	 */
-	void RegisterParticleAddVisibilityColumn(IEditorDataStorageProvider& DataStorage) const;
+	void RegisterParticleAddVisibilityColumn(UE::Editor::DataStorage::ICoreProvider& DataStorage) const;
 
 	/**
 	 * Takes the visibility set on an actor and copies it to the Data Storage if they differ.
 	 */
-	void RegisterParticleVisibilityToColumnQuery(IEditorDataStorageProvider& DataStorage) const;
+	void RegisterParticleVisibilityToColumnQuery(UE::Editor::DataStorage::ICoreProvider& DataStorage) const;
 	/**
 	 * Takes the visibility stored in the Data Storage and copies it to the actor's visibility if the FTypedElementSyncBackToWorldTag
 	 * has been set and the visibility differ.
 	 */
-	void RegisterVisibilityColumnToParticleQuery(IEditorDataStorageProvider& DataStorage) const;
+	void RegisterVisibilityColumnToParticleQuery(UE::Editor::DataStorage::ICoreProvider& DataStorage) const;
 };
diff --git a/Engine/Plugins/ChaosVD/Source/ChaosVD/Private/TEDS/ChaosVDParentDataStorageFactory.cpp b/Engine/Plugins/ChaosVD/Source/ChaosVD/Private/TEDS/ChaosVDParentDataStorageFactory.cpp
--- a/Engine/Plugins/ChaosVD/Source/ChaosVD/Private/TEDS/ChaosVDParentDataStorageFactory.cpp
+++ b/Engine/Plugins/ChaosVD/Source/ChaosVD/Private/TEDS/ChaosVDParentDataStorageFactory.cpp
@@ -21,7 +21,7 @@ namespace Chaos::VD::ParentDataStorageUtils::Private
 {
 	using namespace UE::Editor::DataStorage::Queries;
 
-	void UpdateParentData(IEditorDataStorageProvider& DataStorage, FChaosVDTableRowParentColumn& RowParentColumn, RowHandle ItemRowHandle, RowHandle NewParentRowHandle)
+	void UpdateParentData(ICoreProvider& DataStorage, FChaosVDTableRowParentColumn& RowParentColumn, RowHandle ItemRowHandle, RowHandle NewParentRowHandle)
 	{
 		if (DataStorage.IsRowAvailable(NewParentRowHandle))
 		{
@@ -43,7 +43,7 @@ namespace Chaos::VD::ParentDataStorageUtils::Private
 	};
 }
 
-void UChaosVDParentDataStorageFactory::RegisterQueries(IEditorDataStorageProvider& DataStorage)
+void UChaosVDParentDataStorageFactory::RegisterQueries(UE::Editor::DataStorage::ICoreProvider& DataStorage)
 {
 	Super::RegisterQueries(DataStorage);
 	
@@ -62,7 +62,7 @@ void UChaosVDParentDataStorageFactory::RegisterQueries(IEditorDataStorageProvide
 	RegisterUpdateOrRemoveParentColumn(DataStorage);
 }
 
-void UChaosVDParentDataStorageFactory::RegisterAddParentColumn(IEditorDataStorageProvider& DataStorage) const
+void UChaosVDParentDataStorageFactory::RegisterAddParentColumn(UE::Editor::DataStorage::ICoreProvider& DataStorage) const
 {
 	using namespace UE::Editor::DataStorage::Queries;
 
@@ -96,7 +96,7 @@ void UChaosVDParentDataStorageFactory::RegisterAddParentColumn(IEditorDataStorag
 		.Compile());
 }
 
-void UChaosVDParentDataStorageFactory::RegisterUpdateOrRemoveParentColumn(IEditorDataStorageProvider& DataStorage) const
+void UChaosVDParentDataStorageFactory::RegisterUpdateOrRemoveParentColumn(UE::Editor::DataStorage::ICoreProvider& DataStorage) const
 {
 	using namespace UE::Editor::DataStorage::Queries;
 	using namespace Chaos::VD::ParentDataStorageUtils::Private;
diff --git a/Engine/Plugins/ChaosVD/Source/ChaosVD/Private/TEDS/ChaosVDParentDataStorageFactory.h b/Engine/Plugins/ChaosVD/Source/ChaosVD/Private/TEDS/ChaosVDParentDataStorageFactory.h
--- a/Engine/Plugins/ChaosVD/Source/ChaosVD/Private/TEDS/ChaosVDParentDataStorageFactory.h
+++ b/Engine/Plugins/ChaosVD/Source/ChaosVD/Private/TEDS/ChaosVDParentDataStorageFactory.h
@@ -30,19 +30,19 @@ class UChaosVDParentDataStorageFactory : public UEditorDataStorageFactory
 public:
 	virtual ~UChaosVDParentDataStorageFactory() override = default;
 
-	virtual void RegisterQueries(IEditorDataStorageProvider& DataStorage) override;
+	virtual void RegisterQueries(UE::Editor::DataStorage::ICoreProvider& DataStorage) override;
 
 private:
 	/**
 	 * Checks rows with actors that don't have a parent column yet if one needs to be added whenever
 	 * the row is marked for updates.
 	 */
-	void RegisterAddParentColumn(IEditorDataStorageProvider& DataStorage) const;
+	void RegisterAddParentColumn(UE::Editor::DataStorage::ICoreProvider& DataStorage) const;
 	/**
 	 * Updates the parent column with the parent from the actor or removes it if there's no parent associated
 	 * with the actor anymore.
 	 */
-	void RegisterUpdateOrRemoveParentColumn(IEditorDataStorageProvider& DataStorage) const;
+	void RegisterUpdateOrRemoveParentColumn(UE::Editor::DataStorage::ICoreProvider& DataStorage) const;
 
 	TSharedPtr<FChaosVDBaseSceneObject> DefaultRootObjectForCVDActors;
 };
diff --git a/Engine/Plugins/ChaosVD/Source/ChaosVD/Private/TEDS/ChaosVDParticleEditorDataFactory.cpp b/Engine/Plugins/ChaosVD/Source/ChaosVD/Private/TEDS/ChaosVDParticleEditorDataFactory.cpp
--- a/Engine/Plugins/ChaosVD/Source/ChaosVD/Private/TEDS/ChaosVDParticleEditorDataFactory.cpp
+++ b/Engine/Plugins/ChaosVD/Source/ChaosVD/Private/TEDS/ChaosVDParticleEditorDataFactory.cpp
@@ -19,7 +19,7 @@
 #include "Elements/Interfaces/TypedElementDataStorageInterface.h"
 #include "Elements/Interfaces/TypedElementSelectionInterface.h"
 
-void UChaosVDParticleEditorDataFactory::RegisterTables(IEditorDataStorageProvider& DataStorage, IEditorDataStorageCompatibilityProvider& DataStorageCompatibility)
+void UChaosVDParticleEditorDataFactory::RegisterTables(UE::Editor::DataStorage::ICoreProvider& DataStorage, UE::Editor::DataStorage::ICompatibilityProvider& DataStorageCompatibility)
 {
 	Super::RegisterTables(DataStorage, DataStorageCompatibility);
 
@@ -48,7 +48,7 @@ void UChaosVDParticleEditorDataFactory::RegisterTables(IEditorDataStorageProvide
 	DataStorageCompatibility.RegisterTypeTableAssociation(AChaosVDSolverInfoActor::StaticClass(), ActorTable);
 }
 
-void UChaosVDParticleEditorDataFactory::RegisterQueries(IEditorDataStorageProvider& DataStorage)
+void UChaosVDParticleEditorDataFactory::RegisterQueries(UE::Editor::DataStorage::ICoreProvider& DataStorage)
 {
 	Super::RegisterQueries(DataStorage);
 
diff --git a/Engine/Plugins/ChaosVD/Source/ChaosVD/Private/TEDS/ChaosVDParticleEditorDataFactory.h b/Engine/Plugins/ChaosVD/Source/ChaosVD/Private/TEDS/ChaosVDParticleEditorDataFactory.h
--- a/Engine/Plugins/ChaosVD/Source/ChaosVD/Private/TEDS/ChaosVDParticleEditorDataFactory.h
+++ b/Engine/Plugins/ChaosVD/Source/ChaosVD/Private/TEDS/ChaosVDParticleEditorDataFactory.h
@@ -45,6 +45,6 @@ class UChaosVDParticleEditorDataFactory : public UEditorDataStorageFactory
 public:
 	virtual ~UChaosVDParticleEditorDataFactory() override = default;
 
-	virtual void RegisterTables(IEditorDataStorageProvider& DataStorage, IEditorDataStorageCompatibilityProvider& DataStorageCompatibility) override;
-	virtual void RegisterQueries(IEditorDataStorageProvider& DataStorage) override;
+	virtual void RegisterTables(UE::Editor::DataStorage::ICoreProvider& DataStorage, UE::Editor::DataStorage::ICompatibilityProvider& DataStorageCompatibility) override;
+	virtual void RegisterQueries(UE::Editor::DataStorage::ICoreProvider& DataStorage) override;
 };
diff --git a/Engine/Plugins/ChaosVD/Source/ChaosVD/Private/TEDS/ChaosVDTedsUtils.h b/Engine/Plugins/ChaosVD/Source/ChaosVD/Private/TEDS/ChaosVDTedsUtils.h
--- a/Engine/Plugins/ChaosVD/Source/ChaosVD/Private/TEDS/ChaosVDTedsUtils.h
+++ b/Engine/Plugins/ChaosVD/Source/ChaosVD/Private/TEDS/ChaosVDTedsUtils.h
@@ -13,10 +13,10 @@ namespace Chaos::VD::TedsUtils
 	{
 		using namespace UE::Editor::DataStorage;
 
-		if (const IEditorDataStorageCompatibilityProvider* Compatibility = GetDataStorageFeature<IEditorDataStorageCompatibilityProvider>(CompatibilityFeatureName))
+		if (const ICompatibilityProvider* Compatibility = GetDataStorageFeature<ICompatibilityProvider>(CompatibilityFeatureName))
 		{
 			RowHandle Row = Compatibility->FindRowWithCompatibleObject(Object);
-			if (IEditorDataStorageProvider* DataStorage = GetMutableDataStorageFeature<IEditorDataStorageProvider>(StorageFeatureName))
+			if (ICoreProvider* DataStorage = GetMutableDataStorageFeature<ICoreProvider>(StorageFeatureName))
 			{
 				if (DataStorage->IsRowAvailable(Row))
 				{
@@ -31,7 +31,7 @@ namespace Chaos::VD::TedsUtils
 	{
 		using namespace UE::Editor::DataStorage;
 
-		if (IEditorDataStorageProvider* DataStorage = GetMutableDataStorageFeature<IEditorDataStorageProvider>(StorageFeatureName))
+		if (ICoreProvider* DataStorage = GetMutableDataStorageFeature<ICoreProvider>(StorageFeatureName))
 		{
 			if (DataStorage->IsRowAvailable(Handle))
 			{
@@ -45,7 +45,7 @@ namespace Chaos::VD::TedsUtils
 	{
 		using namespace UE::Editor::DataStorage;
 
-		if (IEditorDataStorageProvider* DataStorage = GetMutableDataStorageFeature<IEditorDataStorageProvider>(StorageFeatureName))
+		if (ICoreProvider* DataStorage = GetMutableDataStorageFeature<ICoreProvider>(StorageFeatureName))
 		{
 			DataStorage->BatchAddRemoveColumns(Rows, { ColumnType::StaticStruct(), 1}, {});
 		}
@@ -56,7 +56,7 @@ namespace Chaos::VD::TedsUtils
 	{
 		using namespace UE::Editor::DataStorage;
 
-		if (IEditorDataStorageProvider* DataStorage = GetMutableDataStorageFeature<IEditorDataStorageProvider>(StorageFeatureName))
+		if (ICoreProvider* DataStorage = GetMutableDataStorageFeature<ICoreProvider>(StorageFeatureName))
 		{
 			if (DataStorage->IsRowAvailable(Handle))
 			{
@@ -70,11 +70,11 @@ namespace Chaos::VD::TedsUtils
 	{
 		using namespace UE::Editor::DataStorage;
 
-		if (const IEditorDataStorageCompatibilityProvider* Compatibility = GetDataStorageFeature<IEditorDataStorageCompatibilityProvider>(CompatibilityFeatureName))
+		if (const ICompatibilityProvider* Compatibility = GetDataStorageFeature<ICompatibilityProvider>(CompatibilityFeatureName))
 		{
 			RowHandle Row = Compatibility->FindRowWithCompatibleObject(Object);
 
-			if (IEditorDataStorageProvider* DataStorage = GetMutableDataStorageFeature<IEditorDataStorageProvider>(StorageFeatureName))
+			if (ICoreProvider* DataStorage = GetMutableDataStorageFeature<ICoreProvider>(StorageFeatureName))
 			{
 				if (DataStorage->IsRowAvailable(Row))
 				{
@@ -85,7 +85,9 @@ namespace Chaos::VD::TedsUtils
 	}
 
 	template<typename ObjectType>
-	static UE::Editor::DataStorage::RowHandle AddObjectToDataStorage(ObjectType* Object, IEditorDataStorageProvider* DataStorageInterface, IEditorDataStorageCompatibilityProvider* DataStorageCompatibility)
+	static UE::Editor::DataStorage::RowHandle AddObjectToDataStorage(ObjectType* Object,
+		UE::Editor::DataStorage::ICoreProvider* DataStorageInterface,
+		UE::Editor::DataStorage::ICompatibilityProvider* DataStorageCompatibility)
 	{
 		using namespace UE::Editor::DataStorage;
 		if (DataStorageCompatibility != nullptr && DataStorageInterface != nullptr)
@@ -100,14 +102,14 @@ namespace Chaos::VD::TedsUtils
 	static UE::Editor::DataStorage::RowHandle AddObjectToDataStorage(ObjectType* Object)
 	{
 		using namespace UE::Editor::DataStorage;
-		IEditorDataStorageProvider* DataStorageInterface = GetMutableDataStorageFeature<IEditorDataStorageProvider>(StorageFeatureName);
-		IEditorDataStorageCompatibilityProvider* DataStorageCompatibility = GetMutableDataStorageFeature<IEditorDataStorageCompatibilityProvider>(CompatibilityFeatureName);
+		ICoreProvider* DataStorageInterface = GetMutableDataStorageFeature<ICoreProvider>(StorageFeatureName);
+		ICompatibilityProvider* DataStorageCompatibility = GetMutableDataStorageFeature<ICompatibilityProvider>(CompatibilityFeatureName);
 
 		return AddObjectToDataStorage(Object, DataStorageInterface, DataStorageCompatibility);	
 	}
 
 	template<typename ObjectType>
-	static void RemoveObjectToFromDataStorage(ObjectType* Object, IEditorDataStorageCompatibilityProvider* DataStorageCompatibility)
+	static void RemoveObjectToFromDataStorage(ObjectType* Object, UE::Editor::DataStorage::ICompatibilityProvider* DataStorageCompatibility)
 	{
 		if (DataStorageCompatibility != nullptr)
 		{
@@ -119,7 +121,7 @@ namespace Chaos::VD::TedsUtils
 	static void RemoveObjectToFromDataStorage(ObjectType* Object)
 	{
 		using namespace UE::Editor::DataStorage;
-		IEditorDataStorageCompatibilityProvider* DataStorageCompatibility = GetMutableDataStorageFeature<IEditorDataStorageCompatibilityProvider>(CompatibilityFeatureName);
+		ICompatibilityProvider* DataStorageCompatibility = GetMutableDataStorageFeature<ICompatibilityProvider>(CompatibilityFeatureName);
 		RemoveObjectToFromDataStorage(Object, DataStorageCompatibility);
 	}
 }
diff --git a/Engine/Plugins/ChaosVD/Source/ChaosVD/Public/ChaosVDSceneParticle.h b/Engine/Plugins/ChaosVD/Source/ChaosVD/Public/ChaosVDSceneParticle.h
--- a/Engine/Plugins/ChaosVD/Source/ChaosVD/Public/ChaosVDSceneParticle.h
+++ b/Engine/Plugins/ChaosVD/Source/ChaosVD/Public/ChaosVDSceneParticle.h
@@ -12,8 +12,11 @@
 
 #include "ChaosVDSceneParticle.generated.h"
 
-class IEditorDataStorageCompatibilityProvider;
-class IEditorDataStorageProvider;
+namespace UE::Editor::DataStorage
+{
+	class ICompatibilityProvider;
+	class IEditorDataStorageProvider;
+} // UE::Editor::DataStorage
 
 UENUM()
 enum class EChaosVDParticleVisibilityUpdateFlags
diff --git a/Engine/Plugins/Experimental/EditorDataStorage/Source/TedsCore/Private/Commands/EditorDataStorageCompatibilityCommands.cpp b/Engine/Plugins/Experimental/EditorDataStorage/Source/TedsCore/Private/Commands/EditorDataStorageCompatibilityCommands.cpp
--- a/Engine/Plugins/Experimental/EditorDataStorage/Source/TedsCore/Private/Commands/EditorDataStorageCompatibilityCommands.cpp
+++ b/Engine/Plugins/Experimental/EditorDataStorage/Source/TedsCore/Private/Commands/EditorDataStorageCompatibilityCommands.cpp
@@ -224,7 +224,7 @@ namespace UE::Editor::DataStorage
 	// FPrepareCommands
 	//
 
-	FPrepareCommands::FPrepareCommands(IEditorDataStorageProvider& InStorage, UEditorDataStorageCompatibility& InStorageCompat,
+	FPrepareCommands::FPrepareCommands(ICoreProvider& InStorage, UEditorDataStorageCompatibility& InStorageCompat,
 		CompatibilityCommandBuffer::FCollection& InCommands)
 		: Storage(InStorage)
 		, StorageCompat(InStorageCompat)
@@ -303,7 +303,7 @@ namespace UE::Editor::DataStorage
 		}
 	}
 
-	void FPrepareCommands::RunPreparation(IEditorDataStorageProvider& Storage, UEditorDataStorageCompatibility& StorageCompat,
+	void FPrepareCommands::RunPreparation(ICoreProvider& Storage, UEditorDataStorageCompatibility& StorageCompat,
 		CompatibilityCommandBuffer::FCollection& Commands)
 	{
 		FPrepareCommands PrepareVisitor(Storage, StorageCompat, Commands);
@@ -489,7 +489,7 @@ namespace UE::Editor::DataStorage
 	// FCommandProcessor
 	//
 
-	FCommandProcessor::FCommandProcessor(IEditorDataStorageProvider& InStorage, UEditorDataStorageCompatibility& InStorageCompatibility)
+	FCommandProcessor::FCommandProcessor(ICoreProvider& InStorage, UEditorDataStorageCompatibility& InStorageCompatibility)
 		: Storage(InStorage)
 		, StorageCompatibility(InStorageCompatibility)
 		, MementoSystem(InStorageCompatibility.Environment->GetMementoSystem())
diff --git a/Engine/Plugins/Experimental/EditorDataStorage/Source/TedsCore/Private/Commands/EditorDataStorageCompatibilityCommands.h b/Engine/Plugins/Experimental/EditorDataStorage/Source/TedsCore/Private/Commands/EditorDataStorageCompatibilityCommands.h
--- a/Engine/Plugins/Experimental/EditorDataStorage/Source/TedsCore/Private/Commands/EditorDataStorageCompatibilityCommands.h
+++ b/Engine/Plugins/Experimental/EditorDataStorage/Source/TedsCore/Private/Commands/EditorDataStorageCompatibilityCommands.h
@@ -9,7 +9,6 @@
 #include "UObject/WeakObjectPtr.h"
 #include "UObject/WeakObjectPtrTemplates.h"
 
-class IEditorDataStorageProvider;
 class UScriptStruct;
 class UEditorDataStorageCompatibility;
 class UObject;
@@ -18,6 +17,7 @@ namespace UE::Editor::DataStorage
 {
 	class FScratchBuffer;
 	class FMementoSystem;
+	class ICoreProvider;
 
 	enum class EObjectType : uint8
 	{
@@ -27,7 +27,7 @@ namespace UE::Editor::DataStorage
 
 	/**
 	 * Objects with type info defined in either UScriptStruct or UClass can be stored into TEDS via the
-	 * IEditorDataStorageCompatibilityProvider
+	 * ICompatibilityProvider
 	 * This is a discriminated union which aids with callbacks made when objects are added
 	 */
 	struct FObjectTypeInfo
@@ -228,7 +228,7 @@ namespace UE::Editor::DataStorage
 	/** Prepares each command for further processing, e.g. resolving the target table. */
 	struct FPrepareCommands final
 	{
-		FPrepareCommands(IEditorDataStorageProvider& InStorage, UEditorDataStorageCompatibility& InStorageCompat,
+		FPrepareCommands(ICoreProvider& InStorage, UEditorDataStorageCompatibility& InStorageCompat,
 			CompatibilityCommandBuffer::FCollection& InCommands);
 
 		template<typename T>
@@ -241,12 +241,12 @@ namespace UE::Editor::DataStorage
 		void operator()(FRemoveInteractiveSyncFromWorldTag& Command);
 		void operator()(FAddSyncFromWorldTag& Command);
 
-		IEditorDataStorageProvider& Storage;
+		ICoreProvider& Storage;
 		UEditorDataStorageCompatibility& StorageCompat;
 		CompatibilityCommandBuffer::FCollection& Commands;
 		int32 CurrentIndex = 0;
 
-		static void RunPreparation(IEditorDataStorageProvider& Storage, UEditorDataStorageCompatibility& StorageCompat, 
+		static void RunPreparation(ICoreProvider& Storage, UEditorDataStorageCompatibility& StorageCompat, 
 			CompatibilityCommandBuffer::FCollection& Commands);
 	};
 
@@ -299,7 +299,7 @@ namespace UE::Editor::DataStorage
 	/** Executes the commands in the command buffer for TEDS Compatibility. */
 	struct FCommandProcessor final
 	{
-		FCommandProcessor(IEditorDataStorageProvider& InStorage, UEditorDataStorageCompatibility& InStorageCompatibility);
+		FCommandProcessor(ICoreProvider& InStorage, UEditorDataStorageCompatibility& InStorageCompatibility);
 
 		void SetupRow(RowHandle Row, UObject* Object);
 		void SetupRow(RowHandle Row, void* Object, TWeakObjectPtr<UScriptStruct> TypeInfo);
@@ -323,7 +323,7 @@ namespace UE::Editor::DataStorage
 		void operator()(FRemoveInteractiveSyncFromWorldTag& Command);
 		void operator()(FAddSyncFromWorldTag& Command);
 
-		IEditorDataStorageProvider& Storage;
+		ICoreProvider& Storage;
 		UEditorDataStorageCompatibility& StorageCompatibility;
 		FMementoSystem& MementoSystem;
 	};
diff --git a/Engine/Plugins/Experimental/EditorDataStorage/Source/TedsCore/Private/Compatibility/TedsTypedElementActorHandleFactory.cpp b/Engine/Plugins/Experimental/EditorDataStorage/Source/TedsCore/Private/Compatibility/TedsTypedElementActorHandleFactory.cpp
--- a/Engine/Plugins/Experimental/EditorDataStorage/Source/TedsCore/Private/Compatibility/TedsTypedElementActorHandleFactory.cpp
+++ b/Engine/Plugins/Experimental/EditorDataStorage/Source/TedsCore/Private/Compatibility/TedsTypedElementActorHandleFactory.cpp
@@ -12,20 +12,20 @@
 #include "Elements/Interfaces/TypedElementDataStorageInterface.h"
 #include "GameFramework/Actor.h"
 
-void UTypedElementActorHandleDataStorageFactory::PreRegister(IEditorDataStorageProvider& DataStorage)
+void UTypedElementActorHandleDataStorageFactory::PreRegister(UE::Editor::DataStorage::ICoreProvider& DataStorage)
 {
 	Super::PreRegister(DataStorage);
 	
 	BridgeEnableDelegateHandle = UE::Editor::DataStorage::Compatibility::OnTypedElementBridgeEnabled().AddUObject(this, &UTypedElementActorHandleDataStorageFactory::HandleBridgeEnabled);
 }
 
-void UTypedElementActorHandleDataStorageFactory::PreShutdown(IEditorDataStorageProvider& DataStorage)
+void UTypedElementActorHandleDataStorageFactory::PreShutdown(UE::Editor::DataStorage::ICoreProvider& DataStorage)
 {
 	UE::Editor::DataStorage::Compatibility::OnTypedElementBridgeEnabled().Remove(BridgeEnableDelegateHandle);
 	BridgeEnableDelegateHandle.Reset();
 }
 
-void UTypedElementActorHandleDataStorageFactory::RegisterQueries(IEditorDataStorageProvider& DataStorage)
+void UTypedElementActorHandleDataStorageFactory::RegisterQueries(UE::Editor::DataStorage::ICoreProvider& DataStorage)
 {
 	Super::RegisterQueries(DataStorage);
 
@@ -43,7 +43,7 @@ void UTypedElementActorHandleDataStorageFactory::RegisterQueries(IEditorDataStor
 	.Compile());
 }
 
-void UTypedElementActorHandleDataStorageFactory::RegisterQuery_ActorHandlePopulate(IEditorDataStorageProvider& DataStorage)
+void UTypedElementActorHandleDataStorageFactory::RegisterQuery_ActorHandlePopulate(UE::Editor::DataStorage::ICoreProvider& DataStorage)
 {
 	using namespace UE::Editor::DataStorage::Queries;
 
@@ -76,7 +76,7 @@ void UTypedElementActorHandleDataStorageFactory::HandleBridgeEnabled(bool bEnabl
 {
 	using namespace UE::Editor::DataStorage;
 
-	IEditorDataStorageProvider* DataStorage = GetMutableDataStorageFeature<IEditorDataStorageProvider>(StorageFeatureName);
+	ICoreProvider* DataStorage = GetMutableDataStorageFeature<ICoreProvider>(StorageFeatureName);
 
 	if (bEnabled)
 	{
diff --git a/Engine/Plugins/Experimental/EditorDataStorage/Source/TedsCore/Private/Compatibility/TedsTypedElementActorHandleFactory.h b/Engine/Plugins/Experimental/EditorDataStorage/Source/TedsCore/Private/Compatibility/TedsTypedElementActorHandleFactory.h
--- a/Engine/Plugins/Experimental/EditorDataStorage/Source/TedsCore/Private/Compatibility/TedsTypedElementActorHandleFactory.h
+++ b/Engine/Plugins/Experimental/EditorDataStorage/Source/TedsCore/Private/Compatibility/TedsTypedElementActorHandleFactory.h
@@ -21,11 +21,11 @@ class UTypedElementActorHandleDataStorageFactory : public UEditorDataStorageFact
 	GENERATED_BODY()
 public:
 
-	virtual void PreRegister(IEditorDataStorageProvider& DataStorage) override;
-	virtual void RegisterQueries(IEditorDataStorageProvider& DataStorage) override;
-	virtual void PreShutdown(IEditorDataStorageProvider& DataStorage) override;
+	virtual void PreRegister(UE::Editor::DataStorage::ICoreProvider& DataStorage) override;
+	virtual void RegisterQueries(UE::Editor::DataStorage::ICoreProvider& DataStorage) override;
+	virtual void PreShutdown(UE::Editor::DataStorage::ICoreProvider& DataStorage) override;
 private:
-	virtual void RegisterQuery_ActorHandlePopulate(IEditorDataStorageProvider& DataStorage);
+	virtual void RegisterQuery_ActorHandlePopulate(UE::Editor::DataStorage::ICoreProvider& DataStorage);
 	void HandleBridgeEnabled(bool bEnabled);
 
 	FDelegateHandle BridgeEnableDelegateHandle;
diff --git a/Engine/Plugins/Experimental/EditorDataStorage/Source/TedsCore/Private/Compatibility/TedsTypedElementBridgeQueries.cpp b/Engine/Plugins/Experimental/EditorDataStorage/Source/TedsCore/Private/Compatibility/TedsTypedElementBridgeQueries.cpp
--- a/Engine/Plugins/Experimental/EditorDataStorage/Source/TedsCore/Private/Compatibility/TedsTypedElementBridgeQueries.cpp
+++ b/Engine/Plugins/Experimental/EditorDataStorage/Source/TedsCore/Private/Compatibility/TedsTypedElementBridgeQueries.cpp
@@ -26,13 +26,13 @@ uint8 UTypedElementBridgeDataStorageFactory::GetOrder() const
 	return 110;
 }
 
-void UTypedElementBridgeDataStorageFactory::PreRegister(IEditorDataStorageProvider& DataStorage)
+void UTypedElementBridgeDataStorageFactory::PreRegister(UE::Editor::DataStorage::ICoreProvider& DataStorage)
 {
 	Super::PreRegister(DataStorage);
 	DebugEnabledDelegateHandle = UE::Editor::DataStorage::Compatibility::Private::CVarBridgeEnabled->OnChangedDelegate().AddUObject(this, &UTypedElementBridgeDataStorageFactory::HandleOnEnabled);
 }
 
-void UTypedElementBridgeDataStorageFactory::PreShutdown(IEditorDataStorageProvider& DataStorage)
+void UTypedElementBridgeDataStorageFactory::PreShutdown(UE::Editor::DataStorage::ICoreProvider& DataStorage)
 {
 	UE::Editor::DataStorage::Compatibility::Private::CVarBridgeEnabled->OnChangedDelegate().Remove(DebugEnabledDelegateHandle);
 	DebugEnabledDelegateHandle.Reset();
@@ -41,7 +41,7 @@ void UTypedElementBridgeDataStorageFactory::PreShutdown(IEditorDataStorageProvid
 	Super::PreShutdown(DataStorage);
 }
 
-void UTypedElementBridgeDataStorageFactory::RegisterQueries(IEditorDataStorageProvider& DataStorage)
+void UTypedElementBridgeDataStorageFactory::RegisterQueries(UE::Editor::DataStorage::ICoreProvider& DataStorage)
 {
 	Super::RegisterQueries(DataStorage);
 
@@ -56,7 +56,7 @@ bool UTypedElementBridgeDataStorageFactory::IsEnabled()
 	return UE::Editor::DataStorage::Compatibility::Private::CVarBridgeEnabled->GetBool();
 }
 
-void UTypedElementBridgeDataStorageFactory::RegisterQuery_NewUObject(IEditorDataStorageProvider& DataStorage)
+void UTypedElementBridgeDataStorageFactory::RegisterQuery_NewUObject(UE::Editor::DataStorage::ICoreProvider& DataStorage)
 {
 	using namespace UE::Editor::DataStorage::Queries;
 	
@@ -66,11 +66,11 @@ void UTypedElementBridgeDataStorageFactory::RegisterQuery_NewUObject(IEditorData
 		.Compile());
 }
 
-void UTypedElementBridgeDataStorageFactory::UnregisterQuery_NewUObject(IEditorDataStorageProvider& DataStorage)
+void UTypedElementBridgeDataStorageFactory::UnregisterQuery_NewUObject(UE::Editor::DataStorage::ICoreProvider& DataStorage)
 {
 }
 
-void UTypedElementBridgeDataStorageFactory::CleanupTypedElementColumns(IEditorDataStorageProvider& DataStorage)
+void UTypedElementBridgeDataStorageFactory::CleanupTypedElementColumns(UE::Editor::DataStorage::ICoreProvider& DataStorage)
 {
 	using namespace UE::Editor::DataStorage::Queries;
 
@@ -92,7 +92,7 @@ void UTypedElementBridgeDataStorageFactory::CleanupTypedElementColumns(IEditorDa
 void UTypedElementBridgeDataStorageFactory::HandleOnEnabled(IConsoleVariable* CVar)
 {
 	using namespace UE::Editor::DataStorage;
-	IEditorDataStorageProvider* DataStorage = GetMutableDataStorageFeature<IEditorDataStorageProvider>(StorageFeatureName);
+	ICoreProvider* DataStorage = GetMutableDataStorageFeature<ICoreProvider>(StorageFeatureName);
 	bool bIsEnabled = CVar->GetBool();
 
 	if (bIsEnabled)
diff --git a/Engine/Plugins/Experimental/EditorDataStorage/Source/TedsCore/Private/Compatibility/TedsTypedElementBridgeQueries.h b/Engine/Plugins/Experimental/EditorDataStorage/Source/TedsCore/Private/Compatibility/TedsTypedElementBridgeQueries.h
--- a/Engine/Plugins/Experimental/EditorDataStorage/Source/TedsCore/Private/Compatibility/TedsTypedElementBridgeQueries.h
+++ b/Engine/Plugins/Experimental/EditorDataStorage/Source/TedsCore/Private/Compatibility/TedsTypedElementBridgeQueries.h
@@ -9,8 +9,6 @@
 
 #include "TedsTypedElementBridgeQueries.generated.h"
 
-class UTypedElementRegistry;
-
 /**
  * This class is responsible for running queries that will ensure Typed Element Handles
  * are cleaned up when TEDS is shut down.
@@ -22,16 +20,16 @@ class UTypedElementBridgeDataStorageFactory : public UEditorDataStorageFactory
 public:
 	~UTypedElementBridgeDataStorageFactory() override = default;
 	virtual uint8 GetOrder() const override;
-	virtual void PreRegister(IEditorDataStorageProvider& DataStorage) override;
-	virtual void PreShutdown(IEditorDataStorageProvider& DataStorage) override;
-	virtual void RegisterQueries(IEditorDataStorageProvider& DataStorage) override;
+	virtual void PreRegister(UE::Editor::DataStorage::ICoreProvider& DataStorage) override;
+	virtual void PreShutdown(UE::Editor::DataStorage::ICoreProvider& DataStorage) override;
+	virtual void RegisterQueries(UE::Editor::DataStorage::ICoreProvider& DataStorage) override;
 
 	static bool IsEnabled();
 
 private:
-	void RegisterQuery_NewUObject(IEditorDataStorageProvider& DataStorage);
-	void UnregisterQuery_NewUObject(IEditorDataStorageProvider& DataStorage);
-	void CleanupTypedElementColumns(IEditorDataStorageProvider& DataStorage);
+	void RegisterQuery_NewUObject(UE::Editor::DataStorage::ICoreProvider& DataStorage);
+	void UnregisterQuery_NewUObject(UE::Editor::DataStorage::ICoreProvider& DataStorage);
+	void CleanupTypedElementColumns(UE::Editor::DataStorage::ICoreProvider& DataStorage);
 	void HandleOnEnabled(IConsoleVariable* CVar);
 	
 	UE::Editor::DataStorage::QueryHandle RemoveTypedElementRowHandleQuery = UE::Editor::DataStorage::InvalidQueryHandle;
diff --git a/Engine/Plugins/Experimental/EditorDataStorage/Source/TedsCore/Private/Compatibility/TypedElementActorIconOverrideQueries.cpp b/Engine/Plugins/Experimental/EditorDataStorage/Source/TedsCore/Private/Compatibility/TypedElementActorIconOverrideQueries.cpp
--- a/Engine/Plugins/Experimental/EditorDataStorage/Source/TedsCore/Private/Compatibility/TypedElementActorIconOverrideQueries.cpp
+++ b/Engine/Plugins/Experimental/EditorDataStorage/Source/TedsCore/Private/Compatibility/TypedElementActorIconOverrideQueries.cpp
@@ -8,7 +8,7 @@
 #include "Elements/Framework/TypedElementQueryBuilder.h"
 #include "GameFramework/Actor.h"
 
-void UActorIconOverrideDataStorageFactory::RegisterQueries(IEditorDataStorageProvider& DataStorage)
+void UActorIconOverrideDataStorageFactory::RegisterQueries(UE::Editor::DataStorage::ICoreProvider& DataStorage)
 {
 	using namespace UE::Editor::DataStorage::Queries;
 
diff --git a/Engine/Plugins/Experimental/EditorDataStorage/Source/TedsCore/Private/Compatibility/TypedElementActorIconOverrideQueries.h b/Engine/Plugins/Experimental/EditorDataStorage/Source/TedsCore/Private/Compatibility/TypedElementActorIconOverrideQueries.h
--- a/Engine/Plugins/Experimental/EditorDataStorage/Source/TedsCore/Private/Compatibility/TypedElementActorIconOverrideQueries.h
+++ b/Engine/Plugins/Experimental/EditorDataStorage/Source/TedsCore/Private/Compatibility/TypedElementActorIconOverrideQueries.h
@@ -15,5 +15,5 @@ class UActorIconOverrideDataStorageFactory : public UEditorDataStorageFactory
 public:
 	virtual ~UActorIconOverrideDataStorageFactory() override = default;
 
-	virtual void RegisterQueries(IEditorDataStorageProvider& DataStorage) override;
+	virtual void RegisterQueries(UE::Editor::DataStorage::ICoreProvider& DataStorage) override;
 };
diff --git a/Engine/Plugins/Experimental/EditorDataStorage/Source/TedsCore/Private/Compatibility/TypedElementActorLabelQueries.cpp b/Engine/Plugins/Experimental/EditorDataStorage/Source/TedsCore/Private/Compatibility/TypedElementActorLabelQueries.cpp
--- a/Engine/Plugins/Experimental/EditorDataStorage/Source/TedsCore/Private/Compatibility/TypedElementActorLabelQueries.cpp
+++ b/Engine/Plugins/Experimental/EditorDataStorage/Source/TedsCore/Private/Compatibility/TypedElementActorLabelQueries.cpp
@@ -13,13 +13,13 @@
 
 #define LOCTEXT_NAMESPACE "TedsCore"
 
-void UActorLabelDataStorageFactory::RegisterQueries(IEditorDataStorageProvider& DataStorage)
+void UActorLabelDataStorageFactory::RegisterQueries(UE::Editor::DataStorage::ICoreProvider& DataStorage)
 {
 	RegisterActorLabelToColumnQuery(DataStorage);
 	RegisterLabelColumnToActorQuery(DataStorage);
 }
 
-void UActorLabelDataStorageFactory::RegisterActorLabelToColumnQuery(IEditorDataStorageProvider& DataStorage) const
+void UActorLabelDataStorageFactory::RegisterActorLabelToColumnQuery(UE::Editor::DataStorage::ICoreProvider& DataStorage) const
 {
 	using namespace UE::Editor::DataStorage::Queries;
 	
@@ -47,7 +47,7 @@ void UActorLabelDataStorageFactory::RegisterActorLabelToColumnQuery(IEditorDataS
 		.Compile());
 }
 
-void UActorLabelDataStorageFactory::RegisterLabelColumnToActorQuery(IEditorDataStorageProvider& DataStorage) const
+void UActorLabelDataStorageFactory::RegisterLabelColumnToActorQuery(UE::Editor::DataStorage::ICoreProvider& DataStorage) const
 {
 	using namespace UE::Editor::DataStorage::Queries;
 
diff --git a/Engine/Plugins/Experimental/EditorDataStorage/Source/TedsCore/Private/Compatibility/TypedElementActorLabelQueries.h b/Engine/Plugins/Experimental/EditorDataStorage/Source/TedsCore/Private/Compatibility/TypedElementActorLabelQueries.h
--- a/Engine/Plugins/Experimental/EditorDataStorage/Source/TedsCore/Private/Compatibility/TypedElementActorLabelQueries.h
+++ b/Engine/Plugins/Experimental/EditorDataStorage/Source/TedsCore/Private/Compatibility/TypedElementActorLabelQueries.h
@@ -16,16 +16,16 @@ class UActorLabelDataStorageFactory : public UEditorDataStorageFactory
 public:
 	~UActorLabelDataStorageFactory() override = default;
 
-	void RegisterQueries(IEditorDataStorageProvider& DataStorage) override;
+	void RegisterQueries(UE::Editor::DataStorage::ICoreProvider& DataStorage) override;
 
 private:
 	/**
 	 * Takes the label set on an actor and copies it to the Data Storage if they differ.
 	 */
-	void RegisterActorLabelToColumnQuery(IEditorDataStorageProvider& DataStorage) const;
+	void RegisterActorLabelToColumnQuery(UE::Editor::DataStorage::ICoreProvider& DataStorage) const;
 	/**
 	 * Takes the label stored in the Data Storage and copies it to the actor's label if the FTypedElementSyncBackToWorldTag
 	 * has been set and the labels differ.
 	 */
-	void RegisterLabelColumnToActorQuery(IEditorDataStorageProvider& DataStorage) const;
+	void RegisterLabelColumnToActorQuery(UE::Editor::DataStorage::ICoreProvider& DataStorage) const;
 };
\ No newline at end of file
diff --git a/Engine/Plugins/Experimental/EditorDataStorage/Source/TedsCore/Private/Compatibility/TypedElementActorParentQueries.cpp b/Engine/Plugins/Experimental/EditorDataStorage/Source/TedsCore/Private/Compatibility/TypedElementActorParentQueries.cpp
--- a/Engine/Plugins/Experimental/EditorDataStorage/Source/TedsCore/Private/Compatibility/TypedElementActorParentQueries.cpp
+++ b/Engine/Plugins/Experimental/EditorDataStorage/Source/TedsCore/Private/Compatibility/TypedElementActorParentQueries.cpp
@@ -20,7 +20,7 @@ namespace UE::Editor::DataStorage::Queries::Private
 		TEXT("Mirror parent information for actors to TEDS (only works when set on startup)"));
 };
 
-void UActorParentDataStorageFactory::RegisterQueries(IEditorDataStorageProvider& DataStorage)
+void UActorParentDataStorageFactory::RegisterQueries(UE::Editor::DataStorage::ICoreProvider& DataStorage)
 {
 	if(UE::Editor::DataStorage::Queries::Private::bAddParentColumnToActors)
 	{
@@ -29,7 +29,7 @@ void UActorParentDataStorageFactory::RegisterQueries(IEditorDataStorageProvider&
 	}
 }
 
-void UActorParentDataStorageFactory::RegisterAddParentColumn(IEditorDataStorageProvider& DataStorage) const
+void UActorParentDataStorageFactory::RegisterAddParentColumn(UE::Editor::DataStorage::ICoreProvider& DataStorage) const
 {
 	using namespace UE::Editor::DataStorage;
 	using namespace UE::Editor::DataStorage::Queries;
@@ -64,7 +64,7 @@ void UActorParentDataStorageFactory::RegisterAddParentColumn(IEditorDataStorageP
 		.Compile());
 }
 
-void UActorParentDataStorageFactory::RegisterUpdateOrRemoveParentColumn(IEditorDataStorageProvider& DataStorage) const
+void UActorParentDataStorageFactory::RegisterUpdateOrRemoveParentColumn(UE::Editor::DataStorage::ICoreProvider& DataStorage) const
 {
 	using namespace UE::Editor::DataStorage::Queries;
 	
diff --git a/Engine/Plugins/Experimental/EditorDataStorage/Source/TedsCore/Private/Compatibility/TypedElementActorParentQueries.h b/Engine/Plugins/Experimental/EditorDataStorage/Source/TedsCore/Private/Compatibility/TypedElementActorParentQueries.h
--- a/Engine/Plugins/Experimental/EditorDataStorage/Source/TedsCore/Private/Compatibility/TypedElementActorParentQueries.h
+++ b/Engine/Plugins/Experimental/EditorDataStorage/Source/TedsCore/Private/Compatibility/TypedElementActorParentQueries.h
@@ -15,17 +15,17 @@ class UActorParentDataStorageFactory : public UEditorDataStorageFactory
 public:
 	~UActorParentDataStorageFactory() override = default;
 
-	void RegisterQueries(IEditorDataStorageProvider& DataStorage) override;
+	void RegisterQueries(UE::Editor::DataStorage::ICoreProvider& DataStorage) override;
 
 private:
 	/**
 	 * Checks rows with actors that don't have a parent column yet if one needs to be added whenever
 	 * the row is marked for updates.
 	 */
-	void RegisterAddParentColumn(IEditorDataStorageProvider& DataStorage) const;
+	void RegisterAddParentColumn(UE::Editor::DataStorage::ICoreProvider& DataStorage) const;
 	/**
 	 * Updates the parent column with the parent from the actor or removes it if there's no parent associated
 	 * with the actor anymore.
 	 */
-	void RegisterUpdateOrRemoveParentColumn(IEditorDataStorageProvider& DataStorage) const;
+	void RegisterUpdateOrRemoveParentColumn(UE::Editor::DataStorage::ICoreProvider& DataStorage) const;
 };
diff --git a/Engine/Plugins/Experimental/EditorDataStorage/Source/TedsCore/Private/Compatibility/TypedElementActorTransformQueries.cpp b/Engine/Plugins/Experimental/EditorDataStorage/Source/TedsCore/Private/Compatibility/TypedElementActorTransformQueries.cpp
--- a/Engine/Plugins/Experimental/EditorDataStorage/Source/TedsCore/Private/Compatibility/TypedElementActorTransformQueries.cpp
+++ b/Engine/Plugins/Experimental/EditorDataStorage/Source/TedsCore/Private/Compatibility/TypedElementActorTransformQueries.cpp
@@ -19,14 +19,14 @@ namespace UE::Editor::DataStorage::ActorTransformQueries::Private
 		TEXT("Sync actor transform changes from TEDS back to the world (only works when set on startup)"));
 }
 
-void UActorTransformDataStorageFactory::RegisterQueries(IEditorDataStorageProvider& DataStorage)
+void UActorTransformDataStorageFactory::RegisterQueries(UE::Editor::DataStorage::ICoreProvider& DataStorage)
 {
 	RegisterActorAddTransformColumn(DataStorage);
 	RegisterActorLocalTransformToColumn(DataStorage);
 	RegisterLocalTransformColumnToActor(DataStorage);
 }
 
-void UActorTransformDataStorageFactory::RegisterActorAddTransformColumn(IEditorDataStorageProvider& DataStorage) const
+void UActorTransformDataStorageFactory::RegisterActorAddTransformColumn(UE::Editor::DataStorage::ICoreProvider& DataStorage) const
 {
 	using namespace UE::Editor::DataStorage::Queries;
 
@@ -49,7 +49,7 @@ void UActorTransformDataStorageFactory::RegisterActorAddTransformColumn(IEditorD
 		.Compile());
 }
 
-void UActorTransformDataStorageFactory::RegisterActorLocalTransformToColumn(IEditorDataStorageProvider& DataStorage) const
+void UActorTransformDataStorageFactory::RegisterActorLocalTransformToColumn(UE::Editor::DataStorage::ICoreProvider& DataStorage) const
 {
 	using namespace UE::Editor::DataStorage::Queries;
 
@@ -76,7 +76,7 @@ void UActorTransformDataStorageFactory::RegisterActorLocalTransformToColumn(IEdi
 		.Compile());
 }
 
-void UActorTransformDataStorageFactory::RegisterLocalTransformColumnToActor(IEditorDataStorageProvider& DataStorage) const
+void UActorTransformDataStorageFactory::RegisterLocalTransformColumnToActor(UE::Editor::DataStorage::ICoreProvider& DataStorage) const
 {
 	using namespace UE::Editor::DataStorage::Queries;
 
diff --git a/Engine/Plugins/Experimental/EditorDataStorage/Source/TedsCore/Private/Compatibility/TypedElementActorTransformQueries.h b/Engine/Plugins/Experimental/EditorDataStorage/Source/TedsCore/Private/Compatibility/TypedElementActorTransformQueries.h
--- a/Engine/Plugins/Experimental/EditorDataStorage/Source/TedsCore/Private/Compatibility/TypedElementActorTransformQueries.h
+++ b/Engine/Plugins/Experimental/EditorDataStorage/Source/TedsCore/Private/Compatibility/TypedElementActorTransformQueries.h
@@ -15,22 +15,22 @@ class UActorTransformDataStorageFactory : public UEditorDataStorageFactory
 public:
 	~UActorTransformDataStorageFactory() override = default;
 
-	void RegisterQueries(IEditorDataStorageProvider& DataStorage) override;
+	void RegisterQueries(UE::Editor::DataStorage::ICoreProvider& DataStorage) override;
 
 private:
 	/**
 	 * Checks actors that don't have a transform column and adds one if an actor has been
 	 * assigned a transform.
 	 */
-	void RegisterActorAddTransformColumn(IEditorDataStorageProvider& DataStorage) const;
+	void RegisterActorAddTransformColumn(UE::Editor::DataStorage::ICoreProvider& DataStorage) const;
 	/**
 	 * Takes the transform set on an actor and copies it to the Data Storage or removes the
 	 * transform column if there's not transform available anymore.
 	 */
-	void RegisterActorLocalTransformToColumn(IEditorDataStorageProvider& DataStorage) const;
+	void RegisterActorLocalTransformToColumn(UE::Editor::DataStorage::ICoreProvider& DataStorage) const;
 	/**
 	 * Takes the transform stored in the Data Storage and copies it to the actor's tranform if
 	 * the FTypedElementSyncBackToWorldTag has been set.
 	 */
-	void RegisterLocalTransformColumnToActor(IEditorDataStorageProvider& DataStorage) const;
+	void RegisterLocalTransformColumnToActor(UE::Editor::DataStorage::ICoreProvider& DataStorage) const;
 };
diff --git a/Engine/Plugins/Experimental/EditorDataStorage/Source/TedsCore/Private/Compatibility/TypedElementActorViewportProcessors.cpp b/Engine/Plugins/Experimental/EditorDataStorage/Source/TedsCore/Private/Compatibility/TypedElementActorViewportProcessors.cpp
--- a/Engine/Plugins/Experimental/EditorDataStorage/Source/TedsCore/Private/Compatibility/TypedElementActorViewportProcessors.cpp
+++ b/Engine/Plugins/Experimental/EditorDataStorage/Source/TedsCore/Private/Compatibility/TypedElementActorViewportProcessors.cpp
@@ -24,7 +24,7 @@ namespace UE::Editor::DataStorage::Private
 
 				TRACE_CPUPROFILER_EVENT_SCOPE(TEDS.Debug.AddOverlayColorToSelectionCommand);
 
-				if (IEditorDataStorageProvider* DataStorage = GetMutableDataStorageFeature<IEditorDataStorageProvider>(StorageFeatureName))
+				if (ICoreProvider* DataStorage = GetMutableDataStorageFeature<ICoreProvider>(StorageFeatureName))
 				{
 					static QueryHandle OverlayQuery = InvalidQueryHandle;
 					if (OverlayQuery == InvalidQueryHandle)
@@ -83,7 +83,7 @@ namespace UE::Editor::DataStorage::Private
 
 				TRACE_CPUPROFILER_EVENT_SCOPE(TEDS.Debug.AddOverlayColorToSelectionCommand);
 
-				if (IEditorDataStorageProvider* DataStorage = GetMutableDataStorageFeature<IEditorDataStorageProvider>(StorageFeatureName))
+				if (ICoreProvider* DataStorage = GetMutableDataStorageFeature<ICoreProvider>(StorageFeatureName))
 				{
 					static QueryHandle OverlayQuery = InvalidQueryHandle;
 					if (OverlayQuery == InvalidQueryHandle)
@@ -135,7 +135,7 @@ namespace UE::Editor::DataStorage::Private
 
 				TRACE_CPUPROFILER_EVENT_SCOPE(TEDS.Debug.AddOverlayColorToSelectionCommand);
 
-				if (IEditorDataStorageProvider* DataStorage = GetMutableDataStorageFeature<IEditorDataStorageProvider>(StorageFeatureName))
+				if (ICoreProvider* DataStorage = GetMutableDataStorageFeature<ICoreProvider>(StorageFeatureName))
 				{
 					static QueryHandle OverlayQuery = InvalidQueryHandle;
 					if (OverlayQuery == InvalidQueryHandle)
@@ -168,13 +168,13 @@ namespace UE::Editor::DataStorage::Private
 			}));
 } // namespace UE::Editor::DataStorage::Private
 
-void UActorViewportDataStorageFactory::RegisterQueries(IEditorDataStorageProvider& DataStorage)
+void UActorViewportDataStorageFactory::RegisterQueries(UE::Editor::DataStorage::ICoreProvider& DataStorage)
 {
 	RegisterOutlineColorColumnToActor(DataStorage);
 	RegisterOverlayColorColumnToActor(DataStorage);
 }
 
-void UActorViewportDataStorageFactory::RegisterOutlineColorColumnToActor(IEditorDataStorageProvider& DataStorage)
+void UActorViewportDataStorageFactory::RegisterOutlineColorColumnToActor(UE::Editor::DataStorage::ICoreProvider& DataStorage)
 {
 	using namespace UE::Editor::DataStorage::Queries;
 
@@ -199,7 +199,7 @@ void UActorViewportDataStorageFactory::RegisterOutlineColorColumnToActor(IEditor
 		.Compile());
 }
 
-void UActorViewportDataStorageFactory::RegisterOverlayColorColumnToActor(IEditorDataStorageProvider& DataStorage)
+void UActorViewportDataStorageFactory::RegisterOverlayColorColumnToActor(UE::Editor::DataStorage::ICoreProvider& DataStorage)
 {
 	using namespace UE::Editor::DataStorage::Queries;
 	
diff --git a/Engine/Plugins/Experimental/EditorDataStorage/Source/TedsCore/Private/Compatibility/TypedElementActorViewportProcessors.h b/Engine/Plugins/Experimental/EditorDataStorage/Source/TedsCore/Private/Compatibility/TypedElementActorViewportProcessors.h
--- a/Engine/Plugins/Experimental/EditorDataStorage/Source/TedsCore/Private/Compatibility/TypedElementActorViewportProcessors.h
+++ b/Engine/Plugins/Experimental/EditorDataStorage/Source/TedsCore/Private/Compatibility/TypedElementActorViewportProcessors.h
@@ -15,9 +15,9 @@ class UActorViewportDataStorageFactory : public UEditorDataStorageFactory
 public:
 	~UActorViewportDataStorageFactory() override = default;
 
-	void RegisterQueries(IEditorDataStorageProvider& DataStorage) override;
+	void RegisterQueries(UE::Editor::DataStorage::ICoreProvider& DataStorage) override;
 
 private:
-	void RegisterOutlineColorColumnToActor(IEditorDataStorageProvider& DataStorage);
-	void RegisterOverlayColorColumnToActor(IEditorDataStorageProvider& DataStorage);
+	void RegisterOutlineColorColumnToActor(UE::Editor::DataStorage::ICoreProvider& DataStorage);
+	void RegisterOverlayColorColumnToActor(UE::Editor::DataStorage::ICoreProvider& DataStorage);
 };
diff --git a/Engine/Plugins/Experimental/EditorDataStorage/Source/TedsCore/Private/Compatibility/TypedElementActorVisibilityQueries.cpp b/Engine/Plugins/Experimental/EditorDataStorage/Source/TedsCore/Private/Compatibility/TypedElementActorVisibilityQueries.cpp
--- a/Engine/Plugins/Experimental/EditorDataStorage/Source/TedsCore/Private/Compatibility/TypedElementActorVisibilityQueries.cpp
+++ b/Engine/Plugins/Experimental/EditorDataStorage/Source/TedsCore/Private/Compatibility/TypedElementActorVisibilityQueries.cpp
@@ -8,14 +8,14 @@
 #include "Elements/Framework/TypedElementQueryBuilder.h"
 #include "GameFramework/Actor.h"
 
-void UActorVisibilityDataStorageFactory::RegisterQueries(IEditorDataStorageProvider& DataStorage)
+void UActorVisibilityDataStorageFactory::RegisterQueries(UE::Editor::DataStorage::ICoreProvider& DataStorage)
 {
 	RegisterActorAddVisibilityColumn(DataStorage);
 	RegisterActorVisibilityToColumnQuery(DataStorage);
 	RegisterVisibilityColumnToActorQuery(DataStorage);
 }
 
-void UActorVisibilityDataStorageFactory::RegisterActorAddVisibilityColumn(IEditorDataStorageProvider& DataStorage) const
+void UActorVisibilityDataStorageFactory::RegisterActorAddVisibilityColumn(UE::Editor::DataStorage::ICoreProvider& DataStorage) const
 {
 	using namespace UE::Editor::DataStorage::Queries;
 
@@ -37,7 +37,7 @@ void UActorVisibilityDataStorageFactory::RegisterActorAddVisibilityColumn(IEdito
 	);
 }
 
-void UActorVisibilityDataStorageFactory::RegisterActorVisibilityToColumnQuery(IEditorDataStorageProvider& DataStorage) const
+void UActorVisibilityDataStorageFactory::RegisterActorVisibilityToColumnQuery(UE::Editor::DataStorage::ICoreProvider& DataStorage) const
 {
 	using namespace UE::Editor::DataStorage::Queries;
 
@@ -58,7 +58,7 @@ void UActorVisibilityDataStorageFactory::RegisterActorVisibilityToColumnQuery(IE
 	);
 }
 
-void UActorVisibilityDataStorageFactory::RegisterVisibilityColumnToActorQuery(IEditorDataStorageProvider& DataStorage) const
+void UActorVisibilityDataStorageFactory::RegisterVisibilityColumnToActorQuery(UE::Editor::DataStorage::ICoreProvider& DataStorage) const
 {
 	using namespace UE::Editor::DataStorage::Queries;
 
diff --git a/Engine/Plugins/Experimental/EditorDataStorage/Source/TedsCore/Private/Compatibility/TypedElementActorVisibilityQueries.h b/Engine/Plugins/Experimental/EditorDataStorage/Source/TedsCore/Private/Compatibility/TypedElementActorVisibilityQueries.h
--- a/Engine/Plugins/Experimental/EditorDataStorage/Source/TedsCore/Private/Compatibility/TypedElementActorVisibilityQueries.h
+++ b/Engine/Plugins/Experimental/EditorDataStorage/Source/TedsCore/Private/Compatibility/TypedElementActorVisibilityQueries.h
@@ -17,22 +17,22 @@ class UActorVisibilityDataStorageFactory : public UEditorDataStorageFactory
 public:
 	~UActorVisibilityDataStorageFactory() override = default;
 
-	void RegisterQueries(IEditorDataStorageProvider& DataStorage) override;
+	void RegisterQueries(UE::Editor::DataStorage::ICoreProvider& DataStorage) override;
 
 private:
 
 	/**
 	 * Adds the Visibility columns to new actors that do not have one already.
 	 */
-	void RegisterActorAddVisibilityColumn(IEditorDataStorageProvider& DataStorage) const;
+	void RegisterActorAddVisibilityColumn(UE::Editor::DataStorage::ICoreProvider& DataStorage) const;
 
 	/**
 	 * Takes the visibility set on an actor and copies it to the Data Storage if they differ.
 	 */
-	void RegisterActorVisibilityToColumnQuery(IEditorDataStorageProvider& DataStorage) const;
+	void RegisterActorVisibilityToColumnQuery(UE::Editor::DataStorage::ICoreProvider& DataStorage) const;
 	/**
 	 * Takes the visibility stored in the Data Storage and copies it to the actor's visibility if the FTypedElementSyncBackToWorldTag
 	 * has been set and the visibility differ.
 	 */
-	void RegisterVisibilityColumnToActorQuery(IEditorDataStorageProvider& DataStorage) const;
+	void RegisterVisibilityColumnToActorQuery(UE::Editor::DataStorage::ICoreProvider& DataStorage) const;
 };
diff --git a/Engine/Plugins/Experimental/EditorDataStorage/Source/TedsCore/Private/Compatibility/TypedElementObjectReinstancingManager.cpp b/Engine/Plugins/Experimental/EditorDataStorage/Source/TedsCore/Private/Compatibility/TypedElementObjectReinstancingManager.cpp
--- a/Engine/Plugins/Experimental/EditorDataStorage/Source/TedsCore/Private/Compatibility/TypedElementObjectReinstancingManager.cpp
+++ b/Engine/Plugins/Experimental/EditorDataStorage/Source/TedsCore/Private/Compatibility/TypedElementObjectReinstancingManager.cpp
@@ -71,7 +71,7 @@ void UTedsObjectReinstancingManager::HandleOnObjectsReinstanced(
 {
 	using namespace UE::Editor::DataStorage;
 
-	IEditorDataStorageProvider* Interface = DataStorage;
+	ICoreProvider* Interface = DataStorage;
 	for (FCoreUObjectDelegates::FReplacementObjectMap::TConstIterator Iter = ObjectReplacementMap.CreateConstIterator(); Iter; ++Iter)
 	{
 		const void* PreDeleteObject = Iter->Key;
diff --git a/Engine/Plugins/Experimental/EditorDataStorage/Source/TedsCore/Private/Compatibility/TypedElementObjectReinstancingManager.h b/Engine/Plugins/Experimental/EditorDataStorage/Source/TedsCore/Private/Compatibility/TypedElementObjectReinstancingManager.h
--- a/Engine/Plugins/Experimental/EditorDataStorage/Source/TedsCore/Private/Compatibility/TypedElementObjectReinstancingManager.h
+++ b/Engine/Plugins/Experimental/EditorDataStorage/Source/TedsCore/Private/Compatibility/TypedElementObjectReinstancingManager.h
@@ -9,7 +9,6 @@
 
 #include "TypedElementObjectReinstancingManager.generated.h"
 
-class IEditorDataStorageCompatibilityProvider;
 class UEditorDataStorageCompatibility;
 class UTypedElementMementoSystem;
 class UEditorDataStorage;
@@ -17,6 +16,7 @@ class UEditorDataStorage;
 namespace UE::Editor::DataStorage
 {
 	struct FObjectTypeInfo;
+	class ICompatibilityProvider;
 }
 
 UCLASS(Transient)
diff --git a/Engine/Plugins/Experimental/EditorDataStorage/Source/TedsCore/Private/Compatibility/TypedElementUObjectWorldQueries.cpp b/Engine/Plugins/Experimental/EditorDataStorage/Source/TedsCore/Private/Compatibility/TypedElementUObjectWorldQueries.cpp
--- a/Engine/Plugins/Experimental/EditorDataStorage/Source/TedsCore/Private/Compatibility/TypedElementUObjectWorldQueries.cpp
+++ b/Engine/Plugins/Experimental/EditorDataStorage/Source/TedsCore/Private/Compatibility/TypedElementUObjectWorldQueries.cpp
@@ -7,13 +7,13 @@
 #include "Elements/Framework/TypedElementQueryBuilder.h"
 #include "Engine/World.h"
 
-void UObjectWorldDataStorageFactory::RegisterQueries(IEditorDataStorageProvider& DataStorage)
+void UObjectWorldDataStorageFactory::RegisterQueries(UE::Editor::DataStorage::ICoreProvider& DataStorage)
 {
 	RegisterAddWorldColumn(DataStorage);
 	RegisterUpdateOrRemoveWorldColumn(DataStorage);
 }
 
-void UObjectWorldDataStorageFactory::RegisterAddWorldColumn(IEditorDataStorageProvider& DataStorage) const
+void UObjectWorldDataStorageFactory::RegisterAddWorldColumn(UE::Editor::DataStorage::ICoreProvider& DataStorage) const
 {
 	using namespace UE::Editor::DataStorage::Queries;
 	
@@ -40,7 +40,7 @@ void UObjectWorldDataStorageFactory::RegisterAddWorldColumn(IEditorDataStoragePr
 		.Compile());
 }
 
-void UObjectWorldDataStorageFactory::RegisterUpdateOrRemoveWorldColumn(IEditorDataStorageProvider& DataStorage) const
+void UObjectWorldDataStorageFactory::RegisterUpdateOrRemoveWorldColumn(UE::Editor::DataStorage::ICoreProvider& DataStorage) const
 {
 	using namespace UE::Editor::DataStorage::Queries;
 	
diff --git a/Engine/Plugins/Experimental/EditorDataStorage/Source/TedsCore/Private/Compatibility/TypedElementUObjectWorldQueries.h b/Engine/Plugins/Experimental/EditorDataStorage/Source/TedsCore/Private/Compatibility/TypedElementUObjectWorldQueries.h
--- a/Engine/Plugins/Experimental/EditorDataStorage/Source/TedsCore/Private/Compatibility/TypedElementUObjectWorldQueries.h
+++ b/Engine/Plugins/Experimental/EditorDataStorage/Source/TedsCore/Private/Compatibility/TypedElementUObjectWorldQueries.h
@@ -15,17 +15,17 @@ class UObjectWorldDataStorageFactory : public UEditorDataStorageFactory
 public:
 	~UObjectWorldDataStorageFactory() override = default;
 
-	void RegisterQueries(IEditorDataStorageProvider& DataStorage) override;
+	void RegisterQueries(UE::Editor::DataStorage::ICoreProvider& DataStorage) override;
 
 private:
 	/**
 	 * Checks rows with UObjects that don't have a world column yet if one needs to be added whenever
 	 * the row is marked for updates.
 	 */
-	void RegisterAddWorldColumn(IEditorDataStorageProvider& DataStorage) const;
+	void RegisterAddWorldColumn(UE::Editor::DataStorage::ICoreProvider& DataStorage) const;
 	/**
 	 * Updates the world column with the world in the UObject or removes it if there's no world associated
 	 * with the UObject anymore.
 	 */
-	void RegisterUpdateOrRemoveWorldColumn(IEditorDataStorageProvider& DataStorage) const;
+	void RegisterUpdateOrRemoveWorldColumn(UE::Editor::DataStorage::ICoreProvider& DataStorage) const;
 };
diff --git a/Engine/Plugins/Experimental/EditorDataStorage/Source/TedsCore/Private/Debug/TypedElementDatabaseDebugCommands.cpp b/Engine/Plugins/Experimental/EditorDataStorage/Source/TedsCore/Private/Debug/TypedElementDatabaseDebugCommands.cpp
--- a/Engine/Plugins/Experimental/EditorDataStorage/Source/TedsCore/Private/Debug/TypedElementDatabaseDebugCommands.cpp
+++ b/Engine/Plugins/Experimental/EditorDataStorage/Source/TedsCore/Private/Debug/TypedElementDatabaseDebugCommands.cpp
@@ -23,7 +23,7 @@ namespace UE::Editor::DataStorage::Private
 	static const UStruct* GetTypeInfo(const TObjectPtr<const UScriptStruct>& TypeInfo) { return TypeInfo; }
 
 	template<typename TypeInfoType>
-	void PrintObjectTypeInformation(IEditorDataStorageProvider* DataStorage, FString Message, FOutputDevice& Output)
+	void PrintObjectTypeInformation(ICoreProvider* DataStorage, FString Message, FOutputDevice& Output)
 	{
 		using namespace UE::Editor::DataStorage::Queries;
 
@@ -65,7 +65,7 @@ namespace UE::Editor::DataStorage::Private
 	{
 		using namespace UE::Editor::DataStorage::Queries;
 
-		if (IEditorDataStorageProvider* DataStorage = GetMutableDataStorageFeature<IEditorDataStorageProvider>(StorageFeatureName))
+		if (ICoreProvider* DataStorage = GetMutableDataStorageFeature<ICoreProvider>(StorageFeatureName))
 		{
 			static QueryHandle LabelQuery = [DataStorage]
 			{
@@ -136,7 +136,7 @@ FAutoConsoleCommandWithOutputDevice PrintObjectTypeInformationConsoleCommand(
 
 			TRACE_CPUPROFILER_EVENT_SCOPE(TEDS.Debug.PrintObjectTypeInfo);
 
-			if (IEditorDataStorageProvider* DataStorage = GetMutableDataStorageFeature<IEditorDataStorageProvider>(StorageFeatureName))
+			if (ICoreProvider* DataStorage = GetMutableDataStorageFeature<ICoreProvider>(StorageFeatureName))
 			{
 				FString Message;
 				Output.Log(TEXT("The Typed Elements Data Storage has the types:"));
@@ -189,17 +189,17 @@ FAutoConsoleCommandWithOutputDevice ListExtensionsConsoleCommand(
 				Message += TEXT('\n');
 			};
 
-			if (const IEditorDataStorageProvider* DataStorage = GetDataStorageFeature<IEditorDataStorageProvider>(StorageFeatureName))
+			if (const ICoreProvider* DataStorage = GetDataStorageFeature<ICoreProvider>(StorageFeatureName))
 			{
 				Message = TEXT("Data Storage Extensions: \n");
 				DataStorage->ListExtensions(RecordExtensions);
 			}
-			if (const IEditorDataStorageCompatibilityProvider* DataStorageCompat = GetDataStorageFeature<IEditorDataStorageCompatibilityProvider>(CompatibilityFeatureName))
+			if (const ICompatibilityProvider* DataStorageCompat = GetDataStorageFeature<ICompatibilityProvider>(CompatibilityFeatureName))
 			{
 				Message += TEXT("Data Storage Compatibility Extensions: \n");
 				DataStorageCompat->ListExtensions(RecordExtensions);
 			}
-			if (const IEditorDataStorageUiProvider* DataStorageUi = GetDataStorageFeature<IEditorDataStorageUiProvider>(UiFeatureName))
+			if (const IUiProvider* DataStorageUi = GetDataStorageFeature<IUiProvider>(UiFeatureName))
 			{
 				Message += TEXT("Data Storage UI Extensions: \n");
 				DataStorageUi->ListExtensions(RecordExtensions);
@@ -217,7 +217,7 @@ static FAutoConsoleCommand CVarCreateRow(
 	FConsoleCommandWithArgsDelegate::CreateLambda([](const TArray<FString>& Args)
 	{
 		using namespace UE::Editor::DataStorage;
-		IEditorDataStorageProvider* DataStorage = GetMutableDataStorageFeature<IEditorDataStorageProvider>(StorageFeatureName);
+		ICoreProvider* DataStorage = GetMutableDataStorageFeature<ICoreProvider>(StorageFeatureName);
 		static TableHandle Table = DataStorage->RegisterTable<FTestColumnA>(FName(TEXT("Debug.CreateRow Table")));
 
 		const RowHandle RowHandle = DataStorage->AddRow(Table);
@@ -231,7 +231,7 @@ static FAutoConsoleCommand CVarAddDynamicColumnTag(
 	FConsoleCommandWithArgsDelegate::CreateLambda([](const TArray<FString>& Args)
 	{
 		using namespace UE::Editor::DataStorage;
-		IEditorDataStorageProvider* DataStorage = GetMutableDataStorageFeature<IEditorDataStorageProvider>(StorageFeatureName);
+		ICoreProvider* DataStorage = GetMutableDataStorageFeature<ICoreProvider>(StorageFeatureName);
 		
 		if (Args.Num() != 2)
 		{
@@ -254,7 +254,7 @@ static FAutoConsoleCommand CVarAddDynamicColumn(
 	FConsoleCommandWithArgsDelegate::CreateLambda([](const TArray<FString>& Args)
 	{
 		using namespace UE::Editor::DataStorage;
-		IEditorDataStorageProvider* DataStorage = GetMutableDataStorageFeature<IEditorDataStorageProvider>(StorageFeatureName);
+		ICoreProvider* DataStorage = GetMutableDataStorageFeature<ICoreProvider>(StorageFeatureName);
 		
 		if (Args.Num() != 2)
 		{
@@ -286,7 +286,7 @@ static FAutoConsoleCommand CVarRemoveDynamicColumn(
 	FConsoleCommandWithArgsDelegate::CreateLambda([](const TArray<FString>& Args)
 	{
 		using namespace UE::Editor::DataStorage;
-		IEditorDataStorageProvider* DataStorage = GetMutableDataStorageFeature<IEditorDataStorageProvider>(StorageFeatureName);
+		ICoreProvider* DataStorage = GetMutableDataStorageFeature<ICoreProvider>(StorageFeatureName);
 		
 		if (Args.Num() != 2)
 		{
@@ -309,7 +309,7 @@ static FAutoConsoleCommand CVarAddToDynamicColumn(
 	FConsoleCommandWithArgsDelegate::CreateLambda([](const TArray<FString>& Args)
 	{
 		using namespace UE::Editor::DataStorage;
-		IEditorDataStorageProvider* DataStorage = GetMutableDataStorageFeature<IEditorDataStorageProvider>(StorageFeatureName);
+		ICoreProvider* DataStorage = GetMutableDataStorageFeature<ICoreProvider>(StorageFeatureName);
 		
 		if (Args.Num() < 3 || Args.Num() > 4)
 		{
@@ -362,7 +362,7 @@ static FAutoConsoleCommand CVarPrintDynamicColumn(
 	FConsoleCommandWithArgsDelegate::CreateLambda([](const TArray<FString>& Args)
 	{
 		using namespace UE::Editor::DataStorage;
-		IEditorDataStorageProvider* DataStorage = GetMutableDataStorageFeature<IEditorDataStorageProvider>(StorageFeatureName);
+		ICoreProvider* DataStorage = GetMutableDataStorageFeature<ICoreProvider>(StorageFeatureName);
 		
 		if (Args.Num() != 2)
 		{
@@ -418,7 +418,7 @@ static FAutoConsoleCommand CVarPrintDynamicColumnWithQuery(
 			return;
 		}
 		
-		IEditorDataStorageProvider* DataStorage = GetMutableDataStorageFeature<IEditorDataStorageProvider>(StorageFeatureName);
+		ICoreProvider* DataStorage = GetMutableDataStorageFeature<ICoreProvider>(StorageFeatureName);
 		
 		const FName Identifier(*Args[0]);
 
@@ -478,7 +478,7 @@ static FAutoConsoleCommand CVarCountDynamicTagWithQuery(
 			return;
 		}
 				
-		IEditorDataStorageProvider* DataStorage = GetMutableDataStorageFeature<IEditorDataStorageProvider>(StorageFeatureName);
+		ICoreProvider* DataStorage = GetMutableDataStorageFeature<ICoreProvider>(StorageFeatureName);
 				
 		const FName Identifier(*Args[0]);
 
@@ -521,7 +521,7 @@ static FAutoConsoleCommand CVarRegisterListDynamicColumnQuery(
 			return;
 		}
 				
-		IEditorDataStorageProvider* DataStorage = GetMutableDataStorageFeature<IEditorDataStorageProvider>(StorageFeatureName);
+		ICoreProvider* DataStorage = GetMutableDataStorageFeature<ICoreProvider>(StorageFeatureName);
 		
 		const FName ActivationGroup(*Args[0]);
 		FName Identifier;
@@ -561,7 +561,8 @@ static FAutoConsoleCommand CVarUnRegisterListDynamicColumnQueries(
 	TEXT(""),
 	FConsoleCommandWithArgsDelegate::CreateLambda([](const TArray<FString>& Args)
 	{
-		IEditorDataStorageProvider* DataStorage = UE::Editor::DataStorage::GetMutableDataStorageFeature<IEditorDataStorageProvider>(UE::Editor::DataStorage::StorageFeatureName);
+		using namespace UE::Editor::DataStorage;
+		ICoreProvider* DataStorage = GetMutableDataStorageFeature<ICoreProvider>(UE::Editor::DataStorage::StorageFeatureName);
 
 		for (UE::Editor::DataStorage::QueryHandle& Handle : DynamicColumnQueries)
 		{
@@ -585,7 +586,7 @@ static FAutoConsoleCommand CVarActivateListDynamicColumnQuery(
 		}
 
 		using namespace UE::Editor::DataStorage;
-		IEditorDataStorageProvider* DataStorage = GetMutableDataStorageFeature<IEditorDataStorageProvider>(StorageFeatureName);
+		ICoreProvider* DataStorage = GetMutableDataStorageFeature<ICoreProvider>(StorageFeatureName);
 		
 		const FName ActivationGroup(*Args[0]);
 
@@ -599,7 +600,7 @@ static FAutoConsoleCommand CVarAddValueTag(
 	FConsoleCommandWithArgsDelegate::CreateLambda([](const TArray<FString>& Args)
 	{
 		using namespace UE::Editor::DataStorage;
-		IEditorDataStorageProvider* DataStorage = GetMutableDataStorageFeature<IEditorDataStorageProvider>(StorageFeatureName);
+		ICoreProvider* DataStorage = GetMutableDataStorageFeature<ICoreProvider>(StorageFeatureName);
 
 		if (Args.Num() != 3)
 		{
@@ -634,7 +635,7 @@ static FAutoConsoleCommand CVarRemoveValueTag(
 	FConsoleCommandWithArgsDelegate::CreateLambda([](const TArray<FString>& Args)
 	{
 		using namespace UE::Editor::DataStorage;
-		IEditorDataStorageProvider* DataStorage = GetMutableDataStorageFeature<IEditorDataStorageProvider>(StorageFeatureName);
+		ICoreProvider* DataStorage = GetMutableDataStorageFeature<ICoreProvider>(StorageFeatureName);
 
 		if (Args.Num() != 2)
 		{
@@ -666,7 +667,7 @@ static FAutoConsoleCommand CVarMatchValueTag(
 	{
 		using namespace UE::Editor::DataStorage::Queries;
 		
-		IEditorDataStorageProvider* DataStorage = GetMutableDataStorageFeature<IEditorDataStorageProvider>(StorageFeatureName);
+		ICoreProvider* DataStorage = GetMutableDataStorageFeature<ICoreProvider>(StorageFeatureName);
 
 		if (Args.Num() < 1 || Args.Num() > 2)
 		{
@@ -720,7 +721,7 @@ static FAutoConsoleCommand CVarAddValueTagFromEnum(
 	{
 		using namespace UE::Editor::DataStorage;
 
-		IEditorDataStorageProvider* DataStorage = GetMutableDataStorageFeature<IEditorDataStorageProvider>(StorageFeatureName);
+		ICoreProvider* DataStorage = GetMutableDataStorageFeature<ICoreProvider>(StorageFeatureName);
 
 		if (Args.Num() < 1 || Args.Num() > 2)
 		{
@@ -758,7 +759,7 @@ static FAutoConsoleCommand CVarRemoveValueTagFromEnum(
 	FConsoleCommandWithArgsDelegate::CreateLambda([](const TArray<FString>& Args)
 	{
 		using namespace UE::Editor::DataStorage;
-		IEditorDataStorageProvider* DataStorage = GetMutableDataStorageFeature<IEditorDataStorageProvider>(StorageFeatureName);
+		ICoreProvider* DataStorage = GetMutableDataStorageFeature<ICoreProvider>(StorageFeatureName);
 
 		if (Args.Num() != 1)
 		{
@@ -780,7 +781,7 @@ static FAutoConsoleCommand CVarMatchValueTagFromEnum(
 	{
 		using namespace UE::Editor::DataStorage::Queries;
 		
-		IEditorDataStorageProvider* DataStorage = GetMutableDataStorageFeature<IEditorDataStorageProvider>(StorageFeatureName);
+		ICoreProvider* DataStorage = GetMutableDataStorageFeature<ICoreProvider>(StorageFeatureName);
 
 		if (Args.Num() > 1)
 		{
diff --git a/Engine/Plugins/Experimental/EditorDataStorage/Source/TedsCore/Private/Memento/TypedElementMementoSystem.cpp b/Engine/Plugins/Experimental/EditorDataStorage/Source/TedsCore/Private/Memento/TypedElementMementoSystem.cpp
--- a/Engine/Plugins/Experimental/EditorDataStorage/Source/TedsCore/Private/Memento/TypedElementMementoSystem.cpp
+++ b/Engine/Plugins/Experimental/EditorDataStorage/Source/TedsCore/Private/Memento/TypedElementMementoSystem.cpp
@@ -11,7 +11,7 @@
 
 namespace UE::Editor::DataStorage
 {
-	FMementoSystem::FMementoSystem(IEditorDataStorageProvider& InDataStorage)
+	FMementoSystem::FMementoSystem(ICoreProvider& InDataStorage)
 		: DataStorage(InDataStorage)
 	{
 		FScopedExclusiveLock Lock(EGlobalLockScope::Public);
diff --git a/Engine/Plugins/Experimental/EditorDataStorage/Source/TedsCore/Private/Memento/TypedElementMementoSystem.h b/Engine/Plugins/Experimental/EditorDataStorage/Source/TedsCore/Private/Memento/TypedElementMementoSystem.h
--- a/Engine/Plugins/Experimental/EditorDataStorage/Source/TedsCore/Private/Memento/TypedElementMementoSystem.h
+++ b/Engine/Plugins/Experimental/EditorDataStorage/Source/TedsCore/Private/Memento/TypedElementMementoSystem.h
@@ -5,15 +5,16 @@
 #include "Containers/Array.h"
 #include "Elements/Common/TypedElementHandles.h"
 
-class IEditorDataStorageProvider;
 class UTedsMementoTranslatorBase;
 
 namespace UE::Editor::DataStorage
 {
+	class ICoreProvider;
+
 	class FMementoSystem
 	{
 	public:
-		explicit FMementoSystem(IEditorDataStorageProvider& InDataStorage);
+		explicit FMementoSystem(ICoreProvider& InDataStorage);
 
 		RowHandle CreateMemento(RowHandle SourceRow);
 		void CreateMemento(RowHandle ReservedMementoRow, RowHandle SourceRow);
@@ -25,6 +26,6 @@ namespace UE::Editor::DataStorage
 
 		TArray<const UTedsMementoTranslatorBase*> MementoTranslators;
 		TableHandle MementoRowBaseTable;
-		IEditorDataStorageProvider& DataStorage;
+		ICoreProvider& DataStorage;
 	};
 } // namespace UE::Editor::DataStorage
