[TEDS] Cleanup aliases, move core interface classes into namespace, and fix up one class name that wasn't following standards

diff --git a/Engine/Plugins/Experimental/EditorDataStorage/Source/TedsCore/Private/Processors/TypedElementProcessorAdaptors.cpp b/Engine/Plugins/Experimental/EditorDataStorage/Source/TedsCore/Private/Processors/TypedElementProcessorAdaptors.cpp
--- a/Engine/Plugins/Experimental/EditorDataStorage/Source/TedsCore/Private/Processors/TypedElementProcessorAdaptors.cpp
+++ b/Engine/Plugins/Experimental/EditorDataStorage/Source/TedsCore/Private/Processors/TypedElementProcessorAdaptors.cpp
@@ -70,7 +70,7 @@ namespace UE::Editor::DataStorage
 				for (int32 Index = 0; Index < TypeCount; ++Index)
 				{
 					checkf(ColumnTypes->IsValid(), TEXT("Attempting to retrieve a column that is not available."));
-					*RetrievedAddresses = *AccessTypes == IEditorDataStorageProvider::EQueryAccessType::ReadWrite
+					*RetrievedAddresses = *AccessTypes == EQueryAccessType::ReadWrite
 						? reinterpret_cast<char*>(Context.GetMutableFragmentView(ColumnTypes->Get()).GetData())
 						: const_cast<char*>(reinterpret_cast<const char*>(Context.GetFragmentView(ColumnTypes->Get()).GetData()));
 
@@ -1070,7 +1070,7 @@ namespace UE::Editor::DataStorage
 			TChunkedArray<RowHandle> RemovedRowsQueue;
 		};
 
-		struct FMassDirectContextForwarder final : public IEditorDataStorageProvider::IDirectQueryContext
+		struct FMassDirectContextForwarder final : public IDirectQueryContext
 		{
 			FMassDirectContextForwarder(FMassExecutionContext& InContext, FEnvironment& InEnvironment)
 				: Implementation(InContext, InEnvironment)
@@ -1090,7 +1090,7 @@ namespace UE::Editor::DataStorage
 			FMassContextCommon Implementation;
 		};
 
-		struct FMassSubqueryContextForwarder  : public IEditorDataStorageProvider::ISubqueryContext
+		struct FMassSubqueryContextForwarder  : public ISubqueryContext
 		{
 			FMassSubqueryContextForwarder(FMassExecutionContext& InContext, FEnvironment& InEnvironment, bool bInBatchDeferredCommands)
 				: Implementation(InContext, InEnvironment, bInBatchDeferredCommands)
@@ -1133,7 +1133,7 @@ namespace UE::Editor::DataStorage
 		struct FMassQueryContextImplementation final : FMassWithEnvironmentContextCommon
 		{
 			FMassQueryContextImplementation(
-				IEditorDataStorageProvider::FQueryDescription& InQueryDescription,
+				FQueryDescription& InQueryDescription,
 				FMassExecutionContext& InContext, 
 				FExtendedQueryStore& InQueryStore,
 				FEnvironment& InEnvironment)
@@ -1157,7 +1157,7 @@ namespace UE::Editor::DataStorage
 			}
 	
 			void GetDependencies(TArrayView<UObject*> RetrievedAddresses, TConstArrayView<TWeakObjectPtr<const UClass>> SubsystemTypes,
-				TConstArrayView<IEditorDataStorageProvider::EQueryAccessType> AccessTypes)
+				TConstArrayView<EQueryAccessType> AccessTypes)
 			{
 				checkf(RetrievedAddresses.Num() == SubsystemTypes.Num(), TEXT("Unable to retrieve a batch of subsystem as the number of addresses "
 					"doesn't match the number of requested subsystem types."));
@@ -1166,12 +1166,12 @@ namespace UE::Editor::DataStorage
 			}
 
 			void GetDependenciesUnguarded(int32 SubsystemCount, UObject** RetrievedAddresses, const TWeakObjectPtr<const UClass>* DependencyTypes,
-				const IEditorDataStorageProvider::EQueryAccessType* AccessTypes)
+				const EQueryAccessType* AccessTypes)
 			{
 				for (int32 Index = 0; Index < SubsystemCount; ++Index)
 				{
 					checkf(DependencyTypes->IsValid(), TEXT("Attempting to retrieve a subsystem that's no longer valid."));
-					*RetrievedAddresses = *AccessTypes == IEditorDataStorageProvider::EQueryAccessType::ReadWrite
+					*RetrievedAddresses = *AccessTypes == EQueryAccessType::ReadWrite
 						? Context.GetMutableSubsystem<USubsystem>(const_cast<UClass*>(DependencyTypes->Get()))
 						: const_cast<USubsystem*>(Context.GetSubsystem<USubsystem>(const_cast<UClass*>(DependencyTypes->Get())));
 
@@ -1232,14 +1232,14 @@ namespace UE::Editor::DataStorage
 				}
 			}
 	
-			IEditorDataStorageProvider::FQueryDescription& QueryDescription;
+			FQueryDescription& QueryDescription;
 			FExtendedQueryStore& QueryStore;
 		};
 
-		struct FMassContextForwarder final : public IEditorDataStorageProvider::IQueryContext
+		struct FMassContextForwarder final : public IQueryContext
 		{
 			FMassContextForwarder(
-				IEditorDataStorageProvider::FQueryDescription& InQueryDescription,
+				FQueryDescription& InQueryDescription,
 				FMassExecutionContext& InContext, 
 				FExtendedQueryStore& InQueryStore,
 				FEnvironment& InEnvironment)
@@ -1305,11 +1305,11 @@ protected:
 	}
 
 	void FPhasePreOrPostAmbleExecutor::ExecuteQuery(
-		IEditorDataStorageProvider::FQueryDescription& Description,
+		FQueryDescription& Description,
 		FExtendedQueryStore& QueryStore,
 		FEnvironment& Environment,
 		FMassEntityQuery& NativeQuery,
-		IEditorDataStorageProvider::QueryCallbackRef Callback)
+		QueryCallbackRef Callback)
 	{
 		if (Description.Callback.ActivationCount > 0)
 		{
@@ -1338,10 +1338,10 @@ FTypedElementQueryProcessorData::FTypedElementQueryProcessorData(UMassProcessor&
 
 bool FTypedElementQueryProcessorData::CommonQueryConfiguration(
 	UMassProcessor& InOwner,
-	FExtendedQuery& InQuery,
-	FExtendedQueryStore::Handle InQueryHandle,
-	FExtendedQueryStore& InQueryStore,
-	FEnvironment& InEnvironment,
+	UE::Editor::DataStorage::FExtendedQuery& InQuery,
+	UE::Editor::DataStorage::FExtendedQueryStore::Handle InQueryHandle,
+	UE::Editor::DataStorage::FExtendedQueryStore& InQueryStore,
+	UE::Editor::DataStorage::FEnvironment& InEnvironment,
 	TArrayView<FMassEntityQuery> Subqueries)
 {
 	using namespace UE::Editor::DataStorage;
@@ -1382,17 +1382,19 @@ bool FTypedElementQueryProcessorData::CommonQueryConfiguration(
 	return false;
 }
 
-EMassProcessingPhase FTypedElementQueryProcessorData::MapToMassProcessingPhase(IEditorDataStorageProvider::EQueryTickPhase Phase)
+EMassProcessingPhase FTypedElementQueryProcessorData::MapToMassProcessingPhase(UE::Editor::DataStorage::EQueryTickPhase Phase)
 {
+	using namespace UE::Editor::DataStorage;
+
 	switch(Phase)
 	{
-	case IEditorDataStorageProvider::EQueryTickPhase::PrePhysics:
+	case EQueryTickPhase::PrePhysics:
 		return EMassProcessingPhase::PrePhysics;
-	case IEditorDataStorageProvider::EQueryTickPhase::DuringPhysics:
+	case EQueryTickPhase::DuringPhysics:
 		return EMassProcessingPhase::DuringPhysics;
-	case IEditorDataStorageProvider::EQueryTickPhase::PostPhysics:
+	case EQueryTickPhase::PostPhysics:
 		return EMassProcessingPhase::PostPhysics;
-	case IEditorDataStorageProvider::EQueryTickPhase::FrameEnd:
+	case EQueryTickPhase::FrameEnd:
 		return EMassProcessingPhase::FrameEnd;
 	default:
 		checkf(false, TEXT("Query tick phase '%i' is unsupported."), static_cast<int>(Phase));
@@ -1402,6 +1404,8 @@ EMassProcessingPhase FTypedElementQueryProcessorData::MapToMassProcessingPhase(I
 
 FString FTypedElementQueryProcessorData::GetProcessorName() const
 {
+	using namespace UE::Editor::DataStorage;
+
 	if (const FExtendedQuery* StoredQuery = QueryStore ? QueryStore->Get(ParentQuery) : nullptr)
 	{
 		return StoredQuery->Description.Callback.Name.ToString();
@@ -1419,8 +1423,8 @@ void FTypedElementQueryProcessorData::DebugOutputDescription(FOutputDevice& Ar,
 
 	if (const FExtendedQuery* StoredQuery = QueryStore ? QueryStore->Get(ParentQuery) : nullptr)
 	{
-		const IEditorDataStorageProvider::FQueryDescription& Description = StoredQuery->Description;
-		const IEditorDataStorageProvider::FQueryDescription::FCallbackData& Callback = Description.Callback;
+		const FQueryDescription& Description = StoredQuery->Description;
+		const FQueryDescription::FCallbackData& Callback = Description.Callback;
 		
 		if (!Callback.Group.IsNone())
 		{
@@ -1478,20 +1482,22 @@ void FTypedElementQueryProcessorData::DebugOutputDescription(FOutputDevice& Ar,
 }
 
 bool FTypedElementQueryProcessorData::PrepareCachedDependenciesOnQuery(
-	IEditorDataStorageProvider::FQueryDescription& Description, FMassExecutionContext& Context)
+	UE::Editor::DataStorage::FQueryDescription& Description, FMassExecutionContext& Context)
 {
+	using namespace UE::Editor::DataStorage;
+
 	const int32 DependencyCount = Description.DependencyTypes.Num();
 	const TWeakObjectPtr<const UClass>* Types = Description.DependencyTypes.GetData();
-	const IEditorDataStorageProvider::EQueryDependencyFlags* Flags = Description.DependencyFlags.GetData();
+	const EQueryDependencyFlags* Flags = Description.DependencyFlags.GetData();
 	TWeakObjectPtr<UObject>* Caches = Description.CachedDependencies.GetData();
 
 	for (int32 Index = 0; Index < DependencyCount; ++Index)
 	{
 		checkf(Types->IsValid(), TEXT("Attempting to retrieve a dependency type that's no longer available."));
 		
-		if (EnumHasAnyFlags(*Flags, IEditorDataStorageProvider::EQueryDependencyFlags::AlwaysRefresh) || !Caches->IsValid())
+		if (EnumHasAnyFlags(*Flags, EQueryDependencyFlags::AlwaysRefresh) || !Caches->IsValid())
 		{
-			*Caches = EnumHasAnyFlags(*Flags, IEditorDataStorageProvider::EQueryDependencyFlags::ReadOnly)
+			*Caches = EnumHasAnyFlags(*Flags, EQueryDependencyFlags::ReadOnly)
 				? const_cast<USubsystem*>(Context.GetSubsystem<USubsystem>(const_cast<UClass*>(Types->Get())))
 				: Context.GetMutableSubsystem<USubsystem>(const_cast<UClass*>(Types->Get()));
 			if (*Caches != nullptr)
@@ -1515,7 +1521,7 @@ UE::Editor::DataStorage::FQueryResult FTypedElementQueryProcessorData::Execute(
 	UE::Editor::DataStorage::FQueryDescription& Description,
 	FMassEntityQuery& NativeQuery, 
 	FMassEntityManager& EntityManager,
-	FEnvironment& Environment,
+	UE::Editor::DataStorage::FEnvironment& Environment,
 	UE::Editor::DataStorage::EDirectQueryExecutionFlags ExecutionFlags)
 {
 	using namespace UE::Editor::DataStorage;
@@ -1561,13 +1567,13 @@ UE::Editor::DataStorage::FQueryResult FTypedElementQueryProcessorData::Execute(
 	UE::Editor::DataStorage::FQueryDescription& Description,
 	FMassEntityQuery& NativeQuery,
 	FMassEntityManager& EntityManager,
-	FEnvironment& Environment,
+	UE::Editor::DataStorage::FEnvironment& Environment,
 	FMassExecutionContext& ParentContext)
 {
 	using namespace UE::Editor::DataStorage;
 
-	IEditorDataStorageProvider::FQueryResult Result;
-	Result.Completed = IEditorDataStorageProvider::FQueryResult::ECompletion::Fully;
+	FQueryResult Result;
+	Result.Completed = FQueryResult::ECompletion::Fully;
 
 	if (Description.Callback.ActivationCount > 0)
 	{
@@ -1597,13 +1603,13 @@ UE::Editor::DataStorage::FQueryResult FTypedElementQueryProcessorData::Execute(
 	UE::Editor::DataStorage::RowHandle RowHandle,
 	FMassEntityQuery& NativeQuery,
 	FMassEntityManager& EntityManager,
-	FEnvironment& Environment,
+	UE::Editor::DataStorage::FEnvironment& Environment,
 	FMassExecutionContext& ParentContext)
 {
 	using namespace UE::Editor::DataStorage;
 
-	IEditorDataStorageProvider::FQueryResult Result;
-	Result.Completed = IEditorDataStorageProvider::FQueryResult::ECompletion::Fully;
+	FQueryResult Result;
+	Result.Completed = FQueryResult::ECompletion::Fully;
 
 	FMassEntityHandle NativeEntity = FMassEntityHandle::FromNumber(RowHandle);
 	if (Description.Callback.ActivationCount > 0 && EntityManager.IsEntityActive(NativeEntity))
@@ -1639,7 +1645,7 @@ void FTypedElementQueryProcessorData::Execute(FMassEntityManager& EntityManager,
 
 	checkf(StoredQuery, TEXT("A query callback was registered for execution without an associated query. Processor: [%s]"), *GetProcessorName());
 	
-	IEditorDataStorageProvider::FQueryDescription& Description = StoredQuery->Description;
+	FQueryDescription& Description = StoredQuery->Description;
 	if (Description.Callback.ActivationCount > 0)
 	{
 		auto ExceteFunction = [this, &Description](FMassExecutionContext& Context)
@@ -1681,10 +1687,10 @@ FMassEntityQuery& UTypedElementQueryProcessorCallbackAdapterProcessorBase::GetQu
 }
 
 bool UTypedElementQueryProcessorCallbackAdapterProcessorBase::ConfigureQueryCallback(
-	FExtendedQuery& Query,
-	FExtendedQueryStore::Handle QueryHandle,
-	FExtendedQueryStore& QueryStore,
-	FEnvironment& Environment)
+	UE::Editor::DataStorage::FExtendedQuery& Query,
+	UE::Editor::DataStorage::FExtendedQueryStore::Handle QueryHandle,
+	UE::Editor::DataStorage::FExtendedQueryStore& QueryStore,
+	UE::Editor::DataStorage::FEnvironment& Environment)
 {
 	return ConfigureQueryCallbackData(Query, QueryHandle, QueryStore, Environment, {});
 }
@@ -1699,10 +1705,10 @@ bool UTypedElementQueryProcessorCallbackAdapterProcessorBase::ShouldAllowQueryBa
 }
 
 bool UTypedElementQueryProcessorCallbackAdapterProcessorBase::ConfigureQueryCallbackData(
-	FExtendedQuery& Query,
-	FExtendedQueryStore::Handle QueryHandle,
-	FExtendedQueryStore& QueryStore,
-	FEnvironment& Environment,
+	UE::Editor::DataStorage::FExtendedQuery& Query,
+	UE::Editor::DataStorage::FExtendedQueryStore::Handle QueryHandle,
+	UE::Editor::DataStorage::FExtendedQueryStore& QueryStore,
+	UE::Editor::DataStorage::FEnvironment& Environment,
 	TArrayView<FMassEntityQuery> Subqueries)
 {
 	using namespace UE::Editor::DataStorage;
@@ -1750,49 +1756,73 @@ void UTypedElementQueryProcessorCallbackAdapterProcessorBase::Execute(FMassEntit
 }
 
 bool UTypedElementQueryProcessorCallbackAdapterProcessorWith1Subquery::ConfigureQueryCallback(
-	FExtendedQuery& Query, FExtendedQueryStore::Handle QueryHandle, FExtendedQueryStore& QueryStore, FEnvironment& Environment)
+	UE::Editor::DataStorage::FExtendedQuery& Query,
+	UE::Editor::DataStorage::FExtendedQueryStore::Handle QueryHandle,
+	UE::Editor::DataStorage::FExtendedQueryStore& QueryStore,
+	UE::Editor::DataStorage::FEnvironment& Environment)
 {
 	return ConfigureQueryCallbackData(Query, QueryHandle, QueryStore, Environment, NativeSubqueries);
 }
 
 bool UTypedElementQueryProcessorCallbackAdapterProcessorWith2Subqueries::ConfigureQueryCallback(
-	FExtendedQuery& Query, FExtendedQueryStore::Handle QueryHandle, FExtendedQueryStore& QueryStore, FEnvironment& Environment)
+	UE::Editor::DataStorage::FExtendedQuery& Query,
+	UE::Editor::DataStorage::FExtendedQueryStore::Handle QueryHandle,
+	UE::Editor::DataStorage::FExtendedQueryStore& QueryStore,
+	UE::Editor::DataStorage::FEnvironment& Environment)
 {
 	return ConfigureQueryCallbackData(Query, QueryHandle, QueryStore, Environment, NativeSubqueries);
 }
 
 bool UTypedElementQueryProcessorCallbackAdapterProcessorWith3Subqueries::ConfigureQueryCallback(
-	FExtendedQuery& Query, FExtendedQueryStore::Handle QueryHandle, FExtendedQueryStore& QueryStore, FEnvironment& Environment)
+	UE::Editor::DataStorage::FExtendedQuery& Query,
+	UE::Editor::DataStorage::FExtendedQueryStore::Handle QueryHandle,
+	UE::Editor::DataStorage::FExtendedQueryStore& QueryStore,
+	UE::Editor::DataStorage::FEnvironment& Environment)
 {
 	return ConfigureQueryCallbackData(Query, QueryHandle, QueryStore, Environment, NativeSubqueries);
 }
 
 bool UTypedElementQueryProcessorCallbackAdapterProcessorWith4Subqueries::ConfigureQueryCallback(
-	FExtendedQuery& Query, FExtendedQueryStore::Handle QueryHandle, FExtendedQueryStore& QueryStore, FEnvironment& Environment)
+	UE::Editor::DataStorage::FExtendedQuery& Query,
+	UE::Editor::DataStorage::FExtendedQueryStore::Handle QueryHandle,
+	UE::Editor::DataStorage::FExtendedQueryStore& QueryStore,
+	UE::Editor::DataStorage::FEnvironment& Environment)
 {
 	return ConfigureQueryCallbackData(Query, QueryHandle, QueryStore, Environment, NativeSubqueries);
 }
 
 bool UTypedElementQueryProcessorCallbackAdapterProcessorWith5Subqueries::ConfigureQueryCallback(
-	FExtendedQuery& Query, FExtendedQueryStore::Handle QueryHandle, FExtendedQueryStore& QueryStore, FEnvironment& Environment)
+	UE::Editor::DataStorage::FExtendedQuery& Query,
+	UE::Editor::DataStorage::FExtendedQueryStore::Handle QueryHandle,
+	UE::Editor::DataStorage::FExtendedQueryStore& QueryStore,
+	UE::Editor::DataStorage::FEnvironment& Environment)
 {
 	return ConfigureQueryCallbackData(Query, QueryHandle, QueryStore, Environment, NativeSubqueries);
 }
 
 bool UTypedElementQueryProcessorCallbackAdapterProcessorWith6Subqueries::ConfigureQueryCallback(
-	FExtendedQuery& Query, FExtendedQueryStore::Handle QueryHandle, FExtendedQueryStore& QueryStore, FEnvironment& Environment)
+	UE::Editor::DataStorage::FExtendedQuery& Query,
+	UE::Editor::DataStorage::FExtendedQueryStore::Handle QueryHandle,
+	UE::Editor::DataStorage::FExtendedQueryStore& QueryStore,
+	UE::Editor::DataStorage::FEnvironment& Environment)
 {
 	return ConfigureQueryCallbackData(Query, QueryHandle, QueryStore, Environment, NativeSubqueries);
 }
 
 bool UTypedElementQueryProcessorCallbackAdapterProcessorWith7Subqueries::ConfigureQueryCallback(
-	FExtendedQuery& Query, FExtendedQueryStore::Handle QueryHandle, FExtendedQueryStore& QueryStore, FEnvironment& Environment)
+	UE::Editor::DataStorage::FExtendedQuery& Query,
+	UE::Editor::DataStorage::FExtendedQueryStore::Handle QueryHandle,
+	UE::Editor::DataStorage::FExtendedQueryStore& QueryStore,
+	UE::Editor::DataStorage::FEnvironment& Environment)
 {
 	return ConfigureQueryCallbackData(Query, QueryHandle, QueryStore, Environment, NativeSubqueries);
 }
 
 bool UTypedElementQueryProcessorCallbackAdapterProcessorWith8Subqueries::ConfigureQueryCallback(
-	FExtendedQuery& Query, FExtendedQueryStore::Handle QueryHandle, FExtendedQueryStore& QueryStore, FEnvironment& Environment)
+	UE::Editor::DataStorage::FExtendedQuery& Query,
+	UE::Editor::DataStorage::FExtendedQueryStore::Handle QueryHandle,
+	UE::Editor::DataStorage::FExtendedQueryStore& QueryStore,
+	UE::Editor::DataStorage::FEnvironment& Environment)
 {
 	return ConfigureQueryCallbackData(Query, QueryHandle, QueryStore, Environment, NativeSubqueries);
 }
@@ -1825,16 +1855,19 @@ EMassObservedOperation UTypedElementQueryObserverCallbackAdapterProcessorBase::G
 }
 
 bool UTypedElementQueryObserverCallbackAdapterProcessorBase::ConfigureQueryCallback(
-	FExtendedQuery& Query, FExtendedQueryStore::Handle QueryHandle, FExtendedQueryStore& QueryStore, FEnvironment& Environment)
+	UE::Editor::DataStorage::FExtendedQuery& Query,
+	UE::Editor::DataStorage::FExtendedQueryStore::Handle QueryHandle,
+	UE::Editor::DataStorage::FExtendedQueryStore& QueryStore,
+	UE::Editor::DataStorage::FEnvironment& Environment)
 {
 	return ConfigureQueryCallbackData(Query, QueryHandle, QueryStore, Environment, {});
 }
 
 bool UTypedElementQueryObserverCallbackAdapterProcessorBase::ConfigureQueryCallbackData(
-	FExtendedQuery& Query,
-	FExtendedQueryStore::Handle QueryHandle,
-	FExtendedQueryStore& QueryStore,
-	FEnvironment& Environment, TArrayView<FMassEntityQuery> Subqueries)
+	UE::Editor::DataStorage::FExtendedQuery& Query,
+	UE::Editor::DataStorage::FExtendedQueryStore::Handle QueryHandle,
+	UE::Editor::DataStorage::FExtendedQueryStore& QueryStore,
+	UE::Editor::DataStorage::FEnvironment& Environment, TArrayView<FMassEntityQuery> Subqueries)
 {
 	using namespace UE::Editor::DataStorage;
 
@@ -1847,10 +1880,10 @@ bool UTypedElementQueryObserverCallbackAdapterProcessorBase::ConfigureQueryCallb
 	
 	switch (Query.Description.Callback.Type)
 	{
-	case IEditorDataStorageProvider::EQueryCallbackType::ObserveAdd:
+	case EQueryCallbackType::ObserveAdd:
 		Operation = EMassObservedOperation::Add;
 		break;
-	case IEditorDataStorageProvider::EQueryCallbackType::ObserveRemove:
+	case EQueryCallbackType::ObserveRemove:
 		Operation = EMassObservedOperation::Remove;
 		break;
 	default:
@@ -1911,49 +1944,73 @@ void UTypedElementQueryObserverCallbackAdapterProcessorBase::Execute(FMassEntity
 }
 
 bool UTypedElementQueryObserverCallbackAdapterProcessorWith1Subquery::ConfigureQueryCallback(
-	FExtendedQuery& Query, FExtendedQueryStore::Handle QueryHandle, FExtendedQueryStore& QueryStore, FEnvironment& Environment)
+	UE::Editor::DataStorage::FExtendedQuery& Query,
+	UE::Editor::DataStorage::FExtendedQueryStore::Handle QueryHandle,
+	UE::Editor::DataStorage::FExtendedQueryStore& QueryStore,
+	UE::Editor::DataStorage::FEnvironment& Environment)
 {
 	return ConfigureQueryCallbackData(Query, QueryHandle, QueryStore, Environment, NativeSubqueries);
 }
 
 bool UTypedElementQueryObserverCallbackAdapterProcessorWith2Subqueries::ConfigureQueryCallback(
-	FExtendedQuery& Query, FExtendedQueryStore::Handle QueryHandle, FExtendedQueryStore& QueryStore, FEnvironment& Environment)
+	UE::Editor::DataStorage::FExtendedQuery& Query,
+	UE::Editor::DataStorage::FExtendedQueryStore::Handle QueryHandle,
+	UE::Editor::DataStorage::FExtendedQueryStore& QueryStore,
+	UE::Editor::DataStorage::FEnvironment& Environment)
 {
 	return ConfigureQueryCallbackData(Query, QueryHandle, QueryStore, Environment, NativeSubqueries);
 }
 
 bool UTypedElementQueryObserverCallbackAdapterProcessorWith3Subqueries::ConfigureQueryCallback(
-	FExtendedQuery& Query, FExtendedQueryStore::Handle QueryHandle, FExtendedQueryStore& QueryStore, FEnvironment& Environment)
+	UE::Editor::DataStorage::FExtendedQuery& Query,
+	UE::Editor::DataStorage::FExtendedQueryStore::Handle QueryHandle,
+	UE::Editor::DataStorage::FExtendedQueryStore& QueryStore,
+	UE::Editor::DataStorage::FEnvironment& Environment)
 {
 	return ConfigureQueryCallbackData(Query, QueryHandle, QueryStore, Environment, NativeSubqueries);
 }
 
 bool UTypedElementQueryObserverCallbackAdapterProcessorWith4Subqueries::ConfigureQueryCallback(
-	FExtendedQuery& Query, FExtendedQueryStore::Handle QueryHandle, FExtendedQueryStore& QueryStore, FEnvironment& Environment)
+	UE::Editor::DataStorage::FExtendedQuery& Query,
+	UE::Editor::DataStorage::FExtendedQueryStore::Handle QueryHandle,
+	UE::Editor::DataStorage::FExtendedQueryStore& QueryStore,
+	UE::Editor::DataStorage::FEnvironment& Environment)
 {
 	return ConfigureQueryCallbackData(Query, QueryHandle, QueryStore, Environment, NativeSubqueries);
 }
 
 bool UTypedElementQueryObserverCallbackAdapterProcessorWith5Subqueries::ConfigureQueryCallback(
-	FExtendedQuery& Query, FExtendedQueryStore::Handle QueryHandle, FExtendedQueryStore& QueryStore, FEnvironment& Environment)
+	UE::Editor::DataStorage::FExtendedQuery& Query,
+	UE::Editor::DataStorage::FExtendedQueryStore::Handle QueryHandle,
+	UE::Editor::DataStorage::FExtendedQueryStore& QueryStore,
+	UE::Editor::DataStorage::FEnvironment& Environment)
 {
 	return ConfigureQueryCallbackData(Query, QueryHandle, QueryStore, Environment, NativeSubqueries);
 }
 
 bool UTypedElementQueryObserverCallbackAdapterProcessorWith6Subqueries::ConfigureQueryCallback(
-	FExtendedQuery& Query, FExtendedQueryStore::Handle QueryHandle, FExtendedQueryStore& QueryStore, FEnvironment& Environment)
+	UE::Editor::DataStorage::FExtendedQuery& Query,
+	UE::Editor::DataStorage::FExtendedQueryStore::Handle QueryHandle,
+	UE::Editor::DataStorage::FExtendedQueryStore& QueryStore,
+	UE::Editor::DataStorage::FEnvironment& Environment)
 {
 	return ConfigureQueryCallbackData(Query, QueryHandle, QueryStore, Environment, NativeSubqueries);
 }
 
 bool UTypedElementQueryObserverCallbackAdapterProcessorWith7Subqueries::ConfigureQueryCallback(
-	FExtendedQuery& Query, FExtendedQueryStore::Handle QueryHandle, FExtendedQueryStore& QueryStore, FEnvironment& Environment)
+	UE::Editor::DataStorage::FExtendedQuery& Query,
+	UE::Editor::DataStorage::FExtendedQueryStore::Handle QueryHandle,
+	UE::Editor::DataStorage::FExtendedQueryStore& QueryStore,
+	UE::Editor::DataStorage::FEnvironment& Environment)
 {
 	return ConfigureQueryCallbackData(Query, QueryHandle, QueryStore, Environment, NativeSubqueries);
 }
 
 bool UTypedElementQueryObserverCallbackAdapterProcessorWith8Subqueries::ConfigureQueryCallback(
-	FExtendedQuery& Query, FExtendedQueryStore::Handle QueryHandle, FExtendedQueryStore& QueryStore, FEnvironment& Environment)
+	UE::Editor::DataStorage::FExtendedQuery& Query,
+	UE::Editor::DataStorage::FExtendedQueryStore::Handle QueryHandle,
+	UE::Editor::DataStorage::FExtendedQueryStore& QueryStore,
+	UE::Editor::DataStorage::FEnvironment& Environment)
 {
 	return ConfigureQueryCallbackData(Query, QueryHandle, QueryStore, Environment, NativeSubqueries);
 }
diff --git a/Engine/Plugins/Experimental/EditorDataStorage/Source/TedsCore/Private/Processors/TypedElementProcessorAdaptors.h b/Engine/Plugins/Experimental/EditorDataStorage/Source/TedsCore/Private/Processors/TypedElementProcessorAdaptors.h
--- a/Engine/Plugins/Experimental/EditorDataStorage/Source/TedsCore/Private/Processors/TypedElementProcessorAdaptors.h
+++ b/Engine/Plugins/Experimental/EditorDataStorage/Source/TedsCore/Private/Processors/TypedElementProcessorAdaptors.h
@@ -24,11 +24,11 @@ namespace UE::Editor::DataStorage
 		~FPhasePreOrPostAmbleExecutor();
 
 		void ExecuteQuery(
-			IEditorDataStorageProvider::FQueryDescription& Description,
+			FQueryDescription& Description,
 			FExtendedQueryStore& QueryStore,
 			FEnvironment& Environment,
 			FMassEntityQuery& NativeQuery,
-			IEditorDataStorageProvider::QueryCallbackRef Callback);
+			QueryCallbackRef Callback);
 
 		FMassExecutionContext Context;
 	};
@@ -39,21 +39,17 @@ struct FTypedElementQueryProcessorData
 {
 	GENERATED_BODY();
 
-	using FExtendedQuery = UE::Editor::DataStorage::FExtendedQuery;
-	using FExtendedQueryStore = UE::Editor::DataStorage::FExtendedQueryStore;
-	using FEnvironment = UE::Editor::DataStorage::FEnvironment;
-
 	FTypedElementQueryProcessorData() = default;
 	explicit FTypedElementQueryProcessorData(UMassProcessor& Owner);
 
 	bool CommonQueryConfiguration(
 		UMassProcessor& InOwner,
-		FExtendedQuery& InQuery,
-		FExtendedQueryStore::Handle InQueryHandle,
-		FExtendedQueryStore& InQueryStore,
-		FEnvironment& InEnvironment,
+		UE::Editor::DataStorage::FExtendedQuery& InQuery,
+		UE::Editor::DataStorage::FExtendedQueryStore::Handle InQueryHandle,
+		UE::Editor::DataStorage::FExtendedQueryStore& InQueryStore,
+		UE::Editor::DataStorage::FEnvironment& InEnvironment,
 		TArrayView<FMassEntityQuery> Subqueries);
-	static EMassProcessingPhase MapToMassProcessingPhase(IEditorDataStorageProvider::EQueryTickPhase Phase);
+	static EMassProcessingPhase MapToMassProcessingPhase(UE::Editor::DataStorage::EQueryTickPhase Phase);
 	FString GetProcessorName() const;
 	void DebugOutputDescription(FOutputDevice& Ar, int32 Indent) const;
 
@@ -62,14 +58,14 @@ struct FTypedElementQueryProcessorData
 		UE::Editor::DataStorage::FQueryDescription& Description,
 		FMassEntityQuery& NativeQuery, 
 		FMassEntityManager& EntityManager,
-		FEnvironment& Environment,
+		UE::Editor::DataStorage::FEnvironment& Environment,
 		UE::Editor::DataStorage::EDirectQueryExecutionFlags ExecutionFlags);
 	static UE::Editor::DataStorage::FQueryResult Execute(
 		UE::Editor::DataStorage::SubqueryCallbackRef& Callback,
 		UE::Editor::DataStorage::FQueryDescription& Description,
 		FMassEntityQuery& NativeQuery,
 		FMassEntityManager& EntityManager,
-		FEnvironment& Environment,
+		UE::Editor::DataStorage::FEnvironment& Environment,
 		FMassExecutionContext& ParentContext);
 	static UE::Editor::DataStorage::FQueryResult Execute(
 		UE::Editor::DataStorage::SubqueryCallbackRef& Callback,
@@ -77,16 +73,16 @@ struct FTypedElementQueryProcessorData
 		UE::Editor::DataStorage::RowHandle RowHandle,
 		FMassEntityQuery& NativeQuery,
 		FMassEntityManager& EntityManager,
-		FEnvironment& Environment,
+		UE::Editor::DataStorage::FEnvironment& Environment,
 		FMassExecutionContext& ParentContext);
 	void Execute(FMassEntityManager& EntityManager, FMassExecutionContext& Context);
 
 	static bool PrepareCachedDependenciesOnQuery(
-		IEditorDataStorageProvider::FQueryDescription& Description, FMassExecutionContext& Context);
+		UE::Editor::DataStorage::FQueryDescription& Description, FMassExecutionContext& Context);
 
-	FExtendedQueryStore::Handle ParentQuery;
-	FExtendedQueryStore* QueryStore{ nullptr };
-	FEnvironment* Environment{ nullptr };
+	UE::Editor::DataStorage::FExtendedQueryStore::Handle ParentQuery;
+	UE::Editor::DataStorage::FExtendedQueryStore* QueryStore{ nullptr };
+	UE::Editor::DataStorage::FEnvironment* Environment{ nullptr };
 	FMassEntityQuery NativeQuery;
 };
 
@@ -97,29 +93,25 @@ UCLASS()
 class UTypedElementQueryProcessorCallbackAdapterProcessorBase : public UMassProcessor
 {
 	GENERATED_BODY()
-	
-	using FExtendedQuery = UE::Editor::DataStorage::FExtendedQuery;
-	using FExtendedQueryStore = UE::Editor::DataStorage::FExtendedQueryStore;
-	using FEnvironment = UE::Editor::DataStorage::FEnvironment;
 
 public:
 	UTypedElementQueryProcessorCallbackAdapterProcessorBase();
 
 	FMassEntityQuery& GetQuery();
 	virtual bool ConfigureQueryCallback(
-		FExtendedQuery& Query,
-		FExtendedQueryStore::Handle QueryHandle,
-		FExtendedQueryStore& QueryStore,
-		FEnvironment& Environment);
+		UE::Editor::DataStorage::FExtendedQuery& Query,
+		UE::Editor::DataStorage::FExtendedQueryStore::Handle QueryHandle,
+		UE::Editor::DataStorage::FExtendedQueryStore& QueryStore,
+		UE::Editor::DataStorage::FEnvironment& Environment);
 
 	virtual bool ShouldAllowQueryBasedPruning(const bool bRuntimeMode) const override;
 
 protected:
 	bool ConfigureQueryCallbackData(
-		FExtendedQuery& Query,
-		FExtendedQueryStore::Handle QueryHandle,
-		FExtendedQueryStore& QueryStore,
-		FEnvironment& Environment,
+		UE::Editor::DataStorage::FExtendedQuery& Query,
+		UE::Editor::DataStorage::FExtendedQueryStore::Handle QueryHandle,
+		UE::Editor::DataStorage::FExtendedQueryStore& QueryStore,
+		UE::Editor::DataStorage::FEnvironment& Environment,
 		TArrayView<FMassEntityQuery> Subqueries);
 	void ConfigureQueries(const TSharedRef<FMassEntityManager>&) override;
 	void Execute(FMassEntityManager& EntityManager, FMassExecutionContext& TargetParentQuery) override;
@@ -152,16 +144,12 @@ class UTypedElementQueryProcessorCallbackAdapterProcessorWith1Subquery final : p
 {
 	GENERATED_BODY()
 
-	using FExtendedQuery = UE::Editor::DataStorage::FExtendedQuery;
-	using FExtendedQueryStore = UE::Editor::DataStorage::FExtendedQueryStore;
-	using FEnvironment = UE::Editor::DataStorage::FEnvironment;
-
 public:
 	virtual bool ConfigureQueryCallback(
-		FExtendedQuery& Query,
-		FExtendedQueryStore::Handle QueryHandle,
-		FExtendedQueryStore& QueryStore,
-		FEnvironment& Environment) override;
+		UE::Editor::DataStorage::FExtendedQuery& Query,
+		UE::Editor::DataStorage::FExtendedQueryStore::Handle QueryHandle,
+		UE::Editor::DataStorage::FExtendedQueryStore& QueryStore,
+		UE::Editor::DataStorage::FEnvironment& Environment) override;
 
 private:
 	UPROPERTY(transient)
@@ -173,16 +161,12 @@ class UTypedElementQueryProcessorCallbackAdapterProcessorWith2Subqueries final :
 {
 	GENERATED_BODY()
 
-	using FExtendedQuery = UE::Editor::DataStorage::FExtendedQuery;
-	using FExtendedQueryStore = UE::Editor::DataStorage::FExtendedQueryStore;
-	using FEnvironment = UE::Editor::DataStorage::FEnvironment;
-
 public:
 	virtual bool ConfigureQueryCallback(
-		FExtendedQuery& Query,
-		FExtendedQueryStore::Handle QueryHandle,
-		FExtendedQueryStore& QueryStore,
-		FEnvironment& Environment) override;
+		UE::Editor::DataStorage::FExtendedQuery& Query,
+		UE::Editor::DataStorage::FExtendedQueryStore::Handle QueryHandle,
+		UE::Editor::DataStorage::FExtendedQueryStore& QueryStore,
+		UE::Editor::DataStorage::FEnvironment& Environment) override;
 
 private:
 	UPROPERTY(transient)
@@ -194,16 +178,12 @@ class UTypedElementQueryProcessorCallbackAdapterProcessorWith3Subqueries final :
 {
 	GENERATED_BODY()
 
-	using FExtendedQuery = UE::Editor::DataStorage::FExtendedQuery;
-	using FExtendedQueryStore = UE::Editor::DataStorage::FExtendedQueryStore;
-	using FEnvironment = UE::Editor::DataStorage::FEnvironment;
-
 public:
 	virtual bool ConfigureQueryCallback(
-		FExtendedQuery& Query,
-		FExtendedQueryStore::Handle QueryHandle,
-		FExtendedQueryStore& QueryStore,
-		FEnvironment& Environment) override;
+		UE::Editor::DataStorage::FExtendedQuery& Query,
+		UE::Editor::DataStorage::FExtendedQueryStore::Handle QueryHandle,
+		UE::Editor::DataStorage::FExtendedQueryStore& QueryStore,
+		UE::Editor::DataStorage::FEnvironment& Environment) override;
 
 private:
 	UPROPERTY(transient)
@@ -215,16 +195,12 @@ class UTypedElementQueryProcessorCallbackAdapterProcessorWith4Subqueries final :
 {
 	GENERATED_BODY()
 
-	using FExtendedQuery = UE::Editor::DataStorage::FExtendedQuery;
-	using FExtendedQueryStore = UE::Editor::DataStorage::FExtendedQueryStore;
-	using FEnvironment = UE::Editor::DataStorage::FEnvironment;
-
 public:
 	virtual bool ConfigureQueryCallback(
-		FExtendedQuery& Query,
-		FExtendedQueryStore::Handle QueryHandle,
-		FExtendedQueryStore& QueryStore,
-		FEnvironment& Environment) override;
+		UE::Editor::DataStorage::FExtendedQuery& Query,
+		UE::Editor::DataStorage::FExtendedQueryStore::Handle QueryHandle,
+		UE::Editor::DataStorage::FExtendedQueryStore& QueryStore,
+		UE::Editor::DataStorage::FEnvironment& Environment) override;
 
 private:
 	UPROPERTY(transient)
@@ -236,16 +212,12 @@ class UTypedElementQueryProcessorCallbackAdapterProcessorWith5Subqueries final :
 {
 	GENERATED_BODY()
 
-	using FExtendedQuery = UE::Editor::DataStorage::FExtendedQuery;
-	using FExtendedQueryStore = UE::Editor::DataStorage::FExtendedQueryStore;
-	using FEnvironment = UE::Editor::DataStorage::FEnvironment;
-
 public:
 	virtual bool ConfigureQueryCallback(
-		FExtendedQuery& Query,
-		FExtendedQueryStore::Handle QueryHandle,
-		FExtendedQueryStore& QueryStore,
-		FEnvironment& Environment) override;
+		UE::Editor::DataStorage::FExtendedQuery& Query,
+		UE::Editor::DataStorage::FExtendedQueryStore::Handle QueryHandle,
+		UE::Editor::DataStorage::FExtendedQueryStore& QueryStore,
+		UE::Editor::DataStorage::FEnvironment& Environment) override;
 
 private:
 	UPROPERTY(transient)
@@ -257,16 +229,12 @@ class UTypedElementQueryProcessorCallbackAdapterProcessorWith6Subqueries final :
 {
 	GENERATED_BODY()
 
-	using FExtendedQuery = UE::Editor::DataStorage::FExtendedQuery;
-	using FExtendedQueryStore = UE::Editor::DataStorage::FExtendedQueryStore;
-	using FEnvironment = UE::Editor::DataStorage::FEnvironment;
-
 public:
 	virtual bool ConfigureQueryCallback(
-		FExtendedQuery& Query,
-		FExtendedQueryStore::Handle QueryHandle,
-		FExtendedQueryStore& QueryStore,
-		FEnvironment& Environment) override;
+		UE::Editor::DataStorage::FExtendedQuery& Query,
+		UE::Editor::DataStorage::FExtendedQueryStore::Handle QueryHandle,
+		UE::Editor::DataStorage::FExtendedQueryStore& QueryStore,
+		UE::Editor::DataStorage::FEnvironment& Environment) override;
 
 private:
 	UPROPERTY(transient)
@@ -278,16 +246,12 @@ class UTypedElementQueryProcessorCallbackAdapterProcessorWith7Subqueries final :
 {
 	GENERATED_BODY()
 
-	using FExtendedQuery = UE::Editor::DataStorage::FExtendedQuery;
-	using FExtendedQueryStore = UE::Editor::DataStorage::FExtendedQueryStore;
-	using FEnvironment = UE::Editor::DataStorage::FEnvironment;
-
 public:
 	virtual bool ConfigureQueryCallback(
-		FExtendedQuery& Query,
-		FExtendedQueryStore::Handle QueryHandle,
-		FExtendedQueryStore& QueryStore,
-		FEnvironment& Environment) override;
+		UE::Editor::DataStorage::FExtendedQuery& Query,
+		UE::Editor::DataStorage::FExtendedQueryStore::Handle QueryHandle,
+		UE::Editor::DataStorage::FExtendedQueryStore& QueryStore,
+		UE::Editor::DataStorage::FEnvironment& Environment) override;
 
 private:
 	UPROPERTY(transient)
@@ -299,16 +263,12 @@ class UTypedElementQueryProcessorCallbackAdapterProcessorWith8Subqueries final :
 {
 	GENERATED_BODY()
 
-	using FExtendedQuery = UE::Editor::DataStorage::FExtendedQuery;
-	using FExtendedQueryStore = UE::Editor::DataStorage::FExtendedQueryStore;
-	using FEnvironment = UE::Editor::DataStorage::FEnvironment;
-
 public:
 	virtual bool ConfigureQueryCallback(
-		FExtendedQuery& Query,
-		FExtendedQueryStore::Handle QueryHandle,
-		FExtendedQueryStore& QueryStore,
-		FEnvironment& Environment) override;
+		UE::Editor::DataStorage::FExtendedQuery& Query,
+		UE::Editor::DataStorage::FExtendedQueryStore::Handle QueryHandle,
+		UE::Editor::DataStorage::FExtendedQueryStore& QueryStore,
+		UE::Editor::DataStorage::FEnvironment& Environment) override;
 
 private:
 	UPROPERTY(transient)
@@ -323,10 +283,6 @@ class UTypedElementQueryObserverCallbackAdapterProcessorBase : public UMassObser
 {
 	GENERATED_BODY()
 
-	using FExtendedQuery = UE::Editor::DataStorage::FExtendedQuery;
-	using FExtendedQueryStore = UE::Editor::DataStorage::FExtendedQueryStore;
-	using FEnvironment = UE::Editor::DataStorage::FEnvironment;
-
 public:
 	UTypedElementQueryObserverCallbackAdapterProcessorBase();
 
@@ -334,17 +290,17 @@ public:
 	const UScriptStruct* GetObservedType() const;
 	EMassObservedOperation GetObservedOperation() const;
 	virtual bool ConfigureQueryCallback(
-		FExtendedQuery& Query,
-		FExtendedQueryStore::Handle QueryHandle,
-		FExtendedQueryStore& QueryStore,
-		FEnvironment& Environment);
+		UE::Editor::DataStorage::FExtendedQuery& Query,
+		UE::Editor::DataStorage::FExtendedQueryStore::Handle QueryHandle,
+		UE::Editor::DataStorage::FExtendedQueryStore& QueryStore,
+		UE::Editor::DataStorage::FEnvironment& Environment);
 
 protected:
 	bool ConfigureQueryCallbackData(
-		FExtendedQuery& Query,
-		FExtendedQueryStore::Handle QueryHandle,
-		FExtendedQueryStore& QueryStore,
-		FEnvironment& Environment,
+		UE::Editor::DataStorage::FExtendedQuery& Query,
+		UE::Editor::DataStorage::FExtendedQueryStore::Handle QueryHandle,
+		UE::Editor::DataStorage::FExtendedQueryStore& QueryStore,
+		UE::Editor::DataStorage::FEnvironment& Environment,
 		TArrayView<FMassEntityQuery> Subqueries);
 	virtual void ConfigureQueries(const TSharedRef<FMassEntityManager>& EntityManager) override;
 	virtual void Execute(FMassEntityManager& EntityManager, FMassExecutionContext& TargetParentQuery) override;
@@ -378,16 +334,12 @@ class UTypedElementQueryObserverCallbackAdapterProcessorWith1Subquery final : pu
 {
 	GENERATED_BODY()
 
-	using FExtendedQuery = UE::Editor::DataStorage::FExtendedQuery;
-	using FExtendedQueryStore = UE::Editor::DataStorage::FExtendedQueryStore;
-	using FEnvironment = UE::Editor::DataStorage::FEnvironment;
-
 public:
 	virtual bool ConfigureQueryCallback(
-		FExtendedQuery& Query,
-		FExtendedQueryStore::Handle QueryHandle,
-		FExtendedQueryStore& QueryStore,
-		FEnvironment& Environment) override;
+		UE::Editor::DataStorage::FExtendedQuery& Query,
+		UE::Editor::DataStorage::FExtendedQueryStore::Handle QueryHandle,
+		UE::Editor::DataStorage::FExtendedQueryStore& QueryStore,
+		UE::Editor::DataStorage::FEnvironment& Environment) override;
 
 private:
 	UPROPERTY(transient)
@@ -399,16 +351,12 @@ class UTypedElementQueryObserverCallbackAdapterProcessorWith2Subqueries final :
 {
 	GENERATED_BODY()
 
-	using FExtendedQuery = UE::Editor::DataStorage::FExtendedQuery;
-	using FExtendedQueryStore = UE::Editor::DataStorage::FExtendedQueryStore;
-	using FEnvironment = UE::Editor::DataStorage::FEnvironment;
-
 public:
 	virtual bool ConfigureQueryCallback(
-		FExtendedQuery& Query,
-		FExtendedQueryStore::Handle QueryHandle,
-		FExtendedQueryStore& QueryStore,
-		FEnvironment& Environment) override;
+		UE::Editor::DataStorage::FExtendedQuery& Query,
+		UE::Editor::DataStorage::FExtendedQueryStore::Handle QueryHandle,
+		UE::Editor::DataStorage::FExtendedQueryStore& QueryStore,
+		UE::Editor::DataStorage::FEnvironment& Environment) override;
 
 private:
 	UPROPERTY(transient)
@@ -420,16 +368,12 @@ class UTypedElementQueryObserverCallbackAdapterProcessorWith3Subqueries final :
 {
 	GENERATED_BODY()
 
-	using FExtendedQuery = UE::Editor::DataStorage::FExtendedQuery;
-	using FExtendedQueryStore = UE::Editor::DataStorage::FExtendedQueryStore;
-	using FEnvironment = UE::Editor::DataStorage::FEnvironment;
-
 public:
 	virtual bool ConfigureQueryCallback(
-		FExtendedQuery& Query,
-		FExtendedQueryStore::Handle QueryHandle,
-		FExtendedQueryStore& QueryStore,
-		FEnvironment& Environment) override;
+		UE::Editor::DataStorage::FExtendedQuery& Query,
+		UE::Editor::DataStorage::FExtendedQueryStore::Handle QueryHandle,
+		UE::Editor::DataStorage::FExtendedQueryStore& QueryStore,
+		UE::Editor::DataStorage::FEnvironment& Environment) override;
 
 private:
 	UPROPERTY(transient)
@@ -441,16 +385,12 @@ class UTypedElementQueryObserverCallbackAdapterProcessorWith4Subqueries final :
 {
 	GENERATED_BODY()
 
-	using FExtendedQuery = UE::Editor::DataStorage::FExtendedQuery;
-	using FExtendedQueryStore = UE::Editor::DataStorage::FExtendedQueryStore;
-	using FEnvironment = UE::Editor::DataStorage::FEnvironment;
-
 public:
 	virtual bool ConfigureQueryCallback(
-		FExtendedQuery& Query,
-		FExtendedQueryStore::Handle QueryHandle,
-		FExtendedQueryStore& QueryStore,
-		FEnvironment& Environment) override;
+		UE::Editor::DataStorage::FExtendedQuery& Query,
+		UE::Editor::DataStorage::FExtendedQueryStore::Handle QueryHandle,
+		UE::Editor::DataStorage::FExtendedQueryStore& QueryStore,
+		UE::Editor::DataStorage::FEnvironment& Environment) override;
 
 private:
 	UPROPERTY(transient)
@@ -462,16 +402,12 @@ class UTypedElementQueryObserverCallbackAdapterProcessorWith5Subqueries final :
 {
 	GENERATED_BODY()
 
-	using FExtendedQuery = UE::Editor::DataStorage::FExtendedQuery;
-	using FExtendedQueryStore = UE::Editor::DataStorage::FExtendedQueryStore;
-	using FEnvironment = UE::Editor::DataStorage::FEnvironment;
-
 public:
 	virtual bool ConfigureQueryCallback(
-		FExtendedQuery& Query,
-		FExtendedQueryStore::Handle QueryHandle,
-		FExtendedQueryStore& QueryStore,
-		FEnvironment& Environment) override;
+		UE::Editor::DataStorage::FExtendedQuery& Query,
+		UE::Editor::DataStorage::FExtendedQueryStore::Handle QueryHandle,
+		UE::Editor::DataStorage::FExtendedQueryStore& QueryStore,
+		UE::Editor::DataStorage::FEnvironment& Environment) override;
 
 private:
 	UPROPERTY(transient)
@@ -483,20 +419,16 @@ class UTypedElementQueryObserverCallbackAdapterProcessorWith6Subqueries final :
 {
 	GENERATED_BODY()
 
-	using FExtendedQuery = UE::Editor::DataStorage::FExtendedQuery;
-	using FExtendedQueryStore = UE::Editor::DataStorage::FExtendedQueryStore;
-	using FEnvironment = UE::Editor::DataStorage::FEnvironment;
-
 public:
 	virtual bool ConfigureQueryCallback(
-		FExtendedQuery& Query,
-		FExtendedQueryStore::Handle QueryHandle,
-		FExtendedQueryStore& QueryStore,
-		FEnvironment& Environment) override;
+		UE::Editor::DataStorage::FExtendedQuery& Query,
+		UE::Editor::DataStorage::FExtendedQueryStore::Handle QueryHandle,
+		UE::Editor::DataStorage::FExtendedQueryStore& QueryStore,
+		UE::Editor::DataStorage::FEnvironment& Environment) override;
 
 private:
 	UPROPERTY(transient)
-		FMassEntityQuery NativeSubqueries[6];
+	FMassEntityQuery NativeSubqueries[6];
 };
 
 UCLASS()
@@ -504,20 +436,16 @@ class UTypedElementQueryObserverCallbackAdapterProcessorWith7Subqueries final :
 {
 	GENERATED_BODY()
 
-	using FExtendedQuery = UE::Editor::DataStorage::FExtendedQuery;
-	using FExtendedQueryStore = UE::Editor::DataStorage::FExtendedQueryStore;
-	using FEnvironment = UE::Editor::DataStorage::FEnvironment;
-
 public:
 	virtual bool ConfigureQueryCallback(
-		FExtendedQuery& Query,
-		FExtendedQueryStore::Handle QueryHandle,
-		FExtendedQueryStore& QueryStore,
-		FEnvironment& Environment) override;
+		UE::Editor::DataStorage::FExtendedQuery& Query,
+		UE::Editor::DataStorage::FExtendedQueryStore::Handle QueryHandle,
+		UE::Editor::DataStorage::FExtendedQueryStore& QueryStore,
+		UE::Editor::DataStorage::FEnvironment& Environment) override;
 
 private:
 	UPROPERTY(transient)
-		FMassEntityQuery NativeSubqueries[7];
+	FMassEntityQuery NativeSubqueries[7];
 };
 
 UCLASS()
@@ -525,21 +453,14 @@ class UTypedElementQueryObserverCallbackAdapterProcessorWith8Subqueries final :
 {
 	GENERATED_BODY()
 
-	using FExtendedQuery = UE::Editor::DataStorage::FExtendedQuery;
-	using FExtendedQueryStore = UE::Editor::DataStorage::FExtendedQueryStore;
-	using FEnvironment = UE::Editor::DataStorage::FEnvironment;
-
 public:
 	virtual bool ConfigureQueryCallback(
-		FExtendedQuery& Query,
-		FExtendedQueryStore::Handle QueryHandle,
-		FExtendedQueryStore& QueryStore,
-		FEnvironment& Environment) override;
+		UE::Editor::DataStorage::FExtendedQuery& Query,
+		UE::Editor::DataStorage::FExtendedQueryStore::Handle QueryHandle,
+		UE::Editor::DataStorage::FExtendedQueryStore& QueryStore,
+		UE::Editor::DataStorage::FEnvironment& Environment) override;
 
 private:
 	UPROPERTY(transient)
-		FMassEntityQuery NativeSubqueries[8];
+	FMassEntityQuery NativeSubqueries[8];
 };
-
-
-
diff --git a/Engine/Plugins/Experimental/EditorDataStorage/Source/TedsCore/Private/Queries/TypedElementExtendedQueryStore.h b/Engine/Plugins/Experimental/EditorDataStorage/Source/TedsCore/Private/Queries/TypedElementExtendedQueryStore.h
--- a/Engine/Plugins/Experimental/EditorDataStorage/Source/TedsCore/Private/Queries/TypedElementExtendedQueryStore.h
+++ b/Engine/Plugins/Experimental/EditorDataStorage/Source/TedsCore/Private/Queries/TypedElementExtendedQueryStore.h
@@ -23,7 +23,7 @@ namespace UE::Editor::DataStorage
 	struct FExtendedQuery
 	{
 		FMassEntityQuery NativeQuery; // Used if there's no processor bound.
-		IEditorDataStorageProvider::FQueryDescription Description;
+		FQueryDescription Description;
 		TStrongObjectPtr<UMassProcessor> Processor;
 		FMassEntityQuery* QueryReference;
 		
@@ -61,7 +61,7 @@ namespace UE::Editor::DataStorage
 
 		 /** Adds a new query to the store and initializes the query with the provided arguments. */
 		Handle RegisterQuery(
-			IEditorDataStorageProvider::FQueryDescription Query,
+			FQueryDescription Query,
 			FEnvironment& Environment,
 			FMassEntityManager& EntityManager,
 			FMassProcessingPhaseManager& PhaseManager);
@@ -72,10 +72,10 @@ namespace UE::Editor::DataStorage
 		void Clear(FMassEntityManager& EntityManager, FMassProcessingPhaseManager& PhaseManager);
 
 		/** Register the defaults for a tick group. These will be applied on top of any settings provided with a query registration. */
-		void RegisterTickGroup(FName GroupName, IEditorDataStorageProvider::EQueryTickPhase Phase,
+		void RegisterTickGroup(FName GroupName, EQueryTickPhase Phase,
 			FName BeforeGroup, FName AfterGroup, EExecutionMode ExecutionMode);
 		/** Removes a previously registered set of tick group defaults. */
-		void UnregisterTickGroup(FName GroupName, IEditorDataStorageProvider::EQueryTickPhase Phase);
+		void UnregisterTickGroup(FName GroupName, EQueryTickPhase Phase);
 
 		/**
 		 * @section Retrieval
@@ -97,7 +97,7 @@ namespace UE::Editor::DataStorage
 		const FExtendedQuery& GetChecked(Handle Entry) const;
 
 		/** Gets the original description used to create an extended query or an empty default if the provided query isn't alive. */
-		const IEditorDataStorageProvider::FQueryDescription& GetQueryDescription(Handle Query) const;
+		const FQueryDescription& GetQueryDescription(Handle Query) const;
 
 		/** Checks to see if a query is still available or has been removed. */
 		bool IsAlive(Handle Entry) const;
@@ -138,17 +138,17 @@ namespace UE::Editor::DataStorage
 			FEnvironment& Environment,
 			FMassExecutionContext& ParentContext,
 			Handle Query,
-			DataStorage::RowHandle Row,
+			RowHandle Row,
 			SubqueryCallbackRef Callback);
 		void RunPhasePreambleQueries(
 			FMassEntityManager& EntityManager,
 			FEnvironment& Environment,
-			IEditorDataStorageProvider::EQueryTickPhase Phase,
+			EQueryTickPhase Phase,
 			float DeltaTime);
 		void RunPhasePostambleQueries(
 			FMassEntityManager& EntityManager,
 			FEnvironment& Environment,
-			IEditorDataStorageProvider::EQueryTickPhase Phase,
+			EQueryTickPhase Phase,
 			float DeltaTime);
 
 		void NotifyNewDynamicColumn(const FDynamicColumnGeneratorInfo& GeneratedColumnInfo);
@@ -156,13 +156,13 @@ namespace UE::Editor::DataStorage
 		void DebugPrintQueryCallbacks(FOutputDevice& Output) const;
 
 	private:
-		using QueryTickPhaseType = std::underlying_type_t<IEditorDataStorageProvider::EQueryTickPhase>;
-		static constexpr QueryTickPhaseType MaxTickPhase = static_cast<QueryTickPhaseType>(IEditorDataStorageProvider::EQueryTickPhase::Max);
+		using QueryTickPhaseType = std::underlying_type_t<EQueryTickPhase>;
+		static constexpr QueryTickPhaseType MaxTickPhase = static_cast<QueryTickPhaseType>(EQueryTickPhase::Max);
 
 		struct FTickGroupId
 		{
 			FName Name;
-			IEditorDataStorageProvider::EQueryTickPhase Phase;
+			EQueryTickPhase Phase;
 
 			friend inline uint32 GetTypeHash(const FTickGroupId& Id) { return HashCombine(GetTypeHash(Id.Name), GetTypeHash(Id.Phase)); }
 			friend inline bool operator==(const FTickGroupId& Lhs, const FTickGroupId& Rhs) { return Lhs.Phase == Rhs.Phase && Lhs.Name == Rhs.Name; }
@@ -185,31 +185,31 @@ namespace UE::Editor::DataStorage
 		EDirectQueryExecutionFlags ExecutionFlags,
 			CallbackReference Callback);
 
-		FMassEntityQuery& SetupNativeQuery(IEditorDataStorageProvider::FQueryDescription& Query, FExtendedQuery& StoredQuery, FEnvironment& Environment);
-		bool SetupDynamicColumns(IEditorDataStorageProvider::FQueryDescription& Query, FEnvironment& Environment);
-		bool SetupSelectedColumns(IEditorDataStorageProvider::FQueryDescription& Query, FMassEntityQuery& NativeQuery);
-		bool SetupConditions(IEditorDataStorageProvider::FQueryDescription& Query, FEnvironment& Environment, FMassEntityQuery& NativeQuery);
-		bool SetupChunkFilters(Handle QueryHandle, IEditorDataStorageProvider::FQueryDescription& Query, FEnvironment& Environment, FMassEntityQuery& NativeQuery);
-		bool SetupDependencies(IEditorDataStorageProvider::FQueryDescription& Query, FMassEntityQuery& NativeQuery);
-		bool SetupTickGroupDefaults(IEditorDataStorageProvider::FQueryDescription& Query);
+		FMassEntityQuery& SetupNativeQuery(FQueryDescription& Query, FExtendedQuery& StoredQuery, FEnvironment& Environment);
+		bool SetupDynamicColumns(FQueryDescription& Query, FEnvironment& Environment);
+		bool SetupSelectedColumns(FQueryDescription& Query, FMassEntityQuery& NativeQuery);
+		bool SetupConditions(FQueryDescription& Query, FEnvironment& Environment, FMassEntityQuery& NativeQuery);
+		bool SetupChunkFilters(Handle QueryHandle, FQueryDescription& Query, FEnvironment& Environment, FMassEntityQuery& NativeQuery);
+		bool SetupDependencies(FQueryDescription& Query, FMassEntityQuery& NativeQuery);
+		bool SetupTickGroupDefaults(FQueryDescription& Query);
 		bool SetupProcessors(Handle QueryHandle, FExtendedQuery& StoredQuery, FEnvironment& Environment,
 			FMassEntityManager& EntityManager, FMassProcessingPhaseManager& PhaseManager);
-		bool SetupActivatable(Handle QueryHandle, IEditorDataStorageProvider::FQueryDescription& Query);
+		bool SetupActivatable(Handle QueryHandle, FQueryDescription& Query);
 
-		EMassFragmentAccess ConvertToNativeAccessType(IEditorDataStorageProvider::EQueryAccessType AccessType);
-		EMassFragmentPresence ConvertToNativePresenceType(IEditorDataStorageProvider::EQueryAccessType AccessType);
+		EMassFragmentAccess ConvertToNativeAccessType(EQueryAccessType AccessType);
+		EMassFragmentPresence ConvertToNativePresenceType(EQueryAccessType AccessType);
 		EMassFragmentPresence ConvertToNativePresenceType(FQueryDescription::EOperatorType OperatorType);
 
-		void RegisterPreambleQuery(IEditorDataStorageProvider::EQueryTickPhase Phase, Handle Query);
-		void RegisterPostambleQuery(IEditorDataStorageProvider::EQueryTickPhase Phase, Handle Query);
-		void UnregisterPreambleQuery(IEditorDataStorageProvider::EQueryTickPhase Phase, Handle Query);
-		void UnregisterPostambleQuery(IEditorDataStorageProvider::EQueryTickPhase Phase, Handle Query);
+		void RegisterPreambleQuery(EQueryTickPhase Phase, Handle Query);
+		void RegisterPostambleQuery(EQueryTickPhase Phase, Handle Query);
+		void UnregisterPreambleQuery(EQueryTickPhase Phase, Handle Query);
+		void UnregisterPostambleQuery(EQueryTickPhase Phase, Handle Query);
 		void RunPhasePreOrPostAmbleQueries(FMassEntityManager& EntityManager, FEnvironment& Environment,
-			IEditorDataStorageProvider::EQueryTickPhase Phase, float DeltaTime, TArray<Handle>& QueryHandles);
+			EQueryTickPhase Phase, float DeltaTime, TArray<Handle>& QueryHandles);
 
 		void UnregisterQueryData(Handle Query, FExtendedQuery& QueryData, FMassEntityManager& EntityManager, FMassProcessingPhaseManager& PhaseManager);
 
-		static const IEditorDataStorageProvider::FQueryDescription EmptyDescription;
+		static const FQueryDescription EmptyDescription;
 
 		QueryStore Queries;
 		TMultiMap<FName, Handle> ActivatableMapping;
diff --git a/Engine/Plugins/Experimental/EditorDataStorage/Source/TedsCore/Private/Queries/TypedElementHierarchyQueries.cpp b/Engine/Plugins/Experimental/EditorDataStorage/Source/TedsCore/Private/Queries/TypedElementHierarchyQueries.cpp
--- a/Engine/Plugins/Experimental/EditorDataStorage/Source/TedsCore/Private/Queries/TypedElementHierarchyQueries.cpp
+++ b/Engine/Plugins/Experimental/EditorDataStorage/Source/TedsCore/Private/Queries/TypedElementHierarchyQueries.cpp
@@ -5,7 +5,7 @@
 #include "Elements/Columns/TypedElementHiearchyColumns.h"
 #include "Elements/Framework/TypedElementQueryBuilder.h"
 
-void UTypedElementHiearchyQueriesFactory::RegisterQueries(IEditorDataStorageProvider& DataStorage)
+void UTypedElementHiearchyQueriesFactory::RegisterQueries(UE::Editor::DataStorage::ICoreProvider& DataStorage)
 {
 	using namespace UE::Editor::DataStorage::Queries;
 
diff --git a/Engine/Plugins/Experimental/EditorDataStorage/Source/TedsCore/Private/Queries/TypedElementHierarchyQueries.h b/Engine/Plugins/Experimental/EditorDataStorage/Source/TedsCore/Private/Queries/TypedElementHierarchyQueries.h
--- a/Engine/Plugins/Experimental/EditorDataStorage/Source/TedsCore/Private/Queries/TypedElementHierarchyQueries.h
+++ b/Engine/Plugins/Experimental/EditorDataStorage/Source/TedsCore/Private/Queries/TypedElementHierarchyQueries.h
@@ -18,5 +18,5 @@ class UTypedElementHiearchyQueriesFactory : public UEditorDataStorageFactory
 public:
 	~UTypedElementHiearchyQueriesFactory() override = default;
 
-	void RegisterQueries(IEditorDataStorageProvider& DataStorage) override;
+	void RegisterQueries(UE::Editor::DataStorage::ICoreProvider& DataStorage) override;
 };
diff --git a/Engine/Plugins/Experimental/EditorDataStorage/Source/TedsCore/Private/Queries/TypedElementMiscQueries.cpp b/Engine/Plugins/Experimental/EditorDataStorage/Source/TedsCore/Private/Queries/TypedElementMiscQueries.cpp
--- a/Engine/Plugins/Experimental/EditorDataStorage/Source/TedsCore/Private/Queries/TypedElementMiscQueries.cpp
+++ b/Engine/Plugins/Experimental/EditorDataStorage/Source/TedsCore/Private/Queries/TypedElementMiscQueries.cpp
@@ -5,7 +5,7 @@
 #include "Elements/Columns/TypedElementMiscColumns.h"
 #include "Elements/Framework/TypedElementQueryBuilder.h"
 
-void UTypedElementRemoveSyncToWorldTagFactory::RegisterQueries(IEditorDataStorageProvider& DataStorage)
+void UTypedElementRemoveSyncToWorldTagFactory::RegisterQueries(UE::Editor::DataStorage::ICoreProvider& DataStorage)
 {
 	using namespace UE::Editor::DataStorage::Queries;
 
diff --git a/Engine/Plugins/Experimental/EditorDataStorage/Source/TedsCore/Private/Queries/TypedElementMiscQueries.h b/Engine/Plugins/Experimental/EditorDataStorage/Source/TedsCore/Private/Queries/TypedElementMiscQueries.h
--- a/Engine/Plugins/Experimental/EditorDataStorage/Source/TedsCore/Private/Queries/TypedElementMiscQueries.h
+++ b/Engine/Plugins/Experimental/EditorDataStorage/Source/TedsCore/Private/Queries/TypedElementMiscQueries.h
@@ -18,5 +18,5 @@ class UTypedElementRemoveSyncToWorldTagFactory : public UEditorDataStorageFactor
 public:
 	~UTypedElementRemoveSyncToWorldTagFactory() override = default;
 
-	void RegisterQueries(IEditorDataStorageProvider& DataStorage) override;
+	void RegisterQueries(UE::Editor::DataStorage::ICoreProvider& DataStorage) override;
 };
diff --git a/Engine/Plugins/Experimental/EditorDataStorage/Source/TedsCore/Private/TypedElementDatabase.cpp b/Engine/Plugins/Experimental/EditorDataStorage/Source/TedsCore/Private/TypedElementDatabase.cpp
--- a/Engine/Plugins/Experimental/EditorDataStorage/Source/TedsCore/Private/TypedElementDatabase.cpp
+++ b/Engine/Plugins/Experimental/EditorDataStorage/Source/TedsCore/Private/TypedElementDatabase.cpp
@@ -35,7 +35,7 @@ FAutoConsoleCommandWithOutputDevice PrintQueryCallbacksConsoleCommand(
 	FConsoleCommandWithOutputDeviceDelegate::CreateLambda([](FOutputDevice& Output)
 		{
 			using namespace UE::Editor::DataStorage;
-			if (IEditorDataStorageProvider* DataStorage = GetMutableDataStorageFeature<IEditorDataStorageProvider>(StorageFeatureName))
+			if (ICoreProvider* DataStorage = GetMutableDataStorageFeature<ICoreProvider>(StorageFeatureName))
 			{
 				DataStorage->DebugPrintQueryCallbacks(Output);
 			}
@@ -327,7 +327,7 @@ UE::Editor::DataStorage::RowHandle UEditorDataStorage::ReserveRow()
 		: UE::Editor::DataStorage::InvalidRowHandle;
 }
 
-void UEditorDataStorage::BatchReserveRows(int32 Count, TFunctionRef<void(RowHandle)> ReservationCallback)
+void UEditorDataStorage::BatchReserveRows(int32 Count, TFunctionRef<void(UE::Editor::DataStorage::RowHandle)> ReservationCallback)
 {
 	if (ActiveEditorEntityManager)
 	{
@@ -341,7 +341,7 @@ void UEditorDataStorage::BatchReserveRows(int32 Count, TFunctionRef<void(RowHand
 	}
 }
 
-void UEditorDataStorage::BatchReserveRows(TArrayView<RowHandle> ReservedRows)
+void UEditorDataStorage::BatchReserveRows(TArrayView<UE::Editor::DataStorage::RowHandle> ReservedRows)
 {
 	using namespace UE::Editor::DataStorage;
 
@@ -351,7 +351,7 @@ void UEditorDataStorage::BatchReserveRows(TArrayView<RowHandle> ReservedRows)
 	}
 }
 
-UE::Editor::DataStorage::RowHandle UEditorDataStorage::AddRow(TableHandle Table)
+UE::Editor::DataStorage::RowHandle UEditorDataStorage::AddRow(UE::Editor::DataStorage::TableHandle Table)
 {
 	using namespace UE::Editor::DataStorage;
 
@@ -362,8 +362,8 @@ UE::Editor::DataStorage::RowHandle UEditorDataStorage::AddRow(TableHandle Table)
 		: UE::Editor::DataStorage::InvalidRowHandle;
 }
 
-UE::Editor::DataStorage::RowHandle UEditorDataStorage::AddRow(TableHandle Table,
-	RowCreationCallbackRef OnCreated)
+UE::Editor::DataStorage::RowHandle UEditorDataStorage::AddRow(UE::Editor::DataStorage::TableHandle Table,
+	UE::Editor::DataStorage::RowCreationCallbackRef OnCreated)
 {
 	using namespace UE::Editor::DataStorage;
 
@@ -386,7 +386,7 @@ UE::Editor::DataStorage::RowHandle UEditorDataStorage::AddRow(TableHandle Table,
 	return InvalidRowHandle;
 }
 
-bool UEditorDataStorage::AddRow(RowHandle ReservedRow, TableHandle Table)
+bool UEditorDataStorage::AddRow(UE::Editor::DataStorage::RowHandle ReservedRow, UE::Editor::DataStorage::TableHandle Table)
 {
 	using namespace UE::Editor::DataStorage;
 
@@ -404,8 +404,8 @@ bool UEditorDataStorage::AddRow(RowHandle ReservedRow, TableHandle Table)
 	}
 }
 
-bool UEditorDataStorage::AddRow(RowHandle ReservedRow, TableHandle Table,
-	RowCreationCallbackRef OnCreated)
+bool UEditorDataStorage::AddRow(UE::Editor::DataStorage::RowHandle ReservedRow, UE::Editor::DataStorage::TableHandle Table,
+	UE::Editor::DataStorage::RowCreationCallbackRef OnCreated)
 {
 	using namespace UE::Editor::DataStorage;
 
@@ -425,7 +425,7 @@ bool UEditorDataStorage::AddRow(RowHandle ReservedRow, TableHandle Table,
 }
 
 bool UEditorDataStorage::BatchAddRow(
-	TableHandle Table, int32 Count, RowCreationCallbackRef OnCreated)
+	UE::Editor::DataStorage::TableHandle Table, int32 Count, UE::Editor::DataStorage::RowCreationCallbackRef OnCreated)
 {
 	using namespace UE::Editor::DataStorage;
 
@@ -450,8 +450,8 @@ bool UEditorDataStorage::BatchAddRow(
 	return false;
 }
 
-bool UEditorDataStorage::BatchAddRow(TableHandle Table, 
-	TConstArrayView<RowHandle> ReservedHandles, RowCreationCallbackRef OnCreated)
+bool UEditorDataStorage::BatchAddRow(UE::Editor::DataStorage::TableHandle Table,
+	TConstArrayView<UE::Editor::DataStorage::RowHandle> ReservedHandles, UE::Editor::DataStorage::RowCreationCallbackRef OnCreated)
 {
 	using namespace UE::Editor;
 	using namespace UE::Editor::DataStorage;
@@ -476,7 +476,7 @@ bool UEditorDataStorage::BatchAddRow(TableHandle Table,
 }
 
 
-void UEditorDataStorage::RemoveRow(RowHandle Row)
+void UEditorDataStorage::RemoveRow(UE::Editor::DataStorage::RowHandle Row)
 {
 	FMassEntityHandle Entity = FMassEntityHandle::FromNumber(Row);
 	if (ActiveEditorEntityManager && ActiveEditorEntityManager->IsEntityValid(Entity))
@@ -494,7 +494,7 @@ void UEditorDataStorage::RemoveRow(RowHandle Row)
 	}
 }
 
-void UEditorDataStorage::BatchRemoveRows(TConstArrayView<RowHandle> Rows)
+void UEditorDataStorage::BatchRemoveRows(TConstArrayView<UE::Editor::DataStorage::RowHandle> Rows)
 {
 	using namespace UE::Editor::DataStorage;
 
@@ -546,27 +546,27 @@ void UEditorDataStorage::RemoveAllRowsWithColumns(TConstArrayView<UScriptStruct*
 	}
 }
 
-bool UEditorDataStorage::IsRowAvailable(RowHandle Row) const
+bool UEditorDataStorage::IsRowAvailable(UE::Editor::DataStorage::RowHandle Row) const
 {
 	return ActiveEditorEntityManager ? UE::Editor::DataStorage::Legacy::FCommandBuffer::Execute_IsRowAvailable(*ActiveEditorEntityManager, Row) : false;
 }
 
-bool UEditorDataStorage::IsRowAssigned(RowHandle Row) const
+bool UEditorDataStorage::IsRowAssigned(UE::Editor::DataStorage::RowHandle Row) const
 {
 	return ActiveEditorEntityManager ? UE::Editor::DataStorage::Legacy::FCommandBuffer::Execute_IsRowAssigned(*ActiveEditorEntityManager, Row) : false;
 }
 
-bool UEditorDataStorage::IsRowAvailableUnsafe(RowHandle Row) const
+bool UEditorDataStorage::IsRowAvailableUnsafe(UE::Editor::DataStorage::RowHandle Row) const
 {
 	return UE::Editor::DataStorage::Legacy::FCommandBuffer::Execute_IsRowAvailable(*ActiveEditorEntityManager, Row);
 }
 
-bool UEditorDataStorage::IsRowAssignedUnsafe(RowHandle Row) const
+bool UEditorDataStorage::IsRowAssignedUnsafe(UE::Editor::DataStorage::RowHandle Row) const
 {
 	return UE::Editor::DataStorage::Legacy::FCommandBuffer::Execute_IsRowAssigned(*ActiveEditorEntityManager, Row);
 }
 
-void UEditorDataStorage::AddColumn(RowHandle Row, const UScriptStruct* ColumnType)
+void UEditorDataStorage::AddColumn(UE::Editor::DataStorage::RowHandle Row, const UScriptStruct* ColumnType)
 {
 	if (ColumnType && ActiveEditorEntityManager)
 	{
@@ -581,9 +581,9 @@ void UEditorDataStorage::AddColumn(RowHandle Row, const UScriptStruct* ColumnTyp
 	}
 }
 
-void UEditorDataStorage::AddColumnData(RowHandle Row, const UScriptStruct* ColumnType,
-	const ColumnCreationCallbackRef& Initializer,
-	ColumnCopyOrMoveCallback Relocator)
+void UEditorDataStorage::AddColumnData(UE::Editor::DataStorage::RowHandle Row, const UScriptStruct* ColumnType,
+	const UE::Editor::DataStorage::ColumnCreationCallbackRef Initializer,
+	UE::Editor::DataStorage::ColumnCopyOrMoveCallback Relocator)
 {
 	if (ActiveEditorEntityManager && UE::Mass::IsA<FMassFragment>(ColumnType))
 	{
@@ -608,7 +608,7 @@ void UEditorDataStorage::AddColumnData(RowHandle Row, const UScriptStruct* Colum
 	}
 }
 
-void UEditorDataStorage::RemoveColumn(RowHandle Row, const UScriptStruct* ColumnType)
+void UEditorDataStorage::RemoveColumn(UE::Editor::DataStorage::RowHandle Row, const UScriptStruct* ColumnType)
 {
 	if (ColumnType && ActiveEditorEntityManager)
 	{
@@ -623,7 +623,7 @@ void UEditorDataStorage::RemoveColumn(RowHandle Row, const UScriptStruct* Column
 	}
 }
 
-const void* UEditorDataStorage::GetColumnData(RowHandle Row, const UScriptStruct* ColumnType) const
+const void* UEditorDataStorage::GetColumnData(UE::Editor::DataStorage::RowHandle Row, const UScriptStruct* ColumnType) const
 {
 	FMassEntityHandle Entity = FMassEntityHandle::FromNumber(Row);
 	if (ActiveEditorEntityManager && UE::Mass::IsA<FMassFragment>(ColumnType))
@@ -644,12 +644,12 @@ const void* UEditorDataStorage::GetColumnData(RowHandle Row, const UScriptStruct
 	return nullptr;
 }
 
-void* UEditorDataStorage::GetColumnData(RowHandle Row, const UScriptStruct* ColumnType)
+void* UEditorDataStorage::GetColumnData(UE::Editor::DataStorage::RowHandle Row, const UScriptStruct* ColumnType)
 {
 	return const_cast<void*>(static_cast<const UEditorDataStorage*>(this)->GetColumnData(Row, ColumnType));
 }
 
-void UEditorDataStorage::AddColumns(RowHandle Row, TConstArrayView<const UScriptStruct*> Columns)
+void UEditorDataStorage::AddColumns(UE::Editor::DataStorage::RowHandle Row, TConstArrayView<const UScriptStruct*> Columns)
 {
 	using namespace UE::Editor::DataStorage;
 	if (ActiveEditorEntityManager)
@@ -673,7 +673,7 @@ void UEditorDataStorage::AddColumns(RowHandle Row, TConstArrayView<const UScript
 	}
 }
 
-void UEditorDataStorage::AddColumn(RowHandle Row, const UE::Editor::DataStorage::FValueTag& Tag, const FName& InValue)
+void UEditorDataStorage::AddColumn(UE::Editor::DataStorage::RowHandle Row, const UE::Editor::DataStorage::FValueTag& Tag, const FName& InValue)
 {
 	if (ActiveEditorEntityManager)
 	{
@@ -687,7 +687,7 @@ void UEditorDataStorage::AddColumn(RowHandle Row, const UE::Editor::DataStorage:
 	}
 }
 
-void UEditorDataStorage::RemoveColumn(RowHandle Row, const UE::Editor::DataStorage::FValueTag& Tag)
+void UEditorDataStorage::RemoveColumn(UE::Editor::DataStorage::RowHandle Row, const UE::Editor::DataStorage::FValueTag& Tag)
 {
 	if (ActiveEditorEntityManager)
 	{
@@ -700,7 +700,7 @@ void UEditorDataStorage::RemoveColumn(RowHandle Row, const UE::Editor::DataStora
 	}
 }
 
-void UEditorDataStorage::RemoveColumns(RowHandle Row, TConstArrayView<const UScriptStruct*> Columns)
+void UEditorDataStorage::RemoveColumns(UE::Editor::DataStorage::RowHandle Row, TConstArrayView<const UScriptStruct*> Columns)
 {
 	using namespace UE::Editor::DataStorage;
 
@@ -725,7 +725,7 @@ void UEditorDataStorage::RemoveColumns(RowHandle Row, TConstArrayView<const UScr
 	}
 }
 
-void UEditorDataStorage::AddRemoveColumns(RowHandle Row,
+void UEditorDataStorage::AddRemoveColumns(UE::Editor::DataStorage::RowHandle Row,
 	TConstArrayView<const UScriptStruct*> ColumnsToAdd, TConstArrayView<const UScriptStruct*> ColumnsToRemove)
 {
 	using namespace UE::Editor::DataStorage;
@@ -767,7 +767,7 @@ void UEditorDataStorage::AddRemoveColumns(RowHandle Row,
 	}
 }
 
-void UEditorDataStorage::BatchAddRemoveColumns(TConstArrayView<RowHandle> Rows, 
+void UEditorDataStorage::BatchAddRemoveColumns(TConstArrayView<UE::Editor::DataStorage::RowHandle> Rows,
 	TConstArrayView<const UScriptStruct*> ColumnsToAdd, TConstArrayView<const UScriptStruct*> ColumnsToRemove)
 {	
 	if (ActiveEditorEntityManager)
@@ -836,7 +836,7 @@ void UEditorDataStorage::BatchAddRemoveColumns(TConstArrayView<RowHandle> Rows,
 	}
 }
 
-bool UEditorDataStorage::HasColumns(RowHandle Row, TConstArrayView<const UScriptStruct*> ColumnTypes) const
+bool UEditorDataStorage::HasColumns(UE::Editor::DataStorage::RowHandle Row, TConstArrayView<const UScriptStruct*> ColumnTypes) const
 {
 	if (ActiveEditorEntityManager)
 	{
@@ -880,7 +880,7 @@ bool UEditorDataStorage::HasColumns(RowHandle Row, TConstArrayView<const UScript
 	return false;
 }
 
-bool UEditorDataStorage::HasColumns(RowHandle Row, TConstArrayView<TWeakObjectPtr<const UScriptStruct>> ColumnTypes) const
+bool UEditorDataStorage::HasColumns(UE::Editor::DataStorage::RowHandle Row, TConstArrayView<TWeakObjectPtr<const UScriptStruct>> ColumnTypes) const
 {
 	if (ActiveEditorEntityManager)
 	{
@@ -926,7 +926,7 @@ bool UEditorDataStorage::HasColumns(RowHandle Row, TConstArrayView<TWeakObjectPt
 	return false;
 }
 
-void UEditorDataStorage::ListColumns(RowHandle Row, ColumnListCallbackRef Callback) const
+void UEditorDataStorage::ListColumns(UE::Editor::DataStorage::RowHandle Row, UE::Editor::DataStorage::ColumnListCallbackRef Callback) const
 {
 	if (ActiveEditorEntityManager)
 	{
@@ -954,7 +954,7 @@ void UEditorDataStorage::ListColumns(RowHandle Row, ColumnListCallbackRef Callba
 	}
 }
 
-void UEditorDataStorage::ListColumns(RowHandle Row, ColumnListWithDataCallbackRef Callback)
+void UEditorDataStorage::ListColumns(UE::Editor::DataStorage::RowHandle Row, UE::Editor::DataStorage::ColumnListWithDataCallbackRef Callback)
 {
 	if (ActiveEditorEntityManager)
 	{
@@ -991,7 +991,7 @@ void UEditorDataStorage::ListColumns(RowHandle Row, ColumnListWithDataCallbackRe
 	}
 }
 
-bool UEditorDataStorage::MatchesColumns(RowHandle Row, const UE::Editor::DataStorage::Queries::FConditions& Conditions) const
+bool UEditorDataStorage::MatchesColumns(UE::Editor::DataStorage::RowHandle Row, const UE::Editor::DataStorage::Queries::FConditions& Conditions) const
 {
 	if (ActiveEditorEntityManager)
 	{
@@ -1025,24 +1025,24 @@ void UEditorDataStorage::ForEachDynamicColumn(const UScriptStruct& Template, TFu
 }
 
 void UEditorDataStorage::RegisterTickGroup(
-	FName GroupName, EQueryTickPhase Phase, FName BeforeGroup, FName AfterGroup, UE::Editor::DataStorage::EExecutionMode ExecutionMode)
+	FName GroupName, UE::Editor::DataStorage::EQueryTickPhase Phase, FName BeforeGroup, FName AfterGroup, UE::Editor::DataStorage::EExecutionMode ExecutionMode)
 {
 	Environment->GetQueryStore().RegisterTickGroup(GroupName, Phase, BeforeGroup, AfterGroup, ExecutionMode);
 }
 
-void UEditorDataStorage::UnregisterTickGroup(FName GroupName, EQueryTickPhase Phase)
+void UEditorDataStorage::UnregisterTickGroup(FName GroupName, UE::Editor::DataStorage::EQueryTickPhase Phase)
 {
 	Environment->GetQueryStore().UnregisterTickGroup(GroupName, Phase);
 }
 
-UE::Editor::DataStorage::QueryHandle UEditorDataStorage::RegisterQuery(FQueryDescription&& Query)
+UE::Editor::DataStorage::QueryHandle UEditorDataStorage::RegisterQuery(UE::Editor::DataStorage::FQueryDescription&& Query)
 {
 	return (ActiveEditorEntityManager && ActiveEditorPhaseManager)
 		? Environment->GetQueryStore().RegisterQuery(MoveTemp(Query), *Environment, *ActiveEditorEntityManager, *ActiveEditorPhaseManager).Packed()
 		: UE::Editor::DataStorage::InvalidQueryHandle;
 }
 
-void UEditorDataStorage::UnregisterQuery(QueryHandle Query)
+void UEditorDataStorage::UnregisterQuery(UE::Editor::DataStorage::QueryHandle Query)
 {
 	if (ActiveEditorEntityManager && ActiveEditorPhaseManager)
 	{
@@ -1051,14 +1051,16 @@ void UEditorDataStorage::UnregisterQuery(QueryHandle Query)
 	}
 }
 
-const IEditorDataStorageProvider::FQueryDescription& UEditorDataStorage::GetQueryDescription(QueryHandle Query) const
+const UE::Editor::DataStorage::FQueryDescription& UEditorDataStorage::GetQueryDescription(UE::Editor::DataStorage::QueryHandle Query) const
 {
 	const UE::Editor::DataStorage::FExtendedQueryStore::Handle StorageHandle(Query);
 	return Environment->GetQueryStore().GetQueryDescription(StorageHandle);
 }
 
-FName UEditorDataStorage::GetQueryTickGroupName(EQueryTickGroups Group) const
+FName UEditorDataStorage::GetQueryTickGroupName(UE::Editor::DataStorage::EQueryTickGroups Group) const
 {
+	using namespace UE::Editor::DataStorage;
+
 	switch (Group)
 	{
 		case EQueryTickGroups::Default:
@@ -1081,13 +1083,14 @@ FName UEditorDataStorage::GetQueryTickGroupName(EQueryTickGroups Group) const
 	}
 }
 
-IEditorDataStorageProvider::FQueryResult UEditorDataStorage::RunQuery(QueryHandle Query)
+UE::Editor::DataStorage::FQueryResult UEditorDataStorage::RunQuery(UE::Editor::DataStorage::QueryHandle Query)
 {
+	using namespace UE::Editor::DataStorage;
 	TRACE_CPUPROFILER_EVENT_SCOPE(TEDS.RunQuery);
 
 	if (ActiveEditorEntityManager)
 	{
-		const UE::Editor::DataStorage::FExtendedQueryStore::Handle StorageHandle(Query);
+		const FExtendedQueryStore::Handle StorageHandle(Query);
 		return Environment->GetQueryStore().RunQuery(*ActiveEditorEntityManager, StorageHandle);
 	}
 	else
@@ -1096,16 +1099,18 @@ IEditorDataStorageProvider::FQueryResult UEditorDataStorage::RunQuery(QueryHandl
 	}
 }
 
-IEditorDataStorageProvider::FQueryResult UEditorDataStorage::RunQuery(
-	QueryHandle Query, IEditorDataStorageProvider::DirectQueryCallbackRef Callback)
+UE::Editor::DataStorage::FQueryResult UEditorDataStorage::RunQuery(
+	UE::Editor::DataStorage::QueryHandle Query, UE::Editor::DataStorage::DirectQueryCallbackRef Callback)
 {
+	using namespace UE::Editor::DataStorage;
+
 	TRACE_CPUPROFILER_EVENT_SCOPE(TEDS.RunQuery);
 
 	if (ActiveEditorEntityManager)
 	{
-		const UE::Editor::DataStorage::FExtendedQueryStore::Handle StorageHandle(Query);
+		const FExtendedQueryStore::Handle StorageHandle(Query);
 		return Environment->GetQueryStore().RunQuery(*ActiveEditorEntityManager, *Environment, StorageHandle, 
-			UE::Editor::DataStorage::EDirectQueryExecutionFlags::Default, Callback);
+			EDirectQueryExecutionFlags::Default, Callback);
 	}
 	else
 	{
@@ -1113,15 +1118,17 @@ IEditorDataStorageProvider::FQueryResult UEditorDataStorage::RunQuery(
 	}
 }
 
-IEditorDataStorageProvider::FQueryResult UEditorDataStorage::RunQuery(
-	QueryHandle Query, UE::Editor::DataStorage::EDirectQueryExecutionFlags Flags,
-	IEditorDataStorageProvider::DirectQueryCallbackRef Callback)
+UE::Editor::DataStorage::FQueryResult UEditorDataStorage::RunQuery(
+	UE::Editor::DataStorage::QueryHandle Query, UE::Editor::DataStorage::EDirectQueryExecutionFlags Flags,
+	UE::Editor::DataStorage::DirectQueryCallbackRef Callback)
 {
+	using namespace UE::Editor::DataStorage;
+
 	TRACE_CPUPROFILER_EVENT_SCOPE(TEDS.RunQuery);
 
 	if (ActiveEditorEntityManager)
 	{
-		const UE::Editor::DataStorage::FExtendedQueryStore::Handle StorageHandle(Query);
+		const FExtendedQueryStore::Handle StorageHandle(Query);
 		return Environment->GetQueryStore().RunQuery(*ActiveEditorEntityManager, *Environment, StorageHandle, Flags, Callback);
 	}
 	else
@@ -1143,12 +1150,12 @@ UE::Editor::DataStorage::RowHandle UEditorDataStorage::LookupMappedRow(const UE:
 	return Environment->GetMappingTable().Lookup(UE::Editor::DataStorage::EGlobalLockScope::Public, Key);
 }
 
-void UEditorDataStorage::MapRow(UE::Editor::DataStorage::FMapKey Key, RowHandle Row)
+void UEditorDataStorage::MapRow(UE::Editor::DataStorage::FMapKey Key, UE::Editor::DataStorage::RowHandle Row)
 {
 	Environment->GetMappingTable().Map(UE::Editor::DataStorage::EGlobalLockScope::Public, MoveTemp(Key), Row);
 }
 
-void UEditorDataStorage::BatchMapRows(TArrayView<TPair<UE::Editor::DataStorage::FMapKey, RowHandle>> MapRowPairs)
+void UEditorDataStorage::BatchMapRows(TArrayView<TPair<UE::Editor::DataStorage::FMapKey, UE::Editor::DataStorage::RowHandle>> MapRowPairs)
 {
 	Environment->GetMappingTable().BatchMap(UE::Editor::DataStorage::EGlobalLockScope::Public, MapRowPairs);
 }
@@ -1163,12 +1170,12 @@ void UEditorDataStorage::RemoveRowMapping(const UE::Editor::DataStorage::FMapKey
 	Environment->GetMappingTable().Remove(UE::Editor::DataStorage::EGlobalLockScope::Public, Key);
 }
 
-FTypedElementOnDataStorageUpdate& UEditorDataStorage::OnUpdate()
+UE::Editor::DataStorage::FTypedElementOnDataStorageUpdate& UEditorDataStorage::OnUpdate()
 {
 	return OnUpdateDelegate;
 }
 
-FTypedElementOnDataStorageUpdate& UEditorDataStorage::OnUpdateCompleted()
+UE::Editor::DataStorage::FTypedElementOnDataStorageUpdate& UEditorDataStorage::OnUpdateCompleted()
 {
 	return OnUpdateCompletedDelegate;
 }
@@ -1196,7 +1203,7 @@ void UEditorDataStorage::ListExtensions(TFunctionRef<void(FName)> Callback) cons
 {
 }
 
-void UEditorDataStorage::PreparePhase(EQueryTickPhase Phase, float DeltaTime)
+void UEditorDataStorage::PreparePhase(UE::Editor::DataStorage::EQueryTickPhase Phase, float DeltaTime)
 {
 	using namespace UE::Editor::DataStorage;
 
@@ -1216,7 +1223,7 @@ void UEditorDataStorage::PreparePhase(EQueryTickPhase Phase, float DeltaTime)
 	}
 }
 
-void UEditorDataStorage::FinalizePhase(EQueryTickPhase Phase, float DeltaTime)
+void UEditorDataStorage::FinalizePhase(UE::Editor::DataStorage::EQueryTickPhase Phase, float DeltaTime)
 {
 	using namespace UE::Editor::DataStorage;
 
diff --git a/Engine/Plugins/Experimental/EditorDataStorage/Source/TedsCore/Private/TypedElementDatabase.h b/Engine/Plugins/Experimental/EditorDataStorage/Source/TedsCore/Private/TypedElementDatabase.h
--- a/Engine/Plugins/Experimental/EditorDataStorage/Source/TedsCore/Private/TypedElementDatabase.h
+++ b/Engine/Plugins/Experimental/EditorDataStorage/Source/TedsCore/Private/TypedElementDatabase.h
@@ -27,19 +27,10 @@ namespace UE::Editor::DataStorage
 UCLASS()
 class TEDSCORE_API UEditorDataStorage final
 	: public UObject
-	, public IEditorDataStorageProvider
+	, public UE::Editor::DataStorage::ICoreProvider
 {
 	GENERATED_BODY()
 
-	using RowCreationCallbackRef = UE::Editor::DataStorage::RowCreationCallbackRef;
-	using ColumnCreationCallbackRef = UE::Editor::DataStorage::ColumnCreationCallbackRef;
-	using ColumnListCallbackRef = UE::Editor::DataStorage::ColumnListCallbackRef;
-	using ColumnListWithDataCallbackRef = UE::Editor::DataStorage::ColumnListWithDataCallbackRef;
-	using ColumnCopyOrMoveCallback = UE::Editor::DataStorage::ColumnCopyOrMoveCallback;
-	using RowHandle = UE::Editor::DataStorage::RowHandle;
-	using TableHandle = UE::Editor::DataStorage::TableHandle;
-	using QueryHandle = UE::Editor::DataStorage::QueryHandle;
-
 public:
 	template<typename FactoryType, typename DatabaseType>
 	class TFactoryIterator
@@ -93,80 +84,80 @@ public:
 	TSharedPtr<FMassEntityManager> GetActiveMutableEditorEntityManager();
 	TSharedPtr<const FMassEntityManager> GetActiveEditorEntityManager() const;
 
-	virtual TableHandle RegisterTable(TConstArrayView<const UScriptStruct*> ColumnList, const FName& Name) override;
-	virtual TableHandle RegisterTable(
-		TableHandle SourceTable, TConstArrayView<const UScriptStruct*> ColumnList, const FName& Name) override;
-	virtual TableHandle FindTable(const FName& Name) override;
-
-	virtual RowHandle ReserveRow() override;
-	virtual void BatchReserveRows(int32 Count, TFunctionRef<void(RowHandle)> ReservationCallback) override;
-	virtual void BatchReserveRows(TArrayView<RowHandle> ReservedRows) override;
-	virtual RowHandle AddRow(TableHandle Table, 
-		RowCreationCallbackRef OnCreated) override;
-	RowHandle AddRow(TableHandle Table) override;
-	virtual bool AddRow(RowHandle ReservedRow, TableHandle Table) override;
-	virtual bool AddRow(RowHandle ReservedRow, TableHandle Table,
-		RowCreationCallbackRef OnCreated) override;
-	virtual bool BatchAddRow(TableHandle Table, int32 Count,
-		RowCreationCallbackRef OnCreated) override;
-	virtual bool BatchAddRow(TableHandle Table, TConstArrayView<RowHandle> ReservedHandles,
-		RowCreationCallbackRef OnCreated) override;
-	virtual void RemoveRow(RowHandle Row) override;
-	virtual void BatchRemoveRows(TConstArrayView<RowHandle> Rows) override;
+	virtual UE::Editor::DataStorage::TableHandle RegisterTable(TConstArrayView<const UScriptStruct*> ColumnList, const FName& Name) override;
+	virtual UE::Editor::DataStorage::TableHandle RegisterTable(
+		UE::Editor::DataStorage::TableHandle SourceTable, TConstArrayView<const UScriptStruct*> ColumnList, const FName& Name) override;
+	virtual UE::Editor::DataStorage::TableHandle FindTable(const FName& Name) override;
+
+	virtual UE::Editor::DataStorage::RowHandle ReserveRow() override;
+	virtual void BatchReserveRows(int32 Count, TFunctionRef<void(UE::Editor::DataStorage::RowHandle)> ReservationCallback) override;
+	virtual void BatchReserveRows(TArrayView<UE::Editor::DataStorage::RowHandle> ReservedRows) override;
+	virtual UE::Editor::DataStorage::RowHandle AddRow(UE::Editor::DataStorage::TableHandle Table, 
+		UE::Editor::DataStorage::RowCreationCallbackRef OnCreated) override;
+	UE::Editor::DataStorage::RowHandle AddRow(UE::Editor::DataStorage::TableHandle Table) override;
+	virtual bool AddRow(UE::Editor::DataStorage::RowHandle ReservedRow, UE::Editor::DataStorage::TableHandle Table) override;
+	virtual bool AddRow(UE::Editor::DataStorage::RowHandle ReservedRow, UE::Editor::DataStorage::TableHandle Table,
+		UE::Editor::DataStorage::RowCreationCallbackRef OnCreated) override;
+	virtual bool BatchAddRow(UE::Editor::DataStorage::TableHandle Table, int32 Count,
+		UE::Editor::DataStorage::RowCreationCallbackRef OnCreated) override;
+	virtual bool BatchAddRow(UE::Editor::DataStorage::TableHandle Table, TConstArrayView<UE::Editor::DataStorage::RowHandle> ReservedHandles,
+		UE::Editor::DataStorage::RowCreationCallbackRef OnCreated) override;
+	virtual void RemoveRow(UE::Editor::DataStorage::RowHandle Row) override;
+	virtual void BatchRemoveRows(TConstArrayView<UE::Editor::DataStorage::RowHandle> Rows) override;
 	virtual void RemoveAllRowsWithColumns(TConstArrayView<UScriptStruct*> Columns) override;
-	virtual bool IsRowAvailable(RowHandle Row) const override;
-	virtual bool IsRowAssigned(RowHandle Row) const override;
+	virtual bool IsRowAvailable(UE::Editor::DataStorage::RowHandle Row) const override;
+	virtual bool IsRowAssigned(UE::Editor::DataStorage::RowHandle Row) const override;
 	/** Same as IsRowAvailable, but doesn't check if the data storage has been initialized. */
-	bool IsRowAvailableUnsafe(RowHandle Row) const;
+	bool IsRowAvailableUnsafe(UE::Editor::DataStorage::RowHandle Row) const;
 	/** Same as IsRowAssigned, but doesn't check if the data storage has been initialized. */
-	bool IsRowAssignedUnsafe(RowHandle Row) const;
-
-	virtual void AddColumn(RowHandle Row, const UScriptStruct* ColumnType) override;
-	virtual void AddColumn(RowHandle Row, const UE::Editor::DataStorage::FValueTag& Tag, const FName& InValue) override;
-	virtual void AddColumnData(RowHandle Row, const UScriptStruct* ColumnType,
-		const ColumnCreationCallbackRef& Initializer,
-		ColumnCopyOrMoveCallback Relocator) override;
-	virtual void RemoveColumn(RowHandle Row, const UScriptStruct* ColumnType) override;
-	virtual void RemoveColumn(RowHandle Row, const UE::Editor::DataStorage::FValueTag& Tag) override;
-	virtual void* GetColumnData(RowHandle Row, const UScriptStruct* ColumnType) override;
-	virtual const void* GetColumnData(RowHandle Row, const UScriptStruct* ColumnType) const override;
-	virtual void AddColumns(RowHandle Row, TConstArrayView<const UScriptStruct*> Columns) override;
-	virtual void RemoveColumns(RowHandle Row, TConstArrayView<const UScriptStruct*> Columns) override;
-	virtual void AddRemoveColumns(RowHandle Row, TConstArrayView<const UScriptStruct*> ColumnsToAdd,
+	bool IsRowAssignedUnsafe(UE::Editor::DataStorage::RowHandle Row) const;
+
+	virtual void AddColumn(UE::Editor::DataStorage::RowHandle Row, const UScriptStruct* ColumnType) override;
+	virtual void AddColumn(UE::Editor::DataStorage::RowHandle Row, const UE::Editor::DataStorage::FValueTag& Tag, const FName& InValue) override;
+	virtual void AddColumnData(UE::Editor::DataStorage::RowHandle Row, const UScriptStruct* ColumnType,
+		const UE::Editor::DataStorage::ColumnCreationCallbackRef Initializer,
+		UE::Editor::DataStorage::ColumnCopyOrMoveCallback Relocator) override;
+	virtual void RemoveColumn(UE::Editor::DataStorage::RowHandle Row, const UScriptStruct* ColumnType) override;
+	virtual void RemoveColumn(UE::Editor::DataStorage::RowHandle Row, const UE::Editor::DataStorage::FValueTag& Tag) override;
+	virtual void* GetColumnData(UE::Editor::DataStorage::RowHandle Row, const UScriptStruct* ColumnType) override;
+	virtual const void* GetColumnData(UE::Editor::DataStorage::RowHandle Row, const UScriptStruct* ColumnType) const override;
+	virtual void AddColumns(UE::Editor::DataStorage::RowHandle Row, TConstArrayView<const UScriptStruct*> Columns) override;
+	virtual void RemoveColumns(UE::Editor::DataStorage::RowHandle Row, TConstArrayView<const UScriptStruct*> Columns) override;
+	virtual void AddRemoveColumns(UE::Editor::DataStorage::RowHandle Row, TConstArrayView<const UScriptStruct*> ColumnsToAdd,
 		TConstArrayView<const UScriptStruct*> ColumnsToRemove) override;
-	virtual void BatchAddRemoveColumns(TConstArrayView<RowHandle> Rows,TConstArrayView<const UScriptStruct*> ColumnsToAdd,
+	virtual void BatchAddRemoveColumns(TConstArrayView<UE::Editor::DataStorage::RowHandle> Rows,TConstArrayView<const UScriptStruct*> ColumnsToAdd,
 		TConstArrayView<const UScriptStruct*> ColumnsToRemove) override;
-	virtual bool HasColumns(RowHandle Row, TConstArrayView<const UScriptStruct*> ColumnTypes) const override;
-	virtual bool HasColumns(RowHandle Row, TConstArrayView<TWeakObjectPtr<const UScriptStruct>> ColumnTypes) const override;
-	virtual void ListColumns(RowHandle Row, ColumnListCallbackRef Callback) const override;
-	virtual void ListColumns(RowHandle Row, ColumnListWithDataCallbackRef Callback) override;
-	virtual bool MatchesColumns(RowHandle Row, const UE::Editor::DataStorage::Queries::FConditions& Conditions) const override;
+	virtual bool HasColumns(UE::Editor::DataStorage::RowHandle Row, TConstArrayView<const UScriptStruct*> ColumnTypes) const override;
+	virtual bool HasColumns(UE::Editor::DataStorage::RowHandle Row, TConstArrayView<TWeakObjectPtr<const UScriptStruct>> ColumnTypes) const override;
+	virtual void ListColumns(UE::Editor::DataStorage::RowHandle Row, UE::Editor::DataStorage::ColumnListCallbackRef Callback) const override;
+	virtual void ListColumns(UE::Editor::DataStorage::RowHandle Row, UE::Editor::DataStorage::ColumnListWithDataCallbackRef Callback) override;
+	virtual bool MatchesColumns(UE::Editor::DataStorage::RowHandle Row, const UE::Editor::DataStorage::Queries::FConditions& Conditions) const override;
 
 	virtual const UScriptStruct* FindDynamicColumn(const UE::Editor::DataStorage::FDynamicColumnDescription& Description) const override;
 	virtual const UScriptStruct* GenerateDynamicColumn(const UE::Editor::DataStorage::FDynamicColumnDescription& Description) override;
 	virtual void ForEachDynamicColumn(const UScriptStruct& Template, TFunctionRef<void(const UScriptStruct& Type)> Callback) const override;
 
-	void RegisterTickGroup(FName GroupName, EQueryTickPhase Phase, FName BeforeGroup, FName AfterGroup, UE::Editor::DataStorage::EExecutionMode ExecutionMode);
-	void UnregisterTickGroup(FName GroupName, EQueryTickPhase Phase);
-
-	QueryHandle RegisterQuery(FQueryDescription&& Query) override;
-	virtual void UnregisterQuery(QueryHandle Query) override;
-	virtual const FQueryDescription& GetQueryDescription(QueryHandle Query) const override;
-	virtual FName GetQueryTickGroupName(EQueryTickGroups Group) const override;
-	virtual FQueryResult RunQuery(QueryHandle Query) override;
-	virtual FQueryResult RunQuery(QueryHandle Query, DirectQueryCallbackRef Callback) override;
-	virtual FQueryResult RunQuery(QueryHandle Query, UE::Editor::DataStorage::EDirectQueryExecutionFlags Flags,
-		DirectQueryCallbackRef Callback) override;
+	void RegisterTickGroup(FName GroupName, UE::Editor::DataStorage::EQueryTickPhase Phase, FName BeforeGroup, FName AfterGroup, UE::Editor::DataStorage::EExecutionMode ExecutionMode);
+	void UnregisterTickGroup(FName GroupName, UE::Editor::DataStorage::EQueryTickPhase Phase);
+
+	UE::Editor::DataStorage::QueryHandle RegisterQuery(UE::Editor::DataStorage::FQueryDescription&& Query) override;
+	virtual void UnregisterQuery(UE::Editor::DataStorage::QueryHandle Query) override;
+	virtual const UE::Editor::DataStorage::FQueryDescription& GetQueryDescription(UE::Editor::DataStorage::QueryHandle Query) const override;
+	virtual FName GetQueryTickGroupName(UE::Editor::DataStorage::EQueryTickGroups Group) const override;
+	virtual UE::Editor::DataStorage::FQueryResult RunQuery(UE::Editor::DataStorage::QueryHandle Query) override;
+	virtual UE::Editor::DataStorage::FQueryResult RunQuery(UE::Editor::DataStorage::QueryHandle Query, UE::Editor::DataStorage::DirectQueryCallbackRef Callback) override;
+	virtual UE::Editor::DataStorage::FQueryResult RunQuery(UE::Editor::DataStorage::QueryHandle Query, UE::Editor::DataStorage::EDirectQueryExecutionFlags Flags,
+		UE::Editor::DataStorage::DirectQueryCallbackRef Callback) override;
 	virtual void ActivateQueries(FName ActivationName) override;
 
-	virtual RowHandle LookupMappedRow(const UE::Editor::DataStorage::FMapKeyView& Key) const override;
-	virtual void MapRow(UE::Editor::DataStorage::FMapKey Key, RowHandle Row) override;
-	virtual void BatchMapRows(TArrayView<TPair<UE::Editor::DataStorage::FMapKey, RowHandle>> MapRowPairs) override;
+	virtual UE::Editor::DataStorage::RowHandle LookupMappedRow(const UE::Editor::DataStorage::FMapKeyView& Key) const override;
+	virtual void MapRow(UE::Editor::DataStorage::FMapKey Key, UE::Editor::DataStorage::RowHandle Row) override;
+	virtual void BatchMapRows(TArrayView<TPair<UE::Editor::DataStorage::FMapKey, UE::Editor::DataStorage::RowHandle>> MapRowPairs) override;
 	virtual void RemapRow(const UE::Editor::DataStorage::FMapKeyView& OriginalKey, UE::Editor::DataStorage::FMapKey NewKey) override;
 	virtual void RemoveRowMapping(const UE::Editor::DataStorage::FMapKeyView& Key) override;
 
-	virtual FTypedElementOnDataStorageUpdate& OnUpdate() override;
-	virtual FTypedElementOnDataStorageUpdate& OnUpdateCompleted() override;
+	virtual UE::Editor::DataStorage::FTypedElementOnDataStorageUpdate& OnUpdate() override;
+	virtual UE::Editor::DataStorage::FTypedElementOnDataStorageUpdate& OnUpdateCompleted() override;
 	virtual bool IsAvailable() const override;
 	virtual void* GetExternalSystemAddress(UClass* Target) override;
 
@@ -176,15 +167,15 @@ public:
 	TSharedPtr<UE::Editor::DataStorage::FEnvironment> GetEnvironment();
 	TSharedPtr<const UE::Editor::DataStorage::FEnvironment> GetEnvironment() const;
 
-	FMassArchetypeHandle LookupArchetype(TableHandle InTableHandle) const;
+	FMassArchetypeHandle LookupArchetype(UE::Editor::DataStorage::TableHandle InTableHandle) const;
 
 	void DebugPrintQueryCallbacks(FOutputDevice& Output) override;
 
 	static void AddReferencedObjects(UObject* InThis, FReferenceCollector& Collector);
 
 private:
-	void PreparePhase(EQueryTickPhase Phase, float DeltaTime);
-	void FinalizePhase(EQueryTickPhase Phase, float DeltaTime);
+	void PreparePhase(UE::Editor::DataStorage::EQueryTickPhase Phase, float DeltaTime);
+	void FinalizePhase(UE::Editor::DataStorage::EQueryTickPhase Phase, float DeltaTime);
 	void Reset();
 
 	int32 GetTableChunkSize(FName TableName) const;
@@ -206,15 +197,15 @@ private:
 	static const FName TickGroupName_SyncDataStorageToExternal;
 	
 	TArray<FMassArchetypeHandle> Tables;
-	TMap<FName, TableHandle> TableNameLookup;
+	TMap<FName, UE::Editor::DataStorage::TableHandle> TableNameLookup;
 
 	// Ordered array of factories by the return value of GetOrder()
 	TArray<FFactoryTypePair> Factories;
 
 	TSharedPtr<UE::Editor::DataStorage::FEnvironment> Environment;
 	
-	FTypedElementOnDataStorageUpdate OnUpdateDelegate;
-	FTypedElementOnDataStorageUpdate OnUpdateCompletedDelegate;
+	UE::Editor::DataStorage::FTypedElementOnDataStorageUpdate OnUpdateDelegate;
+	UE::Editor::DataStorage::FTypedElementOnDataStorageUpdate OnUpdateCompletedDelegate;
 	FDelegateHandle OnPreMassTickHandle;
 	FDelegateHandle OnPostMassTickHandle;
 
diff --git a/Engine/Plugins/Experimental/EditorDataStorage/Source/TedsCore/Private/TypedElementDatabaseCompatibility.cpp b/Engine/Plugins/Experimental/EditorDataStorage/Source/TedsCore/Private/TypedElementDatabaseCompatibility.cpp
--- a/Engine/Plugins/Experimental/EditorDataStorage/Source/TedsCore/Private/TypedElementDatabaseCompatibility.cpp
+++ b/Engine/Plugins/Experimental/EditorDataStorage/Source/TedsCore/Private/TypedElementDatabaseCompatibility.cpp
@@ -345,7 +345,7 @@ void UEditorDataStorageCompatibility::Reset()
 void UEditorDataStorageCompatibility::CreateStandardArchetypes()
 {
 	// Thread-safe as this is only called from a function that has an exclusive lock.
-
+	using namespace UE::Editor::DataStorage;
 	StandardActorTable = Storage->RegisterTable(TTypedElementColumnTypeList<
 			FTypedElementUObjectColumn, FTypedElementUObjectIdColumn, FTypedElementClassTypeInfoColumn,
 			FTypedElementLabelColumn, FTypedElementLabelHashColumn, FTypedElementActorTag,
@@ -724,7 +724,7 @@ void UEditorDataStorageCompatibility::PendingRegistration<AddressType>::ForEachA
 }
 
 template<typename AddressType>
-void UEditorDataStorageCompatibility::PendingRegistration<AddressType>::ProcessEntries(IEditorDataStorageProvider& StorageInterface,
+void UEditorDataStorageCompatibility::PendingRegistration<AddressType>::ProcessEntries(UE::Editor::DataStorage::ICoreProvider& StorageInterface,
 	UEditorDataStorageCompatibility& Compatibility, const TFunctionRef<void(UE::Editor::DataStorage::RowHandle, const AddressType&)>& SetupRowCallback)
 {
 	// Thread-safe as it's only called from functions that already lock using.
@@ -893,7 +893,7 @@ void UEditorDataStorageCompatibility::TickPendingUObjectRegistration()
 		UObjectsPendingRegistration.ProcessEntries(*Storage, *this,
 			[this](UE::Editor::DataStorage::RowHandle Row, const TWeakObjectPtr<UObject>& Object)
 			{
-				IEditorDataStorageProvider* Interface = Storage;
+				UE::Editor::DataStorage::ICoreProvider* Interface = Storage;
 				Interface->AddColumn(Row, FTypedElementUObjectColumn{ .Object = Object });
 				Interface->AddColumn(Row, FTypedElementUObjectIdColumn
 					{ 
@@ -923,7 +923,7 @@ void UEditorDataStorageCompatibility::TickPendingExternalObjectRegistration()
 		ExternalObjectsPendingRegistration.ProcessEntries(*Storage, *this,
 			[this](UE::Editor::DataStorage::RowHandle Row, const ExternalObjectRegistration& Object)
 			{
-				IEditorDataStorageProvider* Interface = Storage;
+				UE::Editor::DataStorage::ICoreProvider* Interface = Storage;
 				Interface->AddColumn(Row, FTypedElementExternalObjectColumn{ .Object = Object.Object });
 				Interface->AddColumn(Row, FTypedElementScriptStructTypeInfoColumn{ .TypeInfo = Object.TypeInfo });
 				// Make sure the new row is tagged for update.
@@ -1261,7 +1261,7 @@ void UEditorDataStorageCompatibility::FRegistrationCommandChange::Revert(UObject
 	if (UObject* TargetRetrieved = TargetObject.Get(/*bEvenIfPendingKill=*/ true))
 	{
 		UEditorDataStorageCompatibility* DataStorageCompat = Owner.Get(); 
-		IEditorDataStorageProvider* DataStorage = DataStorageCompat->Storage;
+		ICoreProvider* DataStorage = DataStorageCompat->Storage;
 			
 		RowHandle ObjectRow = DataStorageCompat->FindRowWithCompatibleObjectExplicit(TargetRetrieved);
 		if (DataStorage->IsRowAvailable(ObjectRow))
@@ -1297,7 +1297,7 @@ UEditorDataStorageCompatibility::FDeregistrationCommandChange::FDeregistrationCo
 {
 	using namespace UE::Editor::DataStorage;
 
-	IEditorDataStorageProvider* DataStorage = InOwner->Storage;
+	ICoreProvider* DataStorage = InOwner->Storage;
 
 	RowHandle ObjectRow = InOwner->FindRowWithCompatibleObjectExplicit(InTargetObject);
 	if (DataStorage->IsRowAvailable(ObjectRow))
diff --git a/Engine/Plugins/Experimental/EditorDataStorage/Source/TedsCore/Private/TypedElementDatabaseCompatibility.h b/Engine/Plugins/Experimental/EditorDataStorage/Source/TedsCore/Private/TypedElementDatabaseCompatibility.h
--- a/Engine/Plugins/Experimental/EditorDataStorage/Source/TedsCore/Private/TypedElementDatabaseCompatibility.h
+++ b/Engine/Plugins/Experimental/EditorDataStorage/Source/TedsCore/Private/TypedElementDatabaseCompatibility.h
@@ -19,19 +19,19 @@
 
 #include "TypedElementDatabaseCompatibility.generated.h"
 
-class IEditorDataStorageProvider;
 struct FMassActorManager;
 class UTypedElementMementoSystem;
 
 namespace UE::Editor::DataStorage
 {
 	class FEnvironment;
+	class ICoreProvider;
 } // namespace UE::Editor::DataStorage
 
 UCLASS()
 class TEDSCORE_API UEditorDataStorageCompatibility
 	: public UObject
-	, public IEditorDataStorageCompatibilityProvider
+	, public UE::Editor::DataStorage::ICompatibilityProvider
 {
 	GENERATED_BODY()
 
@@ -184,7 +184,7 @@ private:
 		int32 Num() const;
 		
 		void ForEachAddress(const TFunctionRef<void(AddressType&)>& Callback);
-		void ProcessEntries(IEditorDataStorageProvider& Storage, UEditorDataStorageCompatibility& Compatibility,
+		void ProcessEntries(UE::Editor::DataStorage::ICoreProvider& Storage, UEditorDataStorageCompatibility& Compatibility,
 			const TFunctionRef<void(UE::Editor::DataStorage::RowHandle, const AddressType&)>& SetupRowCallback);
 		void Reset();
 	};
@@ -206,7 +206,7 @@ private:
 	UE::Editor::DataStorage::TableHandle StandardActorWithTransformTable{ UE::Editor::DataStorage::InvalidTableHandle };
 	UE::Editor::DataStorage::TableHandle StandardUObjectTable{ UE::Editor::DataStorage::InvalidTableHandle };
 	UE::Editor::DataStorage::TableHandle StandardExternalObjectTable{ UE::Editor::DataStorage::InvalidTableHandle };
-	IEditorDataStorageProvider* Storage{ nullptr };
+	UE::Editor::DataStorage::ICoreProvider* Storage{ nullptr };
 
 	/**
 	 * Reference of objects (UObject and AActor) that need to be fully synced from the world to the database.
