[TEDS] Cleanup aliases, move core interface classes into namespace, and fix up one class name that wasn't following standards

diff --git a/Engine/Plugins/Experimental/EditorDataStorage/Source/TedsCore/Private/TypedElementDatabaseUI.cpp b/Engine/Plugins/Experimental/EditorDataStorage/Source/TedsCore/Private/TypedElementDatabaseUI.cpp
--- a/Engine/Plugins/Experimental/EditorDataStorage/Source/TedsCore/Private/TypedElementDatabaseUI.cpp
+++ b/Engine/Plugins/Experimental/EditorDataStorage/Source/TedsCore/Private/TypedElementDatabaseUI.cpp
@@ -47,8 +47,8 @@ namespace Internal
 }
 
 void UEditorDataStorageUi::Initialize(
-	IEditorDataStorageProvider* StorageInterface,
-	IEditorDataStorageCompatibilityProvider* StorageCompatibilityInterface)
+	UE::Editor::DataStorage::ICoreProvider* StorageInterface,
+	UE::Editor::DataStorage::ICompatibilityProvider* StorageCompatibilityInterface)
 {
 	checkf(StorageInterface, TEXT("TEDS' compatibility manager is being initialized with an invalid storage target."));
 
@@ -82,11 +82,11 @@ bool UEditorDataStorageUi::RegisterWidgetFactory(FName Purpose, const UScriptStr
 	{
 		switch (PurposeInfo->Type)
 		{
-		case IEditorDataStorageUiProvider::EPurposeType::Generic:
+		case IUiProvider::EPurposeType::Generic:
 			PurposeInfo->Factories.Emplace(Constructor);
 			PurposeInfo->bIsSorted = false;
 			return true;
-		case IEditorDataStorageUiProvider::EPurposeType::UniqueByName:
+		case IUiProvider::EPurposeType::UniqueByName:
 			if (PurposeInfo->Factories.IsEmpty())
 			{
 				PurposeInfo->Factories.Emplace(Constructor);
@@ -97,13 +97,13 @@ bool UEditorDataStorageUi::RegisterWidgetFactory(FName Purpose, const UScriptStr
 				PurposeInfo->Factories.EmplaceAt(0, Constructor);
 			}
 			return true;
-		case IEditorDataStorageUiProvider::EPurposeType::UniqueByNameAndColumn:
+		case IUiProvider::EPurposeType::UniqueByNameAndColumn:
 			UE_LOG(LogEditorDataStorageUI, Warning,
 				TEXT("Unable to register widget factory '%s' as purpose '%s' requires at least one column for matching."), 
 				*Constructor->GetName(), *Purpose.ToString());
 			return false;
 		default:
-			checkf(false, TEXT("Unexpected IEditorDataStorageUiProvider::EPurposeType found provided when registering widget factory."));
+			checkf(false, TEXT("Unexpected IUiProvider::EPurposeType found provided when registering widget factory."));
 			return false;
 		}
 	}
@@ -128,11 +128,11 @@ bool UEditorDataStorageUi::RegisterWidgetFactory(
 		{
 			switch (PurposeInfo->Type)
 			{
-			case IEditorDataStorageUiProvider::EPurposeType::Generic:
+			case IUiProvider::EPurposeType::Generic:
 				PurposeInfo->Factories.Emplace(Constructor);
 				PurposeInfo->bIsSorted = false;
 				return true;
-			case IEditorDataStorageUiProvider::EPurposeType::UniqueByName:
+			case IUiProvider::EPurposeType::UniqueByName:
 				if (!Columns.IsEmpty())
 				{
 					if (PurposeInfo->Factories.IsEmpty())
@@ -150,7 +150,7 @@ bool UEditorDataStorageUi::RegisterWidgetFactory(
 				{
 					return false;
 				}
-			case IEditorDataStorageUiProvider::EPurposeType::UniqueByNameAndColumn:
+			case IUiProvider::EPurposeType::UniqueByNameAndColumn:
 				if (!Columns.IsEmpty())
 				{
 					PurposeInfo->Factories.Emplace(Constructor, MoveTemp(Columns));
@@ -162,7 +162,7 @@ bool UEditorDataStorageUi::RegisterWidgetFactory(
 					return false;
 				}
 			default:
-				checkf(false, TEXT("Unexpected IEditorDataStorageUiProvider::EPurposeType found provided when registering widget factory."));
+				checkf(false, TEXT("Unexpected IUiProvider::EPurposeType found provided when registering widget factory."));
 				return false;
 			}
 		}
@@ -187,11 +187,11 @@ bool UEditorDataStorageUi::RegisterWidgetFactory(FName Purpose, TUniquePtr<FType
 	{
 		switch (PurposeInfo->Type)
 		{
-		case IEditorDataStorageUiProvider::EPurposeType::Generic:
+		case IUiProvider::EPurposeType::Generic:
 			PurposeInfo->Factories.Emplace(MoveTemp(Constructor));
 			PurposeInfo->bIsSorted = false;
 			return true;
-		case IEditorDataStorageUiProvider::EPurposeType::UniqueByName:
+		case IUiProvider::EPurposeType::UniqueByName:
 			if (PurposeInfo->Factories.IsEmpty())
 			{
 				PurposeInfo->Factories.Emplace(MoveTemp(Constructor));
@@ -202,13 +202,13 @@ bool UEditorDataStorageUi::RegisterWidgetFactory(FName Purpose, TUniquePtr<FType
 				PurposeInfo->Factories.EmplaceAt(0, MoveTemp(Constructor));
 			}
 			return true;
-		case IEditorDataStorageUiProvider::EPurposeType::UniqueByNameAndColumn:
+		case IUiProvider::EPurposeType::UniqueByNameAndColumn:
 			UE_LOG(LogEditorDataStorageUI, Warning,
 				TEXT("Unable to register widget factory '%s' as purpose '%s' requires at least one column for matching."),
 				*Constructor->GetTypeInfo()->GetName(), *Purpose.ToString());
 			return false;
 		default:
-			checkf(false, TEXT("Unexpected IEditorDataStorageUiProvider::EPurposeType found provided when registering widget factory."));
+			checkf(false, TEXT("Unexpected IUiProvider::EPurposeType found provided when registering widget factory."));
 			return false;
 		}
 	}
@@ -231,11 +231,11 @@ bool UEditorDataStorageUi::RegisterWidgetFactory(FName Purpose, TUniquePtr<FType
 		{
 			switch (PurposeInfo->Type)
 			{
-			case IEditorDataStorageUiProvider::EPurposeType::Generic:
+			case IUiProvider::EPurposeType::Generic:
 				PurposeInfo->Factories.Emplace(MoveTemp(Constructor));
 				PurposeInfo->bIsSorted = false;
 				return true;
-			case IEditorDataStorageUiProvider::EPurposeType::UniqueByName:
+			case IUiProvider::EPurposeType::UniqueByName:
 				if (!Columns.IsEmpty())
 				{
 					if (PurposeInfo->Factories.IsEmpty())
@@ -253,7 +253,7 @@ bool UEditorDataStorageUi::RegisterWidgetFactory(FName Purpose, TUniquePtr<FType
 				{
 					return false;
 				}
-			case IEditorDataStorageUiProvider::EPurposeType::UniqueByNameAndColumn:
+			case IUiProvider::EPurposeType::UniqueByNameAndColumn:
 				if (!Columns.IsEmpty())
 				{
 					PurposeInfo->Factories.Emplace(MoveTemp(Constructor), MoveTemp(Columns));
@@ -265,7 +265,7 @@ bool UEditorDataStorageUi::RegisterWidgetFactory(FName Purpose, TUniquePtr<FType
 					return false;
 				}
 			default:
-				checkf(false, TEXT("Unexpected IEditorDataStorageUiProvider::EPurposeType found provided when registering widget factory."));
+				checkf(false, TEXT("Unexpected IUiProvider::EPurposeType found provided when registering widget factory."));
 				return false;
 			}
 		}
@@ -283,7 +283,8 @@ bool UEditorDataStorageUi::RegisterWidgetFactory(FName Purpose, TUniquePtr<FType
 }
 
 void UEditorDataStorageUi::CreateWidgetConstructors(FName Purpose,
-	const UE::Editor::DataStorage::FMetaDataView& Arguments, const WidgetConstructorCallback& Callback)
+	const UE::Editor::DataStorage::FMetaDataView& Arguments,
+	const UE::Editor::DataStorage::IUiProvider::WidgetConstructorCallback& Callback)
 {
 	if (FPurposeInfo* PurposeInfo = WidgetPurposes.Find(Purpose))
 	{
@@ -299,7 +300,7 @@ void UEditorDataStorageUi::CreateWidgetConstructors(FName Purpose,
 
 void UEditorDataStorageUi::CreateWidgetConstructors(FName Purpose, EMatchApproach MatchApproach, 
 	TArray<TWeakObjectPtr<const UScriptStruct>>& Columns, const UE::Editor::DataStorage::FMetaDataView& Arguments,
-	const WidgetConstructorCallback& Callback)
+	const UE::Editor::DataStorage::IUiProvider::WidgetConstructorCallback& Callback)
 {
 	if (FPurposeInfo* PurposeInfo = WidgetPurposes.Find(Purpose))
 	{
@@ -356,7 +357,7 @@ void UEditorDataStorageUi::CreateWidgetConstructors(FName Purpose, EMatchApproac
 }
 
 void UEditorDataStorageUi::ConstructWidgets(FName Purpose, const UE::Editor::DataStorage::FMetaDataView& Arguments,
-	const WidgetCreatedCallback& ConstructionCallback)
+	const UE::Editor::DataStorage::IUiProvider::WidgetCreatedCallback& ConstructionCallback)
 {
 	if (FPurposeInfo* PurposeInfo = WidgetPurposes.Find(Purpose))
 	{
@@ -393,14 +394,15 @@ bool UEditorDataStorageUi::CreateSingleWidgetConstructor(
 	const UE::Editor::DataStorage::FMetaDataView& Arguments,
 	TArray<TWeakObjectPtr<const UScriptStruct>> MatchedColumnTypes,
 	const UE::Editor::DataStorage::Queries::FConditions& QueryConditions,
-	const WidgetConstructorCallback& Callback)
+	const UE::Editor::DataStorage::IUiProvider::WidgetConstructorCallback& Callback)
 {
+	using namespace UE::Editor::DataStorage;
 	struct Visitor
 	{
 		Visitor(
 			TArray<TWeakObjectPtr<const UScriptStruct>>&& InMatchedColumnTypes,
-			const UE::Editor::DataStorage::Queries::FConditions& InQueryConditions,
-			const UE::Editor::DataStorage::FMetaDataView& InArguments,
+			const Queries::FConditions& InQueryConditions,
+			const FMetaDataView& InArguments,
 			const WidgetConstructorCallback& InCallback) 
 			: MatchedColumnTypes(MoveTemp(InMatchedColumnTypes))
 			, QueryConditions(InQueryConditions)
@@ -409,8 +411,8 @@ bool UEditorDataStorageUi::CreateSingleWidgetConstructor(
 		{}
 
 		TArray<TWeakObjectPtr<const UScriptStruct>> MatchedColumnTypes;
-		const UE::Editor::DataStorage::Queries::FConditions& QueryConditions;
-		const UE::Editor::DataStorage::FMetaDataView& Arguments;
+		const Queries::FConditions& QueryConditions;
+		const FMetaDataView& Arguments;
 		const WidgetConstructorCallback& Callback;
 
 		bool operator()(const UScriptStruct* Target)
@@ -449,7 +451,7 @@ bool UEditorDataStorageUi::CreateSingleWidgetConstructor(
 void UEditorDataStorageUi::CreateWidgetInstance(
 	FTypedElementWidgetConstructor& Constructor, 
 	const UE::Editor::DataStorage::FMetaDataView& Arguments,
-	const WidgetCreatedCallback& ConstructionCallback)
+	const UE::Editor::DataStorage::IUiProvider::WidgetCreatedCallback& ConstructionCallback)
 {
 	UE::Editor::DataStorage::RowHandle Row = Storage->AddRow(WidgetTable);
 	Storage->AddColumns(Row, Constructor.GetAdditionalColumnsList());
@@ -573,7 +575,7 @@ void UEditorDataStorageUi::CreateWidgetConstructors_LongestMatch(const TArray<FW
 
 void UEditorDataStorageUi::CreateWidgetConstructors_ExactMatch(const TArray<FWidgetFactory>& WidgetFactories,
 	TArray<TWeakObjectPtr<const UScriptStruct>>& Columns, const UE::Editor::DataStorage::FMetaDataView& Arguments,
-	const WidgetConstructorCallback& Callback)
+	const UE::Editor::DataStorage::IUiProvider::WidgetConstructorCallback& Callback)
 {
 	int32 ColumnCount = Columns.Num();
 	TArray<TWeakObjectPtr<const UScriptStruct>> MatchedColumns;
@@ -606,7 +608,7 @@ void UEditorDataStorageUi::CreateWidgetConstructors_ExactMatch(const TArray<FWid
 
 void UEditorDataStorageUi::CreateWidgetConstructors_SingleMatch(const TArray<FWidgetFactory>& WidgetFactories,
 	TArray<TWeakObjectPtr<const UScriptStruct>>& Columns, const UE::Editor::DataStorage::FMetaDataView& Arguments,
-	const WidgetConstructorCallback& Callback)
+	const UE::Editor::DataStorage::IUiProvider::WidgetConstructorCallback& Callback)
 {
 	auto FactoryIt = WidgetFactories.rbegin();
 	auto FactoryEnd = WidgetFactories.rend();
@@ -678,7 +680,7 @@ UEditorDataStorageUi::FWidgetFactory::FWidgetFactory(TUniquePtr<FTypedElementWid
 		TEXT("Widget constructor registered that didn't contain valid type information."));
 }
 
-const UE::Editor::DataStorage::Queries::FConditions& UEditorDataStorageUi::FWidgetFactory::GetConditions(IEditorDataStorageProvider* DataStorage) const
+const UE::Editor::DataStorage::Queries::FConditions& UEditorDataStorageUi::FWidgetFactory::GetConditions(UE::Editor::DataStorage::ICoreProvider* DataStorage) const
 {
 	Columns.Compile(UE::Editor::DataStorage::Queries::FEditorStorageQueryConditionCompileContext(DataStorage));
 	
diff --git a/Engine/Plugins/Experimental/EditorDataStorage/Source/TedsCore/Private/TypedElementDatabaseUI.h b/Engine/Plugins/Experimental/EditorDataStorage/Source/TedsCore/Private/TypedElementDatabaseUI.h
--- a/Engine/Plugins/Experimental/EditorDataStorage/Source/TedsCore/Private/TypedElementDatabaseUI.h
+++ b/Engine/Plugins/Experimental/EditorDataStorage/Source/TedsCore/Private/TypedElementDatabaseUI.h
@@ -11,15 +11,18 @@
 
 #include "TypedElementDatabaseUI.generated.h"
 
-class IEditorDataStorageProvider;
-class IEditorDataStorageCompatibilityProvider;
+namespace UE::Editor::DataStorage
+{
+	class ICoreProvider;
+	class ICompatibilityProvider;
+} // namespace UE::Editor::DataStorage
 
 TEDSCORE_API DECLARE_LOG_CATEGORY_EXTERN(LogEditorDataStorageUI, Log, All);
 
 UCLASS()
 class TEDSCORE_API UEditorDataStorageUi final
 	: public UObject
-	, public IEditorDataStorageUiProvider
+	, public UE::Editor::DataStorage::IUiProvider
 {
 	GENERATED_BODY()
 
@@ -27,8 +30,8 @@ public:
 	~UEditorDataStorageUi() override = default;
 
 	void Initialize(
-		IEditorDataStorageProvider* StorageInterface, 
-		IEditorDataStorageCompatibilityProvider* StorageCompatibilityInterface);
+		UE::Editor::DataStorage::ICoreProvider* StorageInterface, 
+		UE::Editor::DataStorage::ICompatibilityProvider* StorageCompatibilityInterface);
 	void Deinitialize();
 
 	void RegisterWidgetPurpose(FName Purpose, EPurposeType Type, FText Description) override;
@@ -41,12 +44,14 @@ public:
 		UE::Editor::DataStorage::Queries::FConditions Columns) override;
 
 	void CreateWidgetConstructors(FName Purpose,
-		const UE::Editor::DataStorage::FMetaDataView& Arguments, const WidgetConstructorCallback& Callback) override;
+		const UE::Editor::DataStorage::FMetaDataView& Arguments,
+		const UE::Editor::DataStorage::IUiProvider::WidgetConstructorCallback& Callback) override;
 	void CreateWidgetConstructors(FName Purpose, EMatchApproach MatchApproach, TArray<TWeakObjectPtr<const UScriptStruct>>& Columns,
-		const UE::Editor::DataStorage::FMetaDataView& Arguments, const WidgetConstructorCallback& Callback) override;
+		const UE::Editor::DataStorage::FMetaDataView& Arguments,
+		const UE::Editor::DataStorage::IUiProvider::WidgetConstructorCallback& Callback) override;
 
 	void ConstructWidgets(FName Purpose, const UE::Editor::DataStorage::FMetaDataView& Arguments,
-		const WidgetCreatedCallback& ConstructionCallback) override;
+		const UE::Editor::DataStorage::IUiProvider::WidgetCreatedCallback& ConstructionCallback) override;
 	TSharedPtr<SWidget> ConstructWidget(UE::Editor::DataStorage::RowHandle Row, FTypedElementWidgetConstructor& Constructor,
 		const UE::Editor::DataStorage::FMetaDataView& Arguments) override;
 
@@ -73,7 +78,7 @@ private:
 		explicit FWidgetFactory(TUniquePtr<FTypedElementWidgetConstructor>&& InConstructor);
 		FWidgetFactory(const UScriptStruct* InConstructor, UE::Editor::DataStorage::Queries::FConditions&& InColumns);
 		FWidgetFactory(TUniquePtr<FTypedElementWidgetConstructor>&& InConstructor, UE::Editor::DataStorage::Queries::FConditions&& InColumns);
-		const UE::Editor::DataStorage::Queries::FConditions& GetConditions(IEditorDataStorageProvider* DataStorage) const;
+		const UE::Editor::DataStorage::Queries::FConditions& GetConditions(UE::Editor::DataStorage::ICoreProvider* DataStorage) const;
 
 	private:
 		// Private and mutable so we can control access to it via GetConditions() to compile the conditions on demand
@@ -95,33 +100,33 @@ private:
 		const UE::Editor::DataStorage::FMetaDataView& Arguments,
 		TArray<TWeakObjectPtr<const UScriptStruct>> MatchedColumnTypes,
 		const UE::Editor::DataStorage::Queries::FConditions& QueryConditions,
-		const WidgetConstructorCallback& Callback);
+		const UE::Editor::DataStorage::IUiProvider::WidgetConstructorCallback& Callback);
 
 	void CreateWidgetInstance(
 		FTypedElementWidgetConstructor& Constructor,
 		const UE::Editor::DataStorage::FMetaDataView& Arguments,
-		const WidgetCreatedCallback& ConstructionCallback);
+		const UE::Editor::DataStorage::IUiProvider::WidgetCreatedCallback& ConstructionCallback);
 
 	void CreateWidgetConstructors_LongestMatch(
 		const TArray<FWidgetFactory>& WidgetFactories, 
 		TArray<TWeakObjectPtr<const UScriptStruct>>& Columns, 
 		const UE::Editor::DataStorage::FMetaDataView& Arguments,
-		const WidgetConstructorCallback& Callback);
+		const UE::Editor::DataStorage::IUiProvider::WidgetConstructorCallback& Callback);
 	void CreateWidgetConstructors_ExactMatch(
 		const TArray<FWidgetFactory>& WidgetFactories,
 		TArray<TWeakObjectPtr<const UScriptStruct>>& Columns,
 		const UE::Editor::DataStorage::FMetaDataView& Arguments,
-		const WidgetConstructorCallback& Callback);
+		const UE::Editor::DataStorage::IUiProvider::WidgetConstructorCallback& Callback);
 	void CreateWidgetConstructors_SingleMatch(
 		const TArray<FWidgetFactory>& WidgetFactories,
 		TArray<TWeakObjectPtr<const UScriptStruct>>& Columns,
 		const UE::Editor::DataStorage::FMetaDataView& Arguments,
-		const WidgetConstructorCallback& Callback);
+		const UE::Editor::DataStorage::IUiProvider::WidgetConstructorCallback& Callback);
 
 	UE::Editor::DataStorage::TableHandle WidgetTable{ UE::Editor::DataStorage::InvalidTableHandle };
 	
 	TMap<FName, FPurposeInfo> WidgetPurposes;
 	
-	IEditorDataStorageProvider* Storage{ nullptr };
-	IEditorDataStorageCompatibilityProvider* StorageCompatibility{ nullptr };
+	UE::Editor::DataStorage::ICoreProvider* Storage{ nullptr };
+	UE::Editor::DataStorage::ICompatibilityProvider* StorageCompatibility{ nullptr };
 };
\ No newline at end of file
diff --git a/Engine/Plugins/Experimental/EditorDataStorage/Source/TedsCore/Private/Widgets/STedsWidget.cpp b/Engine/Plugins/Experimental/EditorDataStorage/Source/TedsCore/Private/Widgets/STedsWidget.cpp
--- a/Engine/Plugins/Experimental/EditorDataStorage/Source/TedsCore/Private/Widgets/STedsWidget.cpp
+++ b/Engine/Plugins/Experimental/EditorDataStorage/Source/TedsCore/Private/Widgets/STedsWidget.cpp
@@ -53,8 +53,8 @@ namespace UE::Editor::DataStorage::Widgets
 	void STedsWidget::RegisterTedsWidget(const TSharedPtr<SWidget>& InContentWidget)
 	{
 		using namespace UE::Editor::DataStorage;
-		IEditorDataStorageProvider* Storage = GetStorageIfAvailable();
-		IEditorDataStorageUiProvider* StorageUi = GetStorageUiIfAvailable();
+		ICoreProvider* Storage = GetStorageIfAvailable();
+		IUiProvider* StorageUi = GetStorageUiIfAvailable();
 
 		// If TEDS is not enabled, STedsWidget will just behave like a regular widget
 		if (!Storage || !StorageUi)
@@ -79,7 +79,7 @@ namespace UE::Editor::DataStorage::Widgets
 
 	void STedsWidget::SetContent(const TSharedRef< SWidget >& InContent)
 	{
-		if (IEditorDataStorageProvider* Storage = GetStorageIfAvailable())
+		if (ICoreProvider* Storage = GetStorageIfAvailable())
 		{
 			if (FTypedElementSlateWidgetReferenceColumn* WidgetReferenceColumn = Storage->GetColumn<FTypedElementSlateWidgetReferenceColumn>(UiRowHandle))
 			{
@@ -126,15 +126,15 @@ namespace UE::Editor::DataStorage::Widgets
 		return SharedThis(this);
 	}
 
-	IEditorDataStorageProvider* STedsWidget::GetStorageIfAvailable()
+	ICoreProvider* STedsWidget::GetStorageIfAvailable()
 	{
 		using namespace UE::Editor::DataStorage;
-		return GetMutableDataStorageFeature<IEditorDataStorageProvider>(StorageFeatureName);
+		return GetMutableDataStorageFeature<ICoreProvider>(StorageFeatureName);
 	}
 
-	IEditorDataStorageUiProvider* STedsWidget::GetStorageUiIfAvailable()
+	IUiProvider* STedsWidget::GetStorageUiIfAvailable()
 	{
 		using namespace UE::Editor::DataStorage;
-		return GetMutableDataStorageFeature<IEditorDataStorageUiProvider>(UiFeatureName);
+		return GetMutableDataStorageFeature<IUiProvider>(UiFeatureName);
 	}
 }
\ No newline at end of file
diff --git a/Engine/Plugins/Experimental/EditorDataStorage/Source/TedsCore/Private/Widgets/STedsWidget.h b/Engine/Plugins/Experimental/EditorDataStorage/Source/TedsCore/Private/Widgets/STedsWidget.h
--- a/Engine/Plugins/Experimental/EditorDataStorage/Source/TedsCore/Private/Widgets/STedsWidget.h
+++ b/Engine/Plugins/Experimental/EditorDataStorage/Source/TedsCore/Private/Widgets/STedsWidget.h
@@ -8,7 +8,10 @@
 
 struct FTypedElementWidgetConstructor;
 
-class IEditorDataStorageUiProvider;
+namespace UE::Editor::DataStorage
+{
+	class IUiProvider;
+} // namespace UE::Editor::DataStorage
 
 namespace UE::Editor::DataStorage::Widgets
 {
@@ -50,8 +53,8 @@ namespace UE::Editor::DataStorage::Widgets
 
 		void RegisterTedsWidget(const TSharedPtr<SWidget>& InContentWidget);
 	
-		static IEditorDataStorageProvider* GetStorageIfAvailable();
-		static IEditorDataStorageUiProvider* GetStorageUiIfAvailable();
+		static ICoreProvider* GetStorageIfAvailable();
+		static IUiProvider* GetStorageUiIfAvailable();
 
 	private:
 	
diff --git a/Engine/Plugins/Experimental/EditorDataStorage/Source/TedsCore/Tests/TypedElementAttributeBindingTests.cpp b/Engine/Plugins/Experimental/EditorDataStorage/Source/TedsCore/Tests/TypedElementAttributeBindingTests.cpp
--- a/Engine/Plugins/Experimental/EditorDataStorage/Source/TedsCore/Tests/TypedElementAttributeBindingTests.cpp
+++ b/Engine/Plugins/Experimental/EditorDataStorage/Source/TedsCore/Tests/TypedElementAttributeBindingTests.cpp
@@ -12,7 +12,7 @@ namespace UE::Editor::DataStorage::Tests
 {
 	BEGIN_DEFINE_SPEC(TedsAttributeBindingTestsFixture, "Editor.DataStorage.AttributeBinding", EAutomationTestFlags::EditorContext | EAutomationTestFlags::EngineFilter)
 
-	IEditorDataStorageProvider* TedsInterface = nullptr;
+	ICoreProvider* TedsInterface = nullptr;
 	const FName TestTableName = TEXT("TestTable_AttributeBinding");
 	TableHandle TestTableHandle = InvalidTableHandle;
 	RowHandle TestRowHandle = InvalidRowHandle;
@@ -51,7 +51,7 @@ namespace UE::Editor::DataStorage::Tests
 	{
 		BeforeEach([this]()
 		{
-			TedsInterface = GetMutableDataStorageFeature<IEditorDataStorageProvider>(StorageFeatureName);
+			TedsInterface = GetMutableDataStorageFeature<ICoreProvider>(StorageFeatureName);
 			TestTrue("", TedsInterface != nullptr);
 			
 			TestTableHandle = RegisterTestTable();
diff --git a/Engine/Plugins/Experimental/EditorDataStorage/Source/TedsCore/Tests/TypedElementDataStorageInterfaceTests.cpp b/Engine/Plugins/Experimental/EditorDataStorage/Source/TedsCore/Tests/TypedElementDataStorageInterfaceTests.cpp
--- a/Engine/Plugins/Experimental/EditorDataStorage/Source/TedsCore/Tests/TypedElementDataStorageInterfaceTests.cpp
+++ b/Engine/Plugins/Experimental/EditorDataStorage/Source/TedsCore/Tests/TypedElementDataStorageInterfaceTests.cpp
@@ -10,7 +10,7 @@
 
 namespace UE::Editor::DataStorage::Tests
 {
-	TableHandle RegisterTestTableA(IEditorDataStorageProvider* TedsInterface)
+	TableHandle RegisterTestTableA(ICoreProvider* TedsInterface)
 	{
 		// RegisterTable is not idempotent, so this utility function makes it so
 	
@@ -56,7 +56,7 @@ namespace UE::Editor::DataStorage::Tests
 	template<typename ColumnType>
 	void RegisterTestObserver(
 		FName ObserverName,
-		IEditorDataStorageProvider* TedsInterface, Queries::FObserver::EEvent EventType, int64* QueryCallCountPtr, QueryHandle* QueryHandleInOut)
+		ICoreProvider* TedsInterface, Queries::FObserver::EEvent EventType, int64* QueryCallCountPtr, QueryHandle* QueryHandleInOut)
 	{
 		using namespace UE::Editor::DataStorage::Queries;
 		FQueryDescription QueryDescription = TedsInterface->GetQueryDescription(*QueryHandleInOut);
@@ -78,7 +78,7 @@ namespace UE::Editor::DataStorage::Tests
 	}
 
 	BEGIN_DEFINE_SPEC(EditorDataStorageTestsFixture, "Editor.DataStorage.Core", EAutomationTestFlags::EditorContext | EAutomationTestFlags::EngineFilter)
-		IEditorDataStorageProvider* TedsInterface = nullptr;
+		ICoreProvider* TedsInterface = nullptr;
 		TableHandle TestTableHandleA = InvalidTableHandle;
 
 		TArray<RowHandle> CreatedRows;
@@ -128,7 +128,7 @@ namespace UE::Editor::DataStorage::Tests
 	{
 		BeforeEach([this]()
 		{
-			TedsInterface = GetMutableDataStorageFeature<IEditorDataStorageProvider>(StorageFeatureName);
+			TedsInterface = GetMutableDataStorageFeature<ICoreProvider>(StorageFeatureName);
 			TestTrue("", TedsInterface != nullptr);
 		});
 
diff --git a/Engine/Plugins/Experimental/EditorDataStorage/Source/TedsUI/Private/Processors/WidgetReferenceColumnUpdateProcessor.cpp b/Engine/Plugins/Experimental/EditorDataStorage/Source/TedsUI/Private/Processors/WidgetReferenceColumnUpdateProcessor.cpp
--- a/Engine/Plugins/Experimental/EditorDataStorage/Source/TedsUI/Private/Processors/WidgetReferenceColumnUpdateProcessor.cpp
+++ b/Engine/Plugins/Experimental/EditorDataStorage/Source/TedsUI/Private/Processors/WidgetReferenceColumnUpdateProcessor.cpp
@@ -5,13 +5,13 @@
 #include "Elements/Columns/TypedElementSlateWidgetColumns.h"
 #include "Elements/Framework/TypedElementQueryBuilder.h"
 
-void UWidgetReferenceColumnUpdateFactory::RegisterQueries(IEditorDataStorageProvider& DataStorage)
+void UWidgetReferenceColumnUpdateFactory::RegisterQueries(UE::Editor::DataStorage::ICoreProvider& DataStorage)
 {
 	RegisterDeleteRowOnWidgetDeleteQuery(DataStorage);
 	RegisterDeleteColumnOnWidgetDeleteQuery(DataStorage);
 }
 
-void UWidgetReferenceColumnUpdateFactory::RegisterDeleteRowOnWidgetDeleteQuery(IEditorDataStorageProvider& DataStorage) const
+void UWidgetReferenceColumnUpdateFactory::RegisterDeleteRowOnWidgetDeleteQuery(UE::Editor::DataStorage::ICoreProvider& DataStorage) const
 {
 	using namespace UE::Editor::DataStorage::Queries;
 
@@ -31,7 +31,7 @@ void UWidgetReferenceColumnUpdateFactory::RegisterDeleteRowOnWidgetDeleteQuery(I
     	.Compile());
 }
 
-void UWidgetReferenceColumnUpdateFactory::RegisterDeleteColumnOnWidgetDeleteQuery(IEditorDataStorageProvider& DataStorage) const
+void UWidgetReferenceColumnUpdateFactory::RegisterDeleteColumnOnWidgetDeleteQuery(UE::Editor::DataStorage::ICoreProvider& DataStorage) const
 {
 	using namespace UE::Editor::DataStorage::Queries;
 
diff --git a/Engine/Plugins/Experimental/EditorDataStorage/Source/TedsUI/Private/Processors/WidgetReferenceColumnUpdateProcessor.h b/Engine/Plugins/Experimental/EditorDataStorage/Source/TedsUI/Private/Processors/WidgetReferenceColumnUpdateProcessor.h
--- a/Engine/Plugins/Experimental/EditorDataStorage/Source/TedsUI/Private/Processors/WidgetReferenceColumnUpdateProcessor.h
+++ b/Engine/Plugins/Experimental/EditorDataStorage/Source/TedsUI/Private/Processors/WidgetReferenceColumnUpdateProcessor.h
@@ -20,9 +20,9 @@ class UWidgetReferenceColumnUpdateFactory : public UEditorDataStorageFactory
 public:
 	~UWidgetReferenceColumnUpdateFactory() override = default;
 
-	void RegisterQueries(IEditorDataStorageProvider& DataStorage) override;
+	void RegisterQueries(UE::Editor::DataStorage::ICoreProvider& DataStorage) override;
 
 private:
-	void RegisterDeleteRowOnWidgetDeleteQuery(IEditorDataStorageProvider& DataStorage) const;
-	void RegisterDeleteColumnOnWidgetDeleteQuery(IEditorDataStorageProvider& DataStorage) const;
+	void RegisterDeleteRowOnWidgetDeleteQuery(UE::Editor::DataStorage::ICoreProvider& DataStorage) const;
+	void RegisterDeleteColumnOnWidgetDeleteQuery(UE::Editor::DataStorage::ICoreProvider& DataStorage) const;
 };
diff --git a/Engine/Plugins/Experimental/EditorDataStorage/Source/TedsUI/Private/TedsUIModule.cpp b/Engine/Plugins/Experimental/EditorDataStorage/Source/TedsUI/Private/TedsUIModule.cpp
--- a/Engine/Plugins/Experimental/EditorDataStorage/Source/TedsUI/Private/TedsUIModule.cpp
+++ b/Engine/Plugins/Experimental/EditorDataStorage/Source/TedsUI/Private/TedsUIModule.cpp
@@ -17,11 +17,11 @@ namespace UE::Editor::DataStorage::Private
 	TEXT("Prints a list of all the known widget purposes."),
 	FConsoleCommandWithOutputDeviceDelegate::CreateLambda([](FOutputDevice& Output)
 		{
-			if (IEditorDataStorageUiProvider* UiStorage = GetMutableDataStorageFeature<IEditorDataStorageUiProvider>(UiFeatureName))
+			if (IUiProvider* UiStorage = GetMutableDataStorageFeature<IUiProvider>(UiFeatureName))
 			{
 				Output.Log(TEXT("The Typed Elements Data Storage has recorded the following widget purposes:"));
 				UiStorage->ListWidgetPurposes(
-					[&Output](FName Purpose, IEditorDataStorageUiProvider::EPurposeType, const FText& Description)
+					[&Output](FName Purpose, IUiProvider::EPurposeType, const FText& Description)
 					{
 						Output.Logf(TEXT("    %s - %s"), *Purpose.ToString(), *Description.ToString());
 					});
diff --git a/Engine/Plugins/Experimental/EditorDataStorage/Source/TedsUI/Private/Widgets/CounterWidget.cpp b/Engine/Plugins/Experimental/EditorDataStorage/Source/TedsUI/Private/Widgets/CounterWidget.cpp
--- a/Engine/Plugins/Experimental/EditorDataStorage/Source/TedsUI/Private/Widgets/CounterWidget.cpp
+++ b/Engine/Plugins/Experimental/EditorDataStorage/Source/TedsUI/Private/Widgets/CounterWidget.cpp
@@ -44,7 +44,7 @@ UCounterWidgetFactory::UCounterWidgetFactory()
 	}
 }
 
-void UCounterWidgetFactory::RegisterQueries(IEditorDataStorageProvider& DataStorage)
+void UCounterWidgetFactory::RegisterQueries(UE::Editor::DataStorage::ICoreProvider& DataStorage)
 {
 	using namespace UE::Editor::DataStorage::Queries;
 	
@@ -78,14 +78,15 @@ void UCounterWidgetFactory::RegisterQueries(IEditorDataStorageProvider& DataStor
 	.Compile());
 }
 
-void UCounterWidgetFactory::RegisterWidgetPurposes(IEditorDataStorageUiProvider& DataStorageUi) const
+void UCounterWidgetFactory::RegisterWidgetPurposes(UE::Editor::DataStorage::IUiProvider& DataStorageUi) const
 {
-	DataStorageUi.RegisterWidgetPurpose(WigetPurpose, IEditorDataStorageUiProvider::EPurposeType::Generic,
+	using namespace UE::Editor::DataStorage;
+	DataStorageUi.RegisterWidgetPurpose(WigetPurpose, IUiProvider::EPurposeType::Generic,
 		LOCTEXT("ToolBarPurposeDescription", "Widgets added to the status bar at the bottom editor of the main editor window."));
 }
 
-void UCounterWidgetFactory::RegisterWidgetConstructors(IEditorDataStorageProvider& DataStorage,
-	IEditorDataStorageUiProvider& DataStorageUi) const
+void UCounterWidgetFactory::RegisterWidgetConstructors(UE::Editor::DataStorage::ICoreProvider& DataStorage,
+	UE::Editor::DataStorage::IUiProvider& DataStorageUi) const
 {
 	using namespace UE::Editor::DataStorage::Queries;
 
@@ -114,7 +115,7 @@ void UCounterWidgetFactory::SetupMainWindowIntegrations(TSharedPtr<SWindow> Pare
 	{
 		using namespace UE::Editor::DataStorage;
 
-		IEditorDataStorageUiProvider* UiInterface = GetMutableDataStorageFeature<IEditorDataStorageUiProvider>(UiFeatureName);
+		IUiProvider* UiInterface = GetMutableDataStorageFeature<IUiProvider>(UiFeatureName);
 		checkf(UiInterface, TEXT(
 			"FEditorDataStorageUiModule tried to integrate with the main window before the "
 			"TEDS UI interface is available."));
@@ -159,6 +160,7 @@ FCounterWidgetConstructor::FCounterWidgetConstructor()
 
 TConstArrayView<const UScriptStruct*> FCounterWidgetConstructor::GetAdditionalColumnsList() const
 {
+	using namespace UE::Editor::DataStorage;
 	static TTypedElementColumnTypeList<FCounterWidgetColumn, FTypedElementU32IntValueCacheColumn> Columns;
 	return Columns;
 }
@@ -172,7 +174,7 @@ TSharedPtr<SWidget> FCounterWidgetConstructor::CreateWidget(const UE::Editor::Da
 		.Justification(ETextJustify::Center);
 }
 
-bool FCounterWidgetConstructor::SetColumns(IEditorDataStorageProvider* DataStorage, UE::Editor::DataStorage::RowHandle Row)
+bool FCounterWidgetConstructor::SetColumns(UE::Editor::DataStorage::ICoreProvider* DataStorage, UE::Editor::DataStorage::RowHandle Row)
 {
 	FCounterWidgetColumn* CounterColumn = DataStorage->GetColumn<FCounterWidgetColumn>(Row);
 	checkf(CounterColumn, TEXT("Added a new FCounterWidgetColumn to the Typed Elements Data Storage, but didn't get a valid pointer back."));
diff --git a/Engine/Plugins/Experimental/EditorDataStorage/Source/TedsUI/Private/Widgets/ExportedTextWidget.cpp b/Engine/Plugins/Experimental/EditorDataStorage/Source/TedsUI/Private/Widgets/ExportedTextWidget.cpp
--- a/Engine/Plugins/Experimental/EditorDataStorage/Source/TedsUI/Private/Widgets/ExportedTextWidget.cpp
+++ b/Engine/Plugins/Experimental/EditorDataStorage/Source/TedsUI/Private/Widgets/ExportedTextWidget.cpp
@@ -15,8 +15,8 @@
 // UExportedTextWidgetFactory
 //
 
-void UExportedTextWidgetFactory::RegisterWidgetConstructors(IEditorDataStorageProvider& DataStorage,
-	IEditorDataStorageUiProvider& DataStorageUi) const
+void UExportedTextWidgetFactory::RegisterWidgetConstructors(UE::Editor::DataStorage::ICoreProvider& DataStorage,
+	UE::Editor::DataStorage::IUiProvider& DataStorageUi) const
 {
 	DataStorageUi.RegisterWidgetFactory(FName(TEXT("General.Cell.Default")), FExportedTextWidgetConstructor::StaticStruct());
 }
@@ -32,6 +32,8 @@ FExportedTextWidgetConstructor::FExportedTextWidgetConstructor()
 
 TConstArrayView<const UScriptStruct*> FExportedTextWidgetConstructor::GetAdditionalColumnsList() const
 {
+	using namespace UE::Editor::DataStorage;
+
 	static TTypedElementColumnTypeList<
     		FTypedElementRowReferenceColumn,
     		FTypedElementScriptStructTypeInfoColumn,
@@ -50,7 +52,7 @@ const UE::Editor::DataStorage::Queries::FConditions* FExportedTextWidgetConstruc
 	return nullptr;
 }
 
-FText FExportedTextWidgetConstructor::CreateWidgetDisplayNameText(IEditorDataStorageProvider* DataStorage,
+FText FExportedTextWidgetConstructor::CreateWidgetDisplayNameText(UE::Editor::DataStorage::ICoreProvider* DataStorage,
 	UE::Editor::DataStorage::RowHandle Row) const
 {
 	if (FTypedElementScriptStructTypeInfoColumn* TypeInfoColumn = DataStorage->GetColumn<FTypedElementScriptStructTypeInfoColumn>(Row))
@@ -63,7 +65,8 @@ FText FExportedTextWidgetConstructor::CreateWidgetDisplayNameText(IEditorDataSto
 	}
 }
 
-TSharedPtr<SWidget> FExportedTextWidgetConstructor::CreateWidget(IEditorDataStorageProvider* DataStorage, IEditorDataStorageUiProvider* DataStorageUi,
+TSharedPtr<SWidget> FExportedTextWidgetConstructor::CreateWidget(UE::Editor::DataStorage::ICoreProvider* DataStorage,
+	UE::Editor::DataStorage::IUiProvider* DataStorageUi,
 	UE::Editor::DataStorage::RowHandle TargetRow, UE::Editor::DataStorage::RowHandle WidgetRow,
 	const UE::Editor::DataStorage::FMetaDataView& Arguments)
 {
diff --git a/Engine/Plugins/Experimental/EditorDataStorage/Source/TedsUI/Private/Widgets/GeneralWidgetRegistrationFactory.cpp b/Engine/Plugins/Experimental/EditorDataStorage/Source/TedsUI/Private/Widgets/GeneralWidgetRegistrationFactory.cpp
--- a/Engine/Plugins/Experimental/EditorDataStorage/Source/TedsUI/Private/Widgets/GeneralWidgetRegistrationFactory.cpp
+++ b/Engine/Plugins/Experimental/EditorDataStorage/Source/TedsUI/Private/Widgets/GeneralWidgetRegistrationFactory.cpp
@@ -12,21 +12,22 @@ const FName UGeneralWidgetRegistrationFactory::HeaderPurpose(TEXT("General.Heade
 const FName UGeneralWidgetRegistrationFactory::CellDefaultPurpose(TEXT("General.Cell.Default"));
 const FName UGeneralWidgetRegistrationFactory::HeaderDefaultPurpose(TEXT("General.Header.Default"));
 
-void UGeneralWidgetRegistrationFactory::RegisterWidgetPurposes(IEditorDataStorageUiProvider& DataStorageUi) const
+void UGeneralWidgetRegistrationFactory::RegisterWidgetPurposes(UE::Editor::DataStorage::IUiProvider& DataStorageUi) const
 {
-	DataStorageUi.RegisterWidgetPurpose(CellPurpose, IEditorDataStorageUiProvider::EPurposeType::UniqueByNameAndColumn,
+	using namespace UE::Editor::DataStorage;
+	DataStorageUi.RegisterWidgetPurpose(CellPurpose, IUiProvider::EPurposeType::UniqueByNameAndColumn,
 		LOCTEXT("GeneralCellPurpose", "General purpose widgets that can be used as cells for specific columns or column combinations."));
 
-	DataStorageUi.RegisterWidgetPurpose(LargeCellPurpose, IEditorDataStorageUiProvider::EPurposeType::UniqueByNameAndColumn,
+	DataStorageUi.RegisterWidgetPurpose(LargeCellPurpose, IUiProvider::EPurposeType::UniqueByNameAndColumn,
 		LOCTEXT("GeneralLargeCellPurpose", "General purpose widgets that are specifically designed to be embedded in a space larger than a single cell."));
 
-	DataStorageUi.RegisterWidgetPurpose(HeaderPurpose, IEditorDataStorageUiProvider::EPurposeType::UniqueByNameAndColumn,
+	DataStorageUi.RegisterWidgetPurpose(HeaderPurpose, IUiProvider::EPurposeType::UniqueByNameAndColumn,
 		LOCTEXT("GeneralHeaderPurpose", "General purpose widget that can be used as a header."));
 
-	DataStorageUi.RegisterWidgetPurpose(HeaderDefaultPurpose, IEditorDataStorageUiProvider::EPurposeType::UniqueByName,
+	DataStorageUi.RegisterWidgetPurpose(HeaderDefaultPurpose, IUiProvider::EPurposeType::UniqueByName,
 		LOCTEXT("GeneralHeaderDefaultPurpose", "The default widget to use in headers if no other specialization is provided."));
 	
-	DataStorageUi.RegisterWidgetPurpose(CellDefaultPurpose, IEditorDataStorageUiProvider::EPurposeType::UniqueByName,
+	DataStorageUi.RegisterWidgetPurpose(CellDefaultPurpose, IUiProvider::EPurposeType::UniqueByName,
 		LOCTEXT("GeneralCellDefaultPurpose", "The default widget to use in cells if no other specialization is provided."));
 }
 
diff --git a/Engine/Plugins/Experimental/EditorDataStorage/Source/TedsUI/Private/Widgets/LabelWidget.cpp b/Engine/Plugins/Experimental/EditorDataStorage/Source/TedsUI/Private/Widgets/LabelWidget.cpp
--- a/Engine/Plugins/Experimental/EditorDataStorage/Source/TedsUI/Private/Widgets/LabelWidget.cpp
+++ b/Engine/Plugins/Experimental/EditorDataStorage/Source/TedsUI/Private/Widgets/LabelWidget.cpp
@@ -24,8 +24,8 @@
 // ULabelWidgetFactory
 //
 
-void ULabelWidgetFactory::RegisterWidgetConstructors(IEditorDataStorageProvider& DataStorage,
-	IEditorDataStorageUiProvider& DataStorageUi) const
+void ULabelWidgetFactory::RegisterWidgetConstructors(UE::Editor::DataStorage::ICoreProvider& DataStorage,
+	UE::Editor::DataStorage::IUiProvider& DataStorageUi) const
 {
 	using namespace UE::Editor::DataStorage::Queries;
 
@@ -36,12 +36,13 @@ void ULabelWidgetFactory::RegisterWidgetConstructors(IEditorDataStorageProvider&
 		TColumn<FTypedElementLabelColumn>() || (TColumn<FTypedElementLabelColumn>() && TColumn<FTypedElementLabelHashColumn>()));
 }
 
-void ULabelWidgetFactory::RegisterWidgetPurposes(IEditorDataStorageUiProvider& DataStorageUi) const
+void ULabelWidgetFactory::RegisterWidgetPurposes(UE::Editor::DataStorage::IUiProvider& DataStorageUi) const
 {
-	DataStorageUi.RegisterWidgetPurpose(FName(TEXT("General.RowLabel")), IEditorDataStorageUiProvider::EPurposeType::UniqueByNameAndColumn,
+	using namespace UE::Editor::DataStorage;
+	DataStorageUi.RegisterWidgetPurpose(FName(TEXT("General.RowLabel")), IUiProvider::EPurposeType::UniqueByNameAndColumn,
 	LOCTEXT("GeneralRowLabelPurpose", "Specific purpose to request a widget to display a user facing display name for a row."));
 	
-	DataStorageUi.RegisterWidgetPurpose(FName(TEXT("General.RowLabel.Default")), IEditorDataStorageUiProvider::EPurposeType::UniqueByName,
+	DataStorageUi.RegisterWidgetPurpose(FName(TEXT("General.RowLabel.Default")), IUiProvider::EPurposeType::UniqueByName,
 	LOCTEXT("GeneralRowLabelDefaultPurpose", "Default purpose to request a widget to display a user facing display name for a row."));
 }
 
@@ -57,6 +58,7 @@ FLabelWidgetConstructor::FLabelWidgetConstructor()
 
 TConstArrayView<const UScriptStruct*> FLabelWidgetConstructor::GetAdditionalColumnsList() const
 {
+	using namespace UE::Editor::DataStorage;
 	static const TTypedElementColumnTypeList<
 		FTypedElementRowReferenceColumn,
 		FTypedElementU64IntValueCacheColumn,
@@ -64,8 +66,8 @@ TConstArrayView<const UScriptStruct*> FLabelWidgetConstructor::GetAdditionalColu
 	return Columns;
 }
 
-TSharedPtr<SWidget> FLabelWidgetConstructor::CreateWidget(IEditorDataStorageProvider* DataStorage,
-	IEditorDataStorageUiProvider* DataStorageUi, UE::Editor::DataStorage::RowHandle TargetRow, 
+TSharedPtr<SWidget> FLabelWidgetConstructor::CreateWidget(UE::Editor::DataStorage::ICoreProvider* DataStorage,
+	UE::Editor::DataStorage::IUiProvider* DataStorageUi, UE::Editor::DataStorage::RowHandle TargetRow,
 	UE::Editor::DataStorage::RowHandle WidgetRow, const UE::Editor::DataStorage::FMetaDataView& Arguments)
 {
 	UE::Editor::DataStorage::FAttributeBinder Binder(TargetRow, DataStorage);
diff --git a/Engine/Plugins/Experimental/EditorDataStorage/Source/TedsUI/Private/Widgets/PackagePathWidget.cpp b/Engine/Plugins/Experimental/EditorDataStorage/Source/TedsUI/Private/Widgets/PackagePathWidget.cpp
--- a/Engine/Plugins/Experimental/EditorDataStorage/Source/TedsUI/Private/Widgets/PackagePathWidget.cpp
+++ b/Engine/Plugins/Experimental/EditorDataStorage/Source/TedsUI/Private/Widgets/PackagePathWidget.cpp
@@ -11,8 +11,8 @@
 // UPackagePathWidgetFactory
 //
 
-void UPackagePathWidgetFactory::RegisterWidgetConstructors(IEditorDataStorageProvider& DataStorage,
-	IEditorDataStorageUiProvider& DataStorageUi) const
+void UPackagePathWidgetFactory::RegisterWidgetConstructors(UE::Editor::DataStorage::ICoreProvider& DataStorage,
+	UE::Editor::DataStorage::IUiProvider& DataStorageUi) const
 {
 	using namespace UE::Editor::DataStorage::Queries;
 
@@ -46,8 +46,8 @@ TSharedPtr<SWidget> FPackagePathWidgetConstructor::CreateWidget(const UE::Editor
 }
 
 bool FPackagePathWidgetConstructor::FinalizeWidget(
-	IEditorDataStorageProvider* DataStorage,
-	IEditorDataStorageUiProvider* DataStorageUi,
+	UE::Editor::DataStorage::ICoreProvider* DataStorage,
+	UE::Editor::DataStorage::IUiProvider* DataStorageUi,
 	UE::Editor::DataStorage::RowHandle Row,
 	const TSharedPtr<SWidget>& Widget)
 {
@@ -78,8 +78,8 @@ FLoadedPackagePathWidgetConstructor::FLoadedPackagePathWidgetConstructor()
 }
 
 bool FLoadedPackagePathWidgetConstructor::FinalizeWidget(
-	IEditorDataStorageProvider* DataStorage,
-	IEditorDataStorageUiProvider* DataStorageUi,
+	UE::Editor::DataStorage::ICoreProvider* DataStorage,
+	UE::Editor::DataStorage::IUiProvider* DataStorageUi,
 	UE::Editor::DataStorage::RowHandle Row,
 	const TSharedPtr<SWidget>& Widget)
 {
diff --git a/Engine/Plugins/Experimental/EditorDataStorage/Source/TedsUI/Private/Widgets/SlateColorWidget.cpp b/Engine/Plugins/Experimental/EditorDataStorage/Source/TedsUI/Private/Widgets/SlateColorWidget.cpp
--- a/Engine/Plugins/Experimental/EditorDataStorage/Source/TedsUI/Private/Widgets/SlateColorWidget.cpp
+++ b/Engine/Plugins/Experimental/EditorDataStorage/Source/TedsUI/Private/Widgets/SlateColorWidget.cpp
@@ -14,7 +14,7 @@
 namespace UE::SlateColorWidget::Private
 {
 	FReply SummonColorPicker(const FGeometry& Geometry, const FPointerEvent& PointerEvent,
-		IEditorDataStorageProvider* DataStorage, UE::Editor::DataStorage::RowHandle TargetRow)
+		UE::Editor::DataStorage::ICoreProvider* DataStorage, UE::Editor::DataStorage::RowHandle TargetRow)
 	{
 		if (PointerEvent.GetEffectingButton() != EKeys::LeftMouseButton)
 		{
@@ -40,8 +40,8 @@ namespace UE::SlateColorWidget::Private
 	}
 }
 
-void USlateColorWidgetFactory::RegisterWidgetConstructors(IEditorDataStorageProvider& DataStorage,
-														IEditorDataStorageUiProvider& DataStorageUi) const
+void USlateColorWidgetFactory::RegisterWidgetConstructors(UE::Editor::DataStorage::ICoreProvider& DataStorage,
+	UE::Editor::DataStorage::IUiProvider& DataStorageUi) const
 {
 	using namespace UE::Editor::DataStorage::Queries;
 	DataStorageUi.RegisterWidgetFactory<FSlateColorWidgetConstructor>(TEXT("General.Cell"), TColumn<FSlateColorColumn>());
@@ -54,7 +54,8 @@ FSlateColorWidgetConstructor::FSlateColorWidgetConstructor()
 
 
 
-TSharedPtr<SWidget> FSlateColorWidgetConstructor::CreateWidget(IEditorDataStorageProvider* DataStorage, IEditorDataStorageUiProvider* DataStorageUi, 
+TSharedPtr<SWidget> FSlateColorWidgetConstructor::CreateWidget(UE::Editor::DataStorage::ICoreProvider* DataStorage,
+	UE::Editor::DataStorage::IUiProvider* DataStorageUi,
 	UE::Editor::DataStorage::RowHandle TargetRow, UE::Editor::DataStorage::RowHandle WidgetRow, const UE::Editor::DataStorage::FMetaDataView& Arguments)
 {
 	UE::Editor::DataStorage::FAttributeBinder Binder(TargetRow, DataStorage);
diff --git a/Engine/Plugins/Experimental/EditorDataStorage/Source/TedsUI/Private/Widgets/SlateColorWidget.h b/Engine/Plugins/Experimental/EditorDataStorage/Source/TedsUI/Private/Widgets/SlateColorWidget.h
--- a/Engine/Plugins/Experimental/EditorDataStorage/Source/TedsUI/Private/Widgets/SlateColorWidget.h
+++ b/Engine/Plugins/Experimental/EditorDataStorage/Source/TedsUI/Private/Widgets/SlateColorWidget.h
@@ -7,10 +7,6 @@
 
 #include "SlateColorWidget.generated.h"
 
-class IEditorDataStorageProvider;
-class IEditorDataStorageUiProvider;
-class UScriptStruct;
-
 UCLASS()
 class USlateColorWidgetFactory : public UEditorDataStorageFactory
 {
@@ -19,8 +15,8 @@ class USlateColorWidgetFactory : public UEditorDataStorageFactory
 public:
 	~USlateColorWidgetFactory() override = default;
 
-	void RegisterWidgetConstructors(IEditorDataStorageProvider& DataStorage,
-		IEditorDataStorageUiProvider& DataStorageUi) const override;
+	void RegisterWidgetConstructors(UE::Editor::DataStorage::ICoreProvider& DataStorage,
+		UE::Editor::DataStorage::IUiProvider& DataStorageUi) const override;
 };
 
 // Widget to show and edit the color column in TEDS
@@ -34,8 +30,8 @@ public:
 	~FSlateColorWidgetConstructor() override = default;
 
 	virtual TSharedPtr<SWidget> CreateWidget(
-		IEditorDataStorageProvider* DataStorage,
-		IEditorDataStorageUiProvider* DataStorageUi,
+		UE::Editor::DataStorage::ICoreProvider* DataStorage,
+		UE::Editor::DataStorage::IUiProvider* DataStorageUi,
 		UE::Editor::DataStorage::RowHandle TargetRow,
 		UE::Editor::DataStorage::RowHandle WidgetRow, 
 		const UE::Editor::DataStorage::FMetaDataView& Arguments) override;
diff --git a/Engine/Plugins/Experimental/EditorDataStorage/Source/TedsUI/Private/Widgets/UrlWidget.cpp b/Engine/Plugins/Experimental/EditorDataStorage/Source/TedsUI/Private/Widgets/UrlWidget.cpp
--- a/Engine/Plugins/Experimental/EditorDataStorage/Source/TedsUI/Private/Widgets/UrlWidget.cpp
+++ b/Engine/Plugins/Experimental/EditorDataStorage/Source/TedsUI/Private/Widgets/UrlWidget.cpp
@@ -8,8 +8,8 @@
 #include "Widgets/Input/SHyperlink.h"
 #include "Widgets/SBoxPanel.h"
 
-void UUrlWidgetFactory::RegisterWidgetConstructors(IEditorDataStorageProvider& DataStorage,
-														IEditorDataStorageUiProvider& DataStorageUi) const
+void UUrlWidgetFactory::RegisterWidgetConstructors(UE::Editor::DataStorage::ICoreProvider& DataStorage,
+	UE::Editor::DataStorage::IUiProvider& DataStorageUi) const
 {
 	using namespace UE::Editor::DataStorage::Queries;
 	DataStorageUi.RegisterWidgetFactory<FUrlWidgetConstructor>(TEXT("General.Cell"), TColumn<FUrlColumn>());
@@ -22,11 +22,11 @@ FUrlWidgetConstructor::FUrlWidgetConstructor()
 
 
 
-TSharedPtr<SWidget> FUrlWidgetConstructor::CreateWidget(IEditorDataStorageProvider* DataStorage,
-                                                               IEditorDataStorageUiProvider* DataStorageUi,
-                                                               UE::Editor::DataStorage::RowHandle TargetRow,
-                                                               UE::Editor::DataStorage::RowHandle WidgetRow,
-                                                               const UE::Editor::DataStorage::FMetaDataView& Arguments)
+TSharedPtr<SWidget> FUrlWidgetConstructor::CreateWidget(UE::Editor::DataStorage::ICoreProvider* DataStorage,
+	UE::Editor::DataStorage::IUiProvider* DataStorageUi,
+	UE::Editor::DataStorage::RowHandle TargetRow,
+	UE::Editor::DataStorage::RowHandle WidgetRow,
+	const UE::Editor::DataStorage::FMetaDataView& Arguments)
 {
 	UE::Editor::DataStorage::FAttributeBinder Binder(TargetRow, DataStorage);
 	
diff --git a/Engine/Plugins/Experimental/EditorDataStorage/Source/TedsUI/Private/Widgets/UrlWidget.h b/Engine/Plugins/Experimental/EditorDataStorage/Source/TedsUI/Private/Widgets/UrlWidget.h
--- a/Engine/Plugins/Experimental/EditorDataStorage/Source/TedsUI/Private/Widgets/UrlWidget.h
+++ b/Engine/Plugins/Experimental/EditorDataStorage/Source/TedsUI/Private/Widgets/UrlWidget.h
@@ -7,9 +7,11 @@
 
 #include "UrlWidget.generated.h"
 
-class IEditorDataStorageProvider;
-class IEditorDataStorageUiProvider;
-class UScriptStruct;
+namespace UE::Editor::DataStorage
+{
+	class ICoreProvider;
+	class IUiProvider;
+} // namespace UE::Editor::DataStorage
 
 UCLASS()
 class UUrlWidgetFactory : public UEditorDataStorageFactory
@@ -19,8 +21,8 @@ class UUrlWidgetFactory : public UEditorDataStorageFactory
 public:
 	~UUrlWidgetFactory() override = default;
 
-	void RegisterWidgetConstructors(IEditorDataStorageProvider& DataStorage,
-		IEditorDataStorageUiProvider& DataStorageUi) const override;
+	void RegisterWidgetConstructors(UE::Editor::DataStorage::ICoreProvider& DataStorage,
+		UE::Editor::DataStorage::IUiProvider& DataStorageUi) const override;
 };
 
 // Widget to display a URL in TEDS
@@ -34,8 +36,8 @@ public:
 	~FUrlWidgetConstructor() override = default;
 
 	virtual TSharedPtr<SWidget> CreateWidget(
-		IEditorDataStorageProvider* DataStorage,
-		IEditorDataStorageUiProvider* DataStorageUi,
+		UE::Editor::DataStorage::ICoreProvider* DataStorage,
+		UE::Editor::DataStorage::IUiProvider* DataStorageUi,
 		UE::Editor::DataStorage::RowHandle TargetRow,
 		UE::Editor::DataStorage::RowHandle WidgetRow, 
 		const UE::Editor::DataStorage::FMetaDataView& Arguments) override;
diff --git a/Engine/Plugins/Experimental/EditorDataStorage/Source/TedsUI/Public/Widgets/CounterWidget.h b/Engine/Plugins/Experimental/EditorDataStorage/Source/TedsUI/Public/Widgets/CounterWidget.h
--- a/Engine/Plugins/Experimental/EditorDataStorage/Source/TedsUI/Public/Widgets/CounterWidget.h
+++ b/Engine/Plugins/Experimental/EditorDataStorage/Source/TedsUI/Public/Widgets/CounterWidget.h
@@ -9,7 +9,12 @@
 
 #include "CounterWidget.generated.h"
 
-class IEditorDataStorageProvider;
+namespace UE::Editor::DataStorage
+{
+	class ICoreProvider;
+	class IUiProvider;
+} // namespace UE::Editor::DataStorage
+
 class SWindow;
 class UScriptStruct;
 
@@ -24,10 +29,10 @@ public:
 	UCounterWidgetFactory();
 	~UCounterWidgetFactory() override = default;
 
-	void RegisterQueries(IEditorDataStorageProvider& DataStorage) override;
-	void RegisterWidgetPurposes(IEditorDataStorageUiProvider& DataStorageUi) const override;
-	void RegisterWidgetConstructors(IEditorDataStorageProvider& DataStorage,
-		IEditorDataStorageUiProvider& DataStorageUi) const override;
+	void RegisterQueries(UE::Editor::DataStorage::ICoreProvider& DataStorage) override;
+	void RegisterWidgetPurposes(UE::Editor::DataStorage::IUiProvider& DataStorageUi) const override;
+	void RegisterWidgetConstructors(UE::Editor::DataStorage::ICoreProvider& DataStorage,
+		UE::Editor::DataStorage::IUiProvider& DataStorageUi) const override;
 
 	static void EnableCounterWidgets();
 
@@ -61,7 +66,7 @@ public:
 
 protected:
 	TSharedPtr<SWidget> CreateWidget(const UE::Editor::DataStorage::FMetaDataView& Arguments) override;
-	bool SetColumns(IEditorDataStorageProvider* DataStorage, UE::Editor::DataStorage::RowHandle Row) override;
+	bool SetColumns(UE::Editor::DataStorage::ICoreProvider* DataStorage, UE::Editor::DataStorage::RowHandle Row) override;
 };
 
 USTRUCT(meta = (DisplayName = "Counter widget"))
diff --git a/Engine/Plugins/Experimental/EditorDataStorage/Source/TedsUI/Public/Widgets/ExportedTextWidget.h b/Engine/Plugins/Experimental/EditorDataStorage/Source/TedsUI/Public/Widgets/ExportedTextWidget.h
--- a/Engine/Plugins/Experimental/EditorDataStorage/Source/TedsUI/Public/Widgets/ExportedTextWidget.h
+++ b/Engine/Plugins/Experimental/EditorDataStorage/Source/TedsUI/Public/Widgets/ExportedTextWidget.h
@@ -8,7 +8,11 @@
 
 #include "ExportedTextWidget.generated.h"
 
-class IEditorDataStorageProvider;
+namespace UE::Editor::DataStorage
+{
+	class ICoreProvider;
+} // namespace UE::Editor::DataStorage
+
 class UScriptStruct;
 
 UCLASS()
@@ -19,8 +23,8 @@ class TEDSUI_API UExportedTextWidgetFactory : public UEditorDataStorageFactory
 public:
 	virtual ~UExportedTextWidgetFactory() override = default;
 
-	virtual void RegisterWidgetConstructors(IEditorDataStorageProvider& DataStorage,
-		IEditorDataStorageUiProvider& DataStorageUi) const override;
+	virtual void RegisterWidgetConstructors(UE::Editor::DataStorage::ICoreProvider& DataStorage,
+		UE::Editor::DataStorage::IUiProvider& DataStorageUi) const override;
 };
 
 USTRUCT()
@@ -36,11 +40,11 @@ public:
 	virtual const UE::Editor::DataStorage::Queries::FConditions* GetQueryConditions() const override;
 
 	virtual FText CreateWidgetDisplayNameText(
-		IEditorDataStorageProvider* DataStorage, UE::Editor::DataStorage::RowHandle Row) const override;
+		UE::Editor::DataStorage::ICoreProvider* DataStorage, UE::Editor::DataStorage::RowHandle Row) const override;
 	
 	virtual TSharedPtr<SWidget> CreateWidget(
-		IEditorDataStorageProvider* DataStorage,
-		IEditorDataStorageUiProvider* DataStorageUi,
+		UE::Editor::DataStorage::ICoreProvider* DataStorage,
+		UE::Editor::DataStorage::IUiProvider* DataStorageUi,
 		UE::Editor::DataStorage::RowHandle TargetRow,
 		UE::Editor::DataStorage::RowHandle WidgetRow,
 		const UE::Editor::DataStorage::FMetaDataView& Arguments) override;
diff --git a/Engine/Plugins/Experimental/EditorDataStorage/Source/TedsUI/Public/Widgets/GeneralWidgetRegistrationFactory.h b/Engine/Plugins/Experimental/EditorDataStorage/Source/TedsUI/Public/Widgets/GeneralWidgetRegistrationFactory.h
--- a/Engine/Plugins/Experimental/EditorDataStorage/Source/TedsUI/Public/Widgets/GeneralWidgetRegistrationFactory.h
+++ b/Engine/Plugins/Experimental/EditorDataStorage/Source/TedsUI/Public/Widgets/GeneralWidgetRegistrationFactory.h
@@ -21,5 +21,5 @@ public:
 
 	~UGeneralWidgetRegistrationFactory() override = default;
 
-	void RegisterWidgetPurposes(IEditorDataStorageUiProvider& DataStorageUi) const override;
+	void RegisterWidgetPurposes(UE::Editor::DataStorage::IUiProvider& DataStorageUi) const override;
 };
diff --git a/Engine/Plugins/Experimental/EditorDataStorage/Source/TedsUI/Public/Widgets/LabelWidget.h b/Engine/Plugins/Experimental/EditorDataStorage/Source/TedsUI/Public/Widgets/LabelWidget.h
--- a/Engine/Plugins/Experimental/EditorDataStorage/Source/TedsUI/Public/Widgets/LabelWidget.h
+++ b/Engine/Plugins/Experimental/EditorDataStorage/Source/TedsUI/Public/Widgets/LabelWidget.h
@@ -9,7 +9,11 @@
 
 #include "LabelWidget.generated.h"
 
-class IEditorDataStorageProvider;
+namespace UE::Editor::DataStorage
+{
+	class ICoreProvider;
+} // namespace UE::Editor::DataStorage
+
 class UScriptStruct;
 
 UCLASS()
@@ -20,9 +24,9 @@ class ULabelWidgetFactory : public UEditorDataStorageFactory
 public:
 	~ULabelWidgetFactory() override = default;
 
-	TEDSUI_API void RegisterWidgetConstructors(IEditorDataStorageProvider& DataStorage,
-		IEditorDataStorageUiProvider& DataStorageUi) const override;
-	TEDSUI_API virtual void RegisterWidgetPurposes(IEditorDataStorageUiProvider& DataStorageUi) const override;
+	TEDSUI_API void RegisterWidgetConstructors(UE::Editor::DataStorage::ICoreProvider& DataStorage,
+		UE::Editor::DataStorage::IUiProvider& DataStorageUi) const override;
+	TEDSUI_API virtual void RegisterWidgetPurposes(UE::Editor::DataStorage::IUiProvider& DataStorageUi) const override;
 };
 
 USTRUCT()
@@ -37,8 +41,8 @@ public:
 	TEDSUI_API virtual TConstArrayView<const UScriptStruct*> GetAdditionalColumnsList() const override;
 
 	TEDSUI_API virtual TSharedPtr<SWidget> CreateWidget(
-		IEditorDataStorageProvider* DataStorage,
-		IEditorDataStorageUiProvider* DataStorageUi,
+		UE::Editor::DataStorage::ICoreProvider* DataStorage,
+		UE::Editor::DataStorage::IUiProvider* DataStorageUi,
 		UE::Editor::DataStorage::RowHandle TargetRow,
 		UE::Editor::DataStorage::RowHandle WidgetRow, 
 		const UE::Editor::DataStorage::FMetaDataView& Arguments) override;
diff --git a/Engine/Plugins/Experimental/EditorDataStorage/Source/TedsUI/Public/Widgets/PackagePathWidget.h b/Engine/Plugins/Experimental/EditorDataStorage/Source/TedsUI/Public/Widgets/PackagePathWidget.h
--- a/Engine/Plugins/Experimental/EditorDataStorage/Source/TedsUI/Public/Widgets/PackagePathWidget.h
+++ b/Engine/Plugins/Experimental/EditorDataStorage/Source/TedsUI/Public/Widgets/PackagePathWidget.h
@@ -7,7 +7,11 @@
 
 #include "PackagePathWidget.generated.h"
 
-class IEditorDataStorageProvider;
+namespace UE::Editor::DataStorage
+{
+	class ICoreProvider;
+} // namespace UE::Editor::DataStorage
+
 class SWidget;
 class UScriptStruct;
 
@@ -19,8 +23,8 @@ class TEDSUI_API UPackagePathWidgetFactory : public UEditorDataStorageFactory
 public:
 	~UPackagePathWidgetFactory() override = default;
 
-	void RegisterWidgetConstructors(IEditorDataStorageProvider& DataStorage,
-		IEditorDataStorageUiProvider& DataStorageUi) const override;
+	void RegisterWidgetConstructors(UE::Editor::DataStorage::ICoreProvider& DataStorage,
+		UE::Editor::DataStorage::IUiProvider& DataStorageUi) const override;
 };
 
 USTRUCT()
@@ -36,7 +40,8 @@ protected:
 	explicit FPackagePathWidgetConstructor(const UScriptStruct* InTypeInfo);
 
 	TSharedPtr<SWidget> CreateWidget(const UE::Editor::DataStorage::FMetaDataView& Arguments) override;
-	bool FinalizeWidget(IEditorDataStorageProvider* DataStorage, IEditorDataStorageUiProvider* DataStorageUi,
+	bool FinalizeWidget(UE::Editor::DataStorage::ICoreProvider* DataStorage,
+		UE::Editor::DataStorage::IUiProvider* DataStorageUi,
 		UE::Editor::DataStorage::RowHandle Row, const TSharedPtr<SWidget>& Widget) override;
 };
 
@@ -50,6 +55,7 @@ public:
 	~FLoadedPackagePathWidgetConstructor() override = default;
 
 protected:
-	bool FinalizeWidget(IEditorDataStorageProvider* DataStorage, IEditorDataStorageUiProvider* DataStorageUi,
+	bool FinalizeWidget(UE::Editor::DataStorage::ICoreProvider* DataStorage,
+		UE::Editor::DataStorage::IUiProvider* DataStorageUi,
 		UE::Editor::DataStorage::RowHandle Row, const TSharedPtr<SWidget>& Widget) override;
 };
diff --git a/Engine/Plugins/Experimental/EditorDataStorageFeatures/Source/TedsAlerts/Private/TedsAlertWidget.cpp b/Engine/Plugins/Experimental/EditorDataStorageFeatures/Source/TedsAlerts/Private/TedsAlertWidget.cpp
--- a/Engine/Plugins/Experimental/EditorDataStorageFeatures/Source/TedsAlerts/Private/TedsAlertWidget.cpp
+++ b/Engine/Plugins/Experimental/EditorDataStorageFeatures/Source/TedsAlerts/Private/TedsAlertWidget.cpp
@@ -125,7 +125,7 @@ namespace UE::Editor::DataStorage::Widgets::Private
 					ActionButton.SetVisibility(EVisibility::Visible);
 					ActionButton.SetOnClicked(FOnClicked::CreateLambda([RowWithAlertAction]()
 						{
-							const IEditorDataStorageProvider* DataStorage = GetDataStorageFeature<IEditorDataStorageProvider>(StorageFeatureName);
+							const ICoreProvider* DataStorage = GetDataStorageFeature<ICoreProvider>(StorageFeatureName);
 							if (const FAlertActionColumn* Action =
 								DataStorage->GetColumn<FAlertActionColumn>(RowWithAlertAction))
 							{
@@ -149,8 +149,8 @@ namespace UE::Editor::DataStorage::Widgets::Private
 //
 
 void UAlertWidgetFactory::RegisterWidgetConstructors(
-	IEditorDataStorageProvider& DataStorage,
-	IEditorDataStorageUiProvider& DataStorageUi) const
+	UE::Editor::DataStorage::ICoreProvider& DataStorage,
+	UE::Editor::DataStorage::IUiProvider& DataStorageUi) const
 {
 	using namespace UE::Editor::DataStorage::Columns;
 	using namespace UE::Editor::DataStorage::Queries;
@@ -162,13 +162,13 @@ void UAlertWidgetFactory::RegisterWidgetConstructors(
 		TColumn<FAlertColumn>() || TColumn<FChildAlertColumn>());
 }
 
-void UAlertWidgetFactory::RegisterQueries(IEditorDataStorageProvider& DataStorage)
+void UAlertWidgetFactory::RegisterQueries(UE::Editor::DataStorage::ICoreProvider& DataStorage)
 {
 	RegisterAlertQueries(DataStorage);
 	RegisterAlertHeaderQueries(DataStorage);
 }
 
-void UAlertWidgetFactory::RegisterAlertQueries(IEditorDataStorageProvider& DataStorage)
+void UAlertWidgetFactory::RegisterAlertQueries(UE::Editor::DataStorage::ICoreProvider& DataStorage)
 {
 	using namespace UE::Editor::DataStorage::Columns;
 	using namespace UE::Editor::DataStorage::Queries;
@@ -250,7 +250,7 @@ void UAlertWidgetFactory::RegisterAlertQueries(IEditorDataStorageProvider& DataS
 		.Compile());
 }
 
-void UAlertWidgetFactory::RegisterAlertHeaderQueries(IEditorDataStorageProvider& DataStorage)
+void UAlertWidgetFactory::RegisterAlertHeaderQueries(UE::Editor::DataStorage::ICoreProvider& DataStorage)
 {
 	using namespace UE::Editor::DataStorage::Columns;
 	using namespace UE::Editor::DataStorage::Queries;
@@ -362,12 +362,14 @@ TSharedPtr<SWidget> FAlertWidgetConstructor::CreateWidget(
 
 TConstArrayView<const UScriptStruct*> FAlertWidgetConstructor::GetAdditionalColumnsList() const
 {
+	using namespace UE::Editor::DataStorage;
 	static TTypedElementColumnTypeList<FAlertWidgetTag> Columns;
 	return Columns;
 }
 
-bool FAlertWidgetConstructor::FinalizeWidget(IEditorDataStorageProvider* DataStorage,
-	IEditorDataStorageUiProvider* DataStorageUi, UE::Editor::DataStorage::RowHandle Row, const TSharedPtr<SWidget>& Widget)
+bool FAlertWidgetConstructor::FinalizeWidget(UE::Editor::DataStorage::ICoreProvider* DataStorage,
+	UE::Editor::DataStorage::IUiProvider* DataStorageUi,
+	UE::Editor::DataStorage::RowHandle Row, const TSharedPtr<SWidget>& Widget)
 {
 	using namespace UE::Editor::DataStorage;
 	using namespace UE::Editor::DataStorage::Columns;
@@ -419,12 +421,14 @@ TSharedPtr<SWidget> FAlertHeaderWidgetConstructor::CreateWidget(
 
 TConstArrayView<const UScriptStruct*> FAlertHeaderWidgetConstructor::GetAdditionalColumnsList() const
 {
+	using namespace UE::Editor::DataStorage;
 	static TTypedElementColumnTypeList<FAlertHeaderWidgetTag> Columns;
 	return Columns;
 }
 
-bool FAlertHeaderWidgetConstructor::FinalizeWidget(IEditorDataStorageProvider* DataStorage,
-	IEditorDataStorageUiProvider* DataStorageUi, UE::Editor::DataStorage::RowHandle Row, const TSharedPtr<SWidget>& Widget)
+bool FAlertHeaderWidgetConstructor::FinalizeWidget(UE::Editor::DataStorage::ICoreProvider* DataStorage,
+	UE::Editor::DataStorage::IUiProvider* DataStorageUi,
+	UE::Editor::DataStorage::RowHandle Row, const TSharedPtr<SWidget>& Widget)
 {
 	DataStorage->AddColumn(Row, FHeaderWidgetSizeColumn
 		{
diff --git a/Engine/Plugins/Experimental/EditorDataStorageFeatures/Source/TedsAlerts/Private/TedsAlertWidget.h b/Engine/Plugins/Experimental/EditorDataStorageFeatures/Source/TedsAlerts/Private/TedsAlertWidget.h
--- a/Engine/Plugins/Experimental/EditorDataStorageFeatures/Source/TedsAlerts/Private/TedsAlertWidget.h
+++ b/Engine/Plugins/Experimental/EditorDataStorageFeatures/Source/TedsAlerts/Private/TedsAlertWidget.h
@@ -16,13 +16,13 @@ class UAlertWidgetFactory final : public UEditorDataStorageFactory
 public:
 	~UAlertWidgetFactory() override = default;
 
-	void RegisterWidgetConstructors(IEditorDataStorageProvider& DataStorage,
-		IEditorDataStorageUiProvider& DataStorageUi) const override;
-	void RegisterQueries(IEditorDataStorageProvider& DataStorage) override;
+	void RegisterWidgetConstructors(UE::Editor::DataStorage::ICoreProvider& DataStorage,
+		UE::Editor::DataStorage::IUiProvider& DataStorageUi) const override;
+	void RegisterQueries(UE::Editor::DataStorage::ICoreProvider& DataStorage) override;
 
 private:
-	void RegisterAlertQueries(IEditorDataStorageProvider& DataStorage);
-	void RegisterAlertHeaderQueries(IEditorDataStorageProvider& DataStorage);
+	void RegisterAlertQueries(UE::Editor::DataStorage::ICoreProvider& DataStorage);
+	void RegisterAlertHeaderQueries(UE::Editor::DataStorage::ICoreProvider& DataStorage);
 };
 
 USTRUCT()
@@ -46,7 +46,8 @@ public:
 private:
 	TSharedPtr<SWidget> CreateWidget(const UE::Editor::DataStorage::FMetaDataView& Arguments) override;
 	TConstArrayView<const UScriptStruct*> GetAdditionalColumnsList() const override;
-	bool FinalizeWidget(IEditorDataStorageProvider* DataStorage, IEditorDataStorageUiProvider* DataStorageUi,
+	bool FinalizeWidget(UE::Editor::DataStorage::ICoreProvider* DataStorage,
+		UE::Editor::DataStorage::IUiProvider* DataStorageUi,
 		UE::Editor::DataStorage::RowHandle Row, const TSharedPtr<SWidget>& Widget) override;
 };
 
@@ -68,7 +69,8 @@ public:
 protected:
 	TSharedPtr<SWidget> CreateWidget(const UE::Editor::DataStorage::FMetaDataView& Arguments) override;
 	TConstArrayView<const UScriptStruct*> GetAdditionalColumnsList() const override;
-	bool FinalizeWidget(IEditorDataStorageProvider* DataStorage, IEditorDataStorageUiProvider* DataStorageUi,
+	bool FinalizeWidget(UE::Editor::DataStorage::ICoreProvider* DataStorage,
+		UE::Editor::DataStorage::IUiProvider* DataStorageUi,
 		UE::Editor::DataStorage::RowHandle Row, const TSharedPtr<SWidget>& Widget) override;
 };
 
diff --git a/Engine/Plugins/Experimental/EditorDataStorageFeatures/Source/TedsAlerts/Private/TedsAlerts.cpp b/Engine/Plugins/Experimental/EditorDataStorageFeatures/Source/TedsAlerts/Private/TedsAlerts.cpp
--- a/Engine/Plugins/Experimental/EditorDataStorageFeatures/Source/TedsAlerts/Private/TedsAlerts.cpp
+++ b/Engine/Plugins/Experimental/EditorDataStorageFeatures/Source/TedsAlerts/Private/TedsAlerts.cpp
@@ -10,7 +10,7 @@
 
 namespace UE::Editor::DataStorage::Alerts
 {
-	void AddAlert(IEditorDataStorageProvider& DataStorage, RowHandle TargetRow, const FName& Name, FText Message, 
+	void AddAlert(ICoreProvider& DataStorage, RowHandle TargetRow, const FName& Name, FText Message, 
 		Columns::FAlertColumnType Type, uint8 Priority, Columns::FAlertActionCallback Action)
 	{
 		using namespace UE::Editor::DataStorage::Columns;
@@ -54,14 +54,14 @@ namespace UE::Editor::DataStorage::Alerts
 		
 		Context.PushCommand([TargetRow, Name, Message = MoveTemp(Message), Type, Priority, Action = MoveTemp(Action)]() mutable
 			{
-				if (IEditorDataStorageProvider* DataStorage = GetMutableDataStorageFeature<IEditorDataStorageProvider>(StorageFeatureName))
+				if (ICoreProvider* DataStorage = GetMutableDataStorageFeature<ICoreProvider>(StorageFeatureName))
 				{
 					AddAlert(*DataStorage, TargetRow, Name, MoveTemp(Message), Type, Priority, MoveTemp(Action));
 				}
 			});
 	}
 
-	void UpdateAlertText(IEditorDataStorageProvider& DataStorage, RowHandle TargetRow, const FName& Name, FText Message)
+	void UpdateAlertText(ICoreProvider& DataStorage, RowHandle TargetRow, const FName& Name, FText Message)
 	{
 		using namespace UE::Editor::DataStorage::Columns;
 
@@ -83,14 +83,14 @@ namespace UE::Editor::DataStorage::Alerts
 	{
 		Context.PushCommand([TargetRow, Name, Message = MoveTemp(Message)]() mutable
 			{
-				if (IEditorDataStorageProvider* DataStorage = GetMutableDataStorageFeature<IEditorDataStorageProvider>(StorageFeatureName))
+				if (ICoreProvider* DataStorage = GetMutableDataStorageFeature<ICoreProvider>(StorageFeatureName))
 				{
 					UpdateAlertText(*DataStorage, TargetRow, Name, MoveTemp(Message));
 				}
 			});
 	}
 
-	void UpdateAlertAction(IEditorDataStorageProvider& DataStorage, RowHandle TargetRow, const FName& Name, 
+	void UpdateAlertAction(ICoreProvider& DataStorage, RowHandle TargetRow, const FName& Name, 
 		Columns::FAlertActionCallback Action)
 	{
 		using namespace UE::Editor::DataStorage::Columns;
@@ -127,14 +127,14 @@ namespace UE::Editor::DataStorage::Alerts
 	{
 		Context.PushCommand([TargetRow, Name, Action = MoveTemp(Action)]() mutable
 			{
-				if (IEditorDataStorageProvider* DataStorage = GetMutableDataStorageFeature<IEditorDataStorageProvider>(StorageFeatureName))
+				if (ICoreProvider* DataStorage = GetMutableDataStorageFeature<ICoreProvider>(StorageFeatureName))
 				{
 					UpdateAlertAction(*DataStorage, TargetRow, Name, MoveTemp(Action));
 				}
 			});
 	}
 
-	void RemoveAlert(IEditorDataStorageProvider& DataStorage, RowHandle TargetRow, const FName& Name)
+	void RemoveAlert(ICoreProvider& DataStorage, RowHandle TargetRow, const FName& Name)
 	{
 		using namespace UE::Editor::DataStorage::Columns;
 		using namespace UE::Editor::DataStorage::Queries;
@@ -226,7 +226,7 @@ namespace UE::Editor::DataStorage::Alerts
 
 		Context.PushCommand([TargetRow, Name]() mutable
 			{
-				if (IEditorDataStorageProvider* DataStorage = GetMutableDataStorageFeature<IEditorDataStorageProvider>(StorageFeatureName))
+				if (ICoreProvider* DataStorage = GetMutableDataStorageFeature<ICoreProvider>(StorageFeatureName))
 				{
 					RemoveAlert(*DataStorage, TargetRow, Name);
 				}
diff --git a/Engine/Plugins/Experimental/EditorDataStorageFeatures/Source/TedsAlerts/Private/TedsAlertsCommands.cpp b/Engine/Plugins/Experimental/EditorDataStorageFeatures/Source/TedsAlerts/Private/TedsAlertsCommands.cpp
--- a/Engine/Plugins/Experimental/EditorDataStorageFeatures/Source/TedsAlerts/Private/TedsAlertsCommands.cpp
+++ b/Engine/Plugins/Experimental/EditorDataStorageFeatures/Source/TedsAlerts/Private/TedsAlertsCommands.cpp
@@ -21,7 +21,7 @@ FAutoConsoleCommand AddRandomAlertToRowConsoleCommand(
 
 			TRACE_CPUPROFILER_EVENT_SCOPE(TEDS.Debug.AddRandomAlertToSelectedRows);
 
-			if (IEditorDataStorageProvider* DataStorage = GetMutableDataStorageFeature<IEditorDataStorageProvider>(StorageFeatureName))
+			if (ICoreProvider* DataStorage = GetMutableDataStorageFeature<ICoreProvider>(StorageFeatureName))
 			{
 				static uint32 Count = 0;
 				static FName TestMessageName = FName("Test message");
@@ -75,7 +75,7 @@ FAutoConsoleCommand RemoveSelectedAlertsConsoleCommand(
 			TRACE_CPUPROFILER_EVENT_SCOPE(TEDS.Debug.ClearSelectedAlerts);
 			static FName TestMessageName = FName("Test message");
 
-			if (IEditorDataStorageProvider* DataStorage = GetMutableDataStorageFeature<IEditorDataStorageProvider>(StorageFeatureName))
+			if (ICoreProvider* DataStorage = GetMutableDataStorageFeature<ICoreProvider>(StorageFeatureName))
 			{
 				static QueryHandle AlertQuery = [DataStorage]
 					{
@@ -110,7 +110,7 @@ FAutoConsoleCommand ClearAllAlertsConsoleCommand(
 
 			TRACE_CPUPROFILER_EVENT_SCOPE(TEDS.Debug.ClearAllAlertInfo);
 
-			if (IEditorDataStorageProvider* DataStorage = GetMutableDataStorageFeature<IEditorDataStorageProvider>(StorageFeatureName))
+			if (ICoreProvider* DataStorage = GetMutableDataStorageFeature<ICoreProvider>(StorageFeatureName))
 			{
 				// Remove all alerts on existing rows.
 				{
diff --git a/Engine/Plugins/Experimental/EditorDataStorageFeatures/Source/TedsAlerts/Private/TedsAlertsFactory.cpp b/Engine/Plugins/Experimental/EditorDataStorageFeatures/Source/TedsAlerts/Private/TedsAlertsFactory.cpp
--- a/Engine/Plugins/Experimental/EditorDataStorageFeatures/Source/TedsAlerts/Private/TedsAlertsFactory.cpp
+++ b/Engine/Plugins/Experimental/EditorDataStorageFeatures/Source/TedsAlerts/Private/TedsAlertsFactory.cpp
@@ -13,13 +13,13 @@ const FName UTedsAlertsFactory::AlertChainTableName = "Alerts chain";
 const FName UTedsAlertsFactory::UnsortedAlertChainTableName = "Alerts chain (unsorted)";
 const FName UTedsAlertsFactory::AlertConditionName = FName(TEXT("Alerts"));
 
-void UTedsAlertsFactory::RegisterTables(IEditorDataStorageProvider& DataStorage)
+void UTedsAlertsFactory::RegisterTables(UE::Editor::DataStorage::ICoreProvider& DataStorage)
 {
 	ChainTable = DataStorage.RegisterTable<FTedsAlertColumn, FTedsAlertChainTag>(AlertChainTableName);
 	UnsortedChainTable = DataStorage.RegisterTable<FTedsUnsortedAlertChainTag>(ChainTable, UnsortedAlertChainTableName);
 }
 
-void UTedsAlertsFactory::RegisterQueries(IEditorDataStorageProvider& DataStorage)
+void UTedsAlertsFactory::RegisterQueries(UE::Editor::DataStorage::ICoreProvider& DataStorage)
 {
 	RegisterSubQueries(DataStorage);
 	RegisterSortUnsortedAlertsQuery(DataStorage);
@@ -49,7 +49,7 @@ UE::Editor::DataStorage::QueryHandle UTedsAlertsFactory::GetUnsortedAlertsQuery(
 	return UnsortedAlertsQuery;
 }
 
-void UTedsAlertsFactory::RegisterSubQueries(IEditorDataStorageProvider& DataStorage)
+void UTedsAlertsFactory::RegisterSubQueries(UE::Editor::DataStorage::ICoreProvider& DataStorage)
 {
 	using namespace UE::Editor::DataStorage::Queries;
 
@@ -83,7 +83,7 @@ void UTedsAlertsFactory::RegisterSubQueries(IEditorDataStorageProvider& DataStor
 		.Compile());
 }
 
-void UTedsAlertsFactory::RegisterSortUnsortedAlertsQuery(IEditorDataStorageProvider& DataStorage)
+void UTedsAlertsFactory::RegisterSortUnsortedAlertsQuery(UE::Editor::DataStorage::ICoreProvider& DataStorage)
 {
 	using namespace UE::Editor::DataStorage::Queries;
 
@@ -151,7 +151,7 @@ void UTedsAlertsFactory::RegisterSortUnsortedAlertsQuery(IEditorDataStorageProvi
 		.Compile());
 }
 
-void UTedsAlertsFactory::RegisterOnAddQueries(IEditorDataStorageProvider& DataStorage)
+void UTedsAlertsFactory::RegisterOnAddQueries(UE::Editor::DataStorage::ICoreProvider& DataStorage)
 {
 	using namespace UE::Editor::DataStorage::Queries;
 
@@ -178,7 +178,7 @@ void UTedsAlertsFactory::RegisterOnAddQueries(IEditorDataStorageProvider& DataSt
 		.Compile());
 }
 
-void UTedsAlertsFactory::RegisterOnRemoveQueries(IEditorDataStorageProvider& DataStorage)
+void UTedsAlertsFactory::RegisterOnRemoveQueries(UE::Editor::DataStorage::ICoreProvider& DataStorage)
 {
 	using namespace UE::Editor::DataStorage::Queries;
 	
@@ -228,7 +228,7 @@ void UTedsAlertsFactory::RegisterOnRemoveQueries(IEditorDataStorageProvider& Dat
 		.Compile());
 }
 
-void UTedsAlertsFactory::RegisterParentUpdatesQueries(IEditorDataStorageProvider& DataStorage)
+void UTedsAlertsFactory::RegisterParentUpdatesQueries(UE::Editor::DataStorage::ICoreProvider& DataStorage)
 {
 	using namespace UE::Editor::DataStorage::Queries;
 
@@ -267,7 +267,7 @@ void UTedsAlertsFactory::RegisterParentUpdatesQueries(IEditorDataStorageProvider
 		.Compile());
 }
 
-void UTedsAlertsFactory::RegisterChildAlertUpdatesQueries(IEditorDataStorageProvider& DataStorage)
+void UTedsAlertsFactory::RegisterChildAlertUpdatesQueries(UE::Editor::DataStorage::ICoreProvider& DataStorage)
 {
 	using namespace UE::Editor::DataStorage::Queries;
 
diff --git a/Engine/Plugins/Experimental/EditorDataStorageFeatures/Source/TedsAlerts/Private/TedsAlertsFactory.h b/Engine/Plugins/Experimental/EditorDataStorageFeatures/Source/TedsAlerts/Private/TedsAlertsFactory.h
--- a/Engine/Plugins/Experimental/EditorDataStorageFeatures/Source/TedsAlerts/Private/TedsAlertsFactory.h
+++ b/Engine/Plugins/Experimental/EditorDataStorageFeatures/Source/TedsAlerts/Private/TedsAlertsFactory.h
@@ -23,8 +23,8 @@ public:
 
 	~UTedsAlertsFactory() override = default;
 	
-	void RegisterTables(IEditorDataStorageProvider& DataStorage) override;
-	void RegisterQueries(IEditorDataStorageProvider& DataStorage) override;
+	void RegisterTables(UE::Editor::DataStorage::ICoreProvider& DataStorage) override;
+	void RegisterQueries(UE::Editor::DataStorage::ICoreProvider& DataStorage) override;
 
 	UE::Editor::DataStorage::TableHandle GetAlertChainTable() const;
 	UE::Editor::DataStorage::TableHandle GetUnsortedAlertChainTable() const;
@@ -44,12 +44,12 @@ private:
 	UE::Editor::DataStorage::QueryHandle ChildAlertColumnReadWriteQuery = UE::Editor::DataStorage::InvalidTableHandle;
 	UE::Editor::DataStorage::QueryHandle ParentReadOnlyQuery = UE::Editor::DataStorage::InvalidTableHandle;
 
-	void RegisterSubQueries(IEditorDataStorageProvider& DataStorage);
-	void RegisterSortUnsortedAlertsQuery(IEditorDataStorageProvider& DataStorage);
-	void RegisterOnAddQueries(IEditorDataStorageProvider& DataStorage);
-	void RegisterOnRemoveQueries(IEditorDataStorageProvider& DataStorage);
-	void RegisterParentUpdatesQueries(IEditorDataStorageProvider& DataStorage);
-	void RegisterChildAlertUpdatesQueries(IEditorDataStorageProvider& DataStorage);
+	void RegisterSubQueries(UE::Editor::DataStorage::ICoreProvider& DataStorage);
+	void RegisterSortUnsortedAlertsQuery(UE::Editor::DataStorage::ICoreProvider& DataStorage);
+	void RegisterOnAddQueries(UE::Editor::DataStorage::ICoreProvider& DataStorage);
+	void RegisterOnRemoveQueries(UE::Editor::DataStorage::ICoreProvider& DataStorage);
+	void RegisterParentUpdatesQueries(UE::Editor::DataStorage::ICoreProvider& DataStorage);
+	void RegisterChildAlertUpdatesQueries(UE::Editor::DataStorage::ICoreProvider& DataStorage);
 
 	static void AssignAlert(
 		UE::Editor::DataStorage::ISubqueryContext& TargetContext, UE::Editor::DataStorage::RowHandle TargetRow, FTedsAlertColumn& Target,
diff --git a/Engine/Plugins/Experimental/EditorDataStorageFeatures/Source/TedsAlerts/Public/TedsAlerts.h b/Engine/Plugins/Experimental/EditorDataStorageFeatures/Source/TedsAlerts/Public/TedsAlerts.h
--- a/Engine/Plugins/Experimental/EditorDataStorageFeatures/Source/TedsAlerts/Public/TedsAlerts.h
+++ b/Engine/Plugins/Experimental/EditorDataStorageFeatures/Source/TedsAlerts/Public/TedsAlerts.h
@@ -7,11 +7,10 @@
 #include "TedsAlertColumns.h"
 #include "UObject/NameTypes.h"
 
-class IEditorDataStorageProvider;
-
 namespace UE::Editor::DataStorage
 {
 	struct IQueryContext;
+	class ICoreProvider;
 }
 
 namespace UE::Editor::DataStorage::Alerts
@@ -30,7 +29,7 @@ namespace UE::Editor::DataStorage::Alerts
 	 *		one is shown when the active alert is removed. Higher values indicate higher priority.
 	 * @Action An optional callback that gets triggered when the user presses the alert icon.
 	 */
-	TEDSALERTS_API void AddAlert(IEditorDataStorageProvider& DataStorage, RowHandle TargetRow, const FName& Name, FText Message, 
+	TEDSALERTS_API void AddAlert(ICoreProvider& DataStorage, RowHandle TargetRow, const FName& Name, FText Message, 
 		Columns::FAlertColumnType Type, uint8 Priority = 127, Columns::FAlertActionCallback Action = {});
 	/**
 	 * Add an alert to the provided row. Various tools like the Outliner monitor for alerts and can show them for instance as
@@ -56,7 +55,7 @@ namespace UE::Editor::DataStorage::Alerts
 	 * @Name The name that uniquely identifies the alert.
 	 * @Message The new message for the alert.
 	 */
-	TEDSALERTS_API void UpdateAlertText(IEditorDataStorageProvider& DataStorage, RowHandle TargetRow, const FName& Name, FText Message);
+	TEDSALERTS_API void UpdateAlertText(ICoreProvider& DataStorage, RowHandle TargetRow, const FName& Name, FText Message);
 	/**
 	 * Locates the first alerts with the provided name in the alert chain found at the provided row and updates its text.
 	 * @Context A reference to the query callback context.
@@ -75,7 +74,7 @@ namespace UE::Editor::DataStorage::Alerts
 	 * @Action The new action to trigger.
 	 */
 	TEDSALERTS_API void UpdateAlertAction(
-		IEditorDataStorageProvider& DataStorage, RowHandle TargetRow, const FName& Name, Columns::FAlertActionCallback Action);
+		ICoreProvider& DataStorage, RowHandle TargetRow, const FName& Name, Columns::FAlertActionCallback Action);
 	/**
 	 * Locates the first alerts with the provided name in the alert chain found at the provided row and updates its action.
 	 * If there's no action it will be added. If the provided action isn't bound, the alert action will be removed.
@@ -94,7 +93,7 @@ namespace UE::Editor::DataStorage::Alerts
 	 * @TargetRow The row the alert is stored on.
 	 * @Name The name that uniquely identifies the alert.
 	 */
-	TEDSALERTS_API void RemoveAlert(IEditorDataStorageProvider& DataStorage, RowHandle TargetRow, const FName& Name);
+	TEDSALERTS_API void RemoveAlert(ICoreProvider& DataStorage, RowHandle TargetRow, const FName& Name);
 	/**
 	 * Searches the alert chain found at the target row for the first alert with the provided name and removes it. The chain will
 	 * automatically be patched, keeping the chain in order of priority. If the root alert is removed, the UI will be automatically
diff --git a/Engine/Plugins/Experimental/EditorDataStorageFeatures/Source/TedsAssetData/Private/CB/TedsAssetDataCBDataSource.cpp b/Engine/Plugins/Experimental/EditorDataStorageFeatures/Source/TedsAssetData/Private/CB/TedsAssetDataCBDataSource.cpp
--- a/Engine/Plugins/Experimental/EditorDataStorageFeatures/Source/TedsAssetData/Private/CB/TedsAssetDataCBDataSource.cpp
+++ b/Engine/Plugins/Experimental/EditorDataStorageFeatures/Source/TedsAssetData/Private/CB/TedsAssetDataCBDataSource.cpp
@@ -47,7 +47,7 @@ TAutoConsoleVariable<bool> CVarTEDSAssetDataCBSourceIncludeTagsAndValues(TEXT("T
 		}
 	}));
 
-FTedsAssetDataCBDataSource::FTedsAssetDataCBDataSource(IEditorDataStorageProvider& InDatabase)
+FTedsAssetDataCBDataSource::FTedsAssetDataCBDataSource(UE::Editor::DataStorage::ICoreProvider& InDatabase)
 	: Database(InDatabase)
 {
 	using namespace UE::Editor::DataStorage;
diff --git a/Engine/Plugins/Experimental/EditorDataStorageFeatures/Source/TedsAssetData/Private/CB/TedsAssetDataCBDataSource.h b/Engine/Plugins/Experimental/EditorDataStorageFeatures/Source/TedsAssetData/Private/CB/TedsAssetDataCBDataSource.h
--- a/Engine/Plugins/Experimental/EditorDataStorageFeatures/Source/TedsAssetData/Private/CB/TedsAssetDataCBDataSource.h
+++ b/Engine/Plugins/Experimental/EditorDataStorageFeatures/Source/TedsAssetData/Private/CB/TedsAssetDataCBDataSource.h
@@ -13,7 +13,6 @@ class FAssetPackageData;
 class FAssetPropertyTagCache;
 class IAssetRegistry;
 class IPlugin;
-class IEditorDataStorageProvider;
 
 struct FAssetData;
 struct FAssetDataColumn_Experimental;
@@ -22,6 +21,7 @@ struct FAssetPathColumn_Experimental;
 namespace UE::Editor::DataStorage
 {
 	struct IQueryContext;
+	class ICoreProvider;
 }
 
 namespace UE::Editor::AssetData::Private
@@ -30,7 +30,7 @@ namespace UE::Editor::AssetData::Private
 class FTedsAssetDataCBDataSource
 {
 public:
-	explicit FTedsAssetDataCBDataSource(IEditorDataStorageProvider& InDatabase);
+	explicit FTedsAssetDataCBDataSource(UE::Editor::DataStorage::ICoreProvider& InDatabase);
 	~FTedsAssetDataCBDataSource();
 	
 	// Toggle the storage of the metadata from the asset registry into teds 
@@ -75,7 +75,7 @@ private:
 	void ProcessAssetDataAndPathUpdateQueryCallback(DataStorage::IQueryContext& Context, const DataStorage::RowHandle* Rows, const FAssetDataColumn_Experimental* AssetDataColumn) const;
 	void ProcessAssetDataUpdateQueryCallback(DataStorage::IQueryContext& Context, const DataStorage::RowHandle*Rows, const FAssetDataColumn_Experimental* AssetDataColumn) const;
 
-	IEditorDataStorageProvider& Database;
+	UE::Editor::DataStorage::ICoreProvider& Database;
 	DataStorage::QueryHandle ProcessPathQuery;
 	DataStorage::QueryHandle ProcessAssetDataPathUpdateQuery;
 	DataStorage::QueryHandle ProcessAssetDataAndPathUpdateQuery;
diff --git a/Engine/Plugins/Experimental/EditorDataStorageFeatures/Source/TedsAssetData/Private/Debug/TedsAssetDataDebugCommands.cpp b/Engine/Plugins/Experimental/EditorDataStorageFeatures/Source/TedsAssetData/Private/Debug/TedsAssetDataDebugCommands.cpp
--- a/Engine/Plugins/Experimental/EditorDataStorageFeatures/Source/TedsAssetData/Private/Debug/TedsAssetDataDebugCommands.cpp
+++ b/Engine/Plugins/Experimental/EditorDataStorageFeatures/Source/TedsAssetData/Private/Debug/TedsAssetDataDebugCommands.cpp
@@ -29,7 +29,7 @@ namespace UE::Editor::DataStorage::Debug::Private
 		TEXT("Print some debug information on the specified path."),
 		FConsoleCommandWithArgsDelegate::CreateLambda([](const TArray<FString>& InArgs)
 			{
-				const IEditorDataStorageProvider* Database = GetMutableDataStorageFeature<IEditorDataStorageProvider>(StorageFeatureName);
+				const ICoreProvider* Database = GetMutableDataStorageFeature<ICoreProvider>(StorageFeatureName);
 				IAssetRegistry& AssetRegistry = IAssetRegistry::GetChecked();
 
 				for (const FString& Path : InArgs)
@@ -159,7 +159,7 @@ namespace UE::Editor::DataStorage::Debug::Private
 		ULevel::StaticClass(), UBlueprint::StaticClass(), UTexture::StaticClass() };
 
 	// Populate an asset row with random information
-	void PopulateRowWithRandomInfo(RowHandle Row, IEditorDataStorageProvider* DataStorage)
+	void PopulateRowWithRandomInfo(RowHandle Row, ICoreProvider* DataStorage)
 	{
 		// Don't modify any rows that aren't our placeholder assets
 		if (!DataStorage->HasColumns<FAssetTag>(Row))
@@ -230,7 +230,7 @@ namespace UE::Editor::DataStorage::Debug::Private
 				TRACE_CPUPROFILER_EVENT_SCOPE(TEDS.Debug.CreateDebugAssetRows);
 
 
-				IEditorDataStorageProvider* DataStorage = GetMutableDataStorageFeature<IEditorDataStorageProvider>(StorageFeatureName);
+				ICoreProvider* DataStorage = GetMutableDataStorageFeature<ICoreProvider>(StorageFeatureName);
 
 				if (!DataStorage)
 				{
@@ -262,7 +262,7 @@ namespace UE::Editor::DataStorage::Debug::Private
 			{
 				using namespace UE::Editor::DataStorage::Queries;
 
-				IEditorDataStorageProvider* DataStorage = GetMutableDataStorageFeature<IEditorDataStorageProvider>(StorageFeatureName);
+				ICoreProvider* DataStorage = GetMutableDataStorageFeature<ICoreProvider>(StorageFeatureName);
 
 				if (!DataStorage)
 				{
@@ -275,7 +275,7 @@ namespace UE::Editor::DataStorage::Debug::Private
 				TArray<RowHandle> Rows;
 				
 				DataStorage->RunQuery(AssetQueryHandle,
-					CreateDirectQueryCallbackBinding([&Rows](const IEditorDataStorageProvider::IDirectQueryContext& Context, const RowHandle*)
+					CreateDirectQueryCallbackBinding([&Rows](const IDirectQueryContext& Context, const RowHandle*)
 					{
 						Rows.Append(Context.GetRowHandles());
 					}));
diff --git a/Engine/Plugins/Experimental/EditorDataStorageFeatures/Source/TedsAssetData/Private/Processors/AssetProcessors.cpp b/Engine/Plugins/Experimental/EditorDataStorageFeatures/Source/TedsAssetData/Private/Processors/AssetProcessors.cpp
--- a/Engine/Plugins/Experimental/EditorDataStorageFeatures/Source/TedsAssetData/Private/Processors/AssetProcessors.cpp
+++ b/Engine/Plugins/Experimental/EditorDataStorageFeatures/Source/TedsAssetData/Private/Processors/AssetProcessors.cpp
@@ -83,7 +83,7 @@ namespace UE::TedsAssetDataFactory::Private
 	}
 }
 
-void UTedsAssetDataFactory::RegisterQueries(IEditorDataStorageProvider& DataStorage)
+void UTedsAssetDataFactory::RegisterQueries(UE::Editor::DataStorage::ICoreProvider& DataStorage)
 {
 	using namespace UE::Editor::DataStorage;
 	using namespace UE::Editor::DataStorage::Columns;
@@ -189,7 +189,7 @@ void UTedsAssetDataFactory::RegisterQueries(IEditorDataStorageProvider& DataStor
 		);
 }
 
-void UTedsAssetDataFactory::PreRegister(IEditorDataStorageProvider& DataStorage)
+void UTedsAssetDataFactory::PreRegister(UE::Editor::DataStorage::ICoreProvider& DataStorage)
 {
 	if (FContentBrowserModule* ContentBrowserModule = FModuleManager::Get().GetModulePtr<FContentBrowserModule>("ContentBrowser"))
 	{
@@ -197,7 +197,7 @@ void UTedsAssetDataFactory::PreRegister(IEditorDataStorageProvider& DataStorage)
 	}
 }
 
-void UTedsAssetDataFactory::PreShutdown(IEditorDataStorageProvider& DataStorage)
+void UTedsAssetDataFactory::PreShutdown(UE::Editor::DataStorage::ICoreProvider& DataStorage)
 {
 	if (FContentBrowserModule* ContentBrowserModule = FModuleManager::Get().GetModulePtr<FContentBrowserModule>("ContentBrowser"))
 	{
@@ -205,7 +205,7 @@ void UTedsAssetDataFactory::PreShutdown(IEditorDataStorageProvider& DataStorage)
 	}
 }
 
-void UTedsAssetDataFactory::OnSetFolderColor(const FString& Path, IEditorDataStorageProvider* DataStorage)
+void UTedsAssetDataFactory::OnSetFolderColor(const FString& Path, UE::Editor::DataStorage::ICoreProvider* DataStorage)
 {
 	const UE::Editor::DataStorage::FMapKey PathKey = UE::Editor::DataStorage::FMapKey(FName(Path));
 	const UE::Editor::DataStorage::RowHandle Row = DataStorage->LookupMappedRow(PathKey);
diff --git a/Engine/Plugins/Experimental/EditorDataStorageFeatures/Source/TedsAssetData/Private/Processors/AssetProcessors.h b/Engine/Plugins/Experimental/EditorDataStorageFeatures/Source/TedsAssetData/Private/Processors/AssetProcessors.h
--- a/Engine/Plugins/Experimental/EditorDataStorageFeatures/Source/TedsAssetData/Private/Processors/AssetProcessors.h
+++ b/Engine/Plugins/Experimental/EditorDataStorageFeatures/Source/TedsAssetData/Private/Processors/AssetProcessors.h
@@ -6,7 +6,10 @@
 
 #include "AssetProcessors.generated.h"
 
-class IEditorDataStorageProvider;
+namespace UE::Editor::DataStorage
+{
+	class ICoreProvider;
+} // namespace UE::Editor::DataStorage
 
 UCLASS()
 class UTedsAssetDataFactory : public UEditorDataStorageFactory
@@ -16,11 +19,11 @@ class UTedsAssetDataFactory : public UEditorDataStorageFactory
 public:
 	~UTedsAssetDataFactory() override = default;
 
-	void RegisterQueries(IEditorDataStorageProvider& DataStorage) override;
-	virtual void PreRegister(IEditorDataStorageProvider& DataStorage) override;
-	virtual void PreShutdown(IEditorDataStorageProvider& DataStorage) override;
+	void RegisterQueries(UE::Editor::DataStorage::ICoreProvider& DataStorage) override;
+	virtual void PreRegister(UE::Editor::DataStorage::ICoreProvider& DataStorage) override;
+	virtual void PreShutdown(UE::Editor::DataStorage::ICoreProvider& DataStorage) override;
 
 protected:
 
-	void OnSetFolderColor(const FString& Path, IEditorDataStorageProvider* DataStorage);
+	void OnSetFolderColor(const FString& Path, UE::Editor::DataStorage::ICoreProvider* DataStorage);
 };
diff --git a/Engine/Plugins/Experimental/EditorDataStorageFeatures/Source/TedsAssetData/Private/TedsAssetData.cpp b/Engine/Plugins/Experimental/EditorDataStorageFeatures/Source/TedsAssetData/Private/TedsAssetData.cpp
--- a/Engine/Plugins/Experimental/EditorDataStorageFeatures/Source/TedsAssetData/Private/TedsAssetData.cpp
+++ b/Engine/Plugins/Experimental/EditorDataStorageFeatures/Source/TedsAssetData/Private/TedsAssetData.cpp
@@ -37,7 +37,7 @@ struct FPopulateAssetDataRowArgs
 
 // Only safe if the GT is blocked during the operation
 template<typename TAssetData>
-FPopulateAssetDataRowArgs ThreadSafe_PopulateAssetDataTableRow(TAssetData&& InAssetData, const IEditorDataStorageProvider& Database)
+FPopulateAssetDataRowArgs ThreadSafe_PopulateAssetDataTableRow(TAssetData&& InAssetData, const ICoreProvider& Database)
 {
 	FPopulateAssetDataRowArgs Output;
 	Output.ObjectPathKey = FMapKey(InAssetData.GetSoftObjectPath());
@@ -59,7 +59,7 @@ FPopulateAssetDataRowArgs ThreadSafe_PopulateAssetDataTableRow(TAssetData&& InAs
 }
 
 
-void PopulateAssetDataTableRow(FPopulateAssetDataRowArgs&& InAssetDataRowArgs, IEditorDataStorageProvider& Database, RowHandle RowHandle)
+void PopulateAssetDataTableRow(FPopulateAssetDataRowArgs&& InAssetDataRowArgs, ICoreProvider& Database, RowHandle RowHandle)
 {
 	if (InAssetDataRowArgs.PathRow != InvalidRowHandle)
 	{
@@ -142,7 +142,7 @@ FPopulatePathRowArgs ThreadSafe_PopulatePathRowArgs(FMapKey&& AssetRegistryPathK
 	return Args;
 }
 
-void PopulatePathDataTableRow(FPopulatePathRowArgs&& InPopulatePathRowArgs, IEditorDataStorageProvider& Database, RowHandle InRowHandle)
+void PopulatePathDataTableRow(FPopulatePathRowArgs&& InPopulatePathRowArgs, ICoreProvider& Database, RowHandle InRowHandle)
 {
 	if (InPopulatePathRowArgs.ParentAssetRegistryPathKey.IsSet())
 	{
@@ -181,7 +181,7 @@ void PopulatePathDataTableRow(FPopulatePathRowArgs&& InPopulatePathRowArgs, IEdi
 }
 
 
-FTedsAssetData::FTedsAssetData(IEditorDataStorageProvider& InDatabase)
+FTedsAssetData::FTedsAssetData(UE::Editor::DataStorage::ICoreProvider& InDatabase)
 	: Database(InDatabase)
 {
 	using namespace UE::Editor::DataStorage::Queries;
diff --git a/Engine/Plugins/Experimental/EditorDataStorageFeatures/Source/TedsAssetData/Private/TedsAssetData.h b/Engine/Plugins/Experimental/EditorDataStorageFeatures/Source/TedsAssetData/Private/TedsAssetData.h
--- a/Engine/Plugins/Experimental/EditorDataStorageFeatures/Source/TedsAssetData/Private/TedsAssetData.h
+++ b/Engine/Plugins/Experimental/EditorDataStorageFeatures/Source/TedsAssetData/Private/TedsAssetData.h
@@ -8,7 +8,10 @@
 
 struct FAssetData;
 
-class IEditorDataStorageProvider;
+namespace UE::Editor::DataStorage
+{
+	class ICoreProvider;
+} // UE::Editor::DataStorage
 
 namespace UE::Editor::AssetData::Private
 {
@@ -18,7 +21,7 @@ namespace UE::Editor::AssetData::Private
 	class FTedsAssetData
 	{
 	public:
-		FTedsAssetData(IEditorDataStorageProvider& InDatabase);
+		FTedsAssetData(UE::Editor::DataStorage::ICoreProvider& InDatabase);
 
 		~FTedsAssetData();
 
@@ -35,7 +38,7 @@ namespace UE::Editor::AssetData::Private
 		void OnPathsAdded(TConstArrayView<FStringView> InPathsAdded);
 		void OnPathsRemoved(TConstArrayView<FStringView> InPathsRemoved);
 
-		IEditorDataStorageProvider& Database;
+		UE::Editor::DataStorage::ICoreProvider& Database;
 		DataStorage::TableHandle PathsTable = DataStorage::InvalidTableHandle;
 		DataStorage::TableHandle AssetsDataTable = DataStorage::InvalidTableHandle;
 
diff --git a/Engine/Plugins/Experimental/EditorDataStorageFeatures/Source/TedsAssetData/Private/TedsAssetDataModule.cpp b/Engine/Plugins/Experimental/EditorDataStorageFeatures/Source/TedsAssetData/Private/TedsAssetDataModule.cpp
--- a/Engine/Plugins/Experimental/EditorDataStorageFeatures/Source/TedsAssetData/Private/TedsAssetDataModule.cpp
+++ b/Engine/Plugins/Experimental/EditorDataStorageFeatures/Source/TedsAssetData/Private/TedsAssetDataModule.cpp
@@ -107,7 +107,7 @@ void FTedsAssetDataModule::ProcessDependentEvents()
 void FTedsAssetDataModule::InitAssetRegistryStorage()
 {
 	using namespace UE::Editor::DataStorage;
-	IEditorDataStorageProvider& MutableDataStorage = *GetMutableDataStorageFeature<IEditorDataStorageProvider>(StorageFeatureName);
+	ICoreProvider& MutableDataStorage = *GetMutableDataStorageFeature<ICoreProvider>(StorageFeatureName);
 
 	AssetDataCBDataSource = MakeUnique<Private::FTedsAssetDataCBDataSource>(MutableDataStorage);
 	AssetRegistryStorage = MakeUnique<Private::FTedsAssetData>(MutableDataStorage);
diff --git a/Engine/Plugins/Experimental/EditorDataStorageFeatures/Source/TedsAssetData/Private/Tests/TedsAssetDataTest.cpp b/Engine/Plugins/Experimental/EditorDataStorageFeatures/Source/TedsAssetData/Private/Tests/TedsAssetDataTest.cpp
--- a/Engine/Plugins/Experimental/EditorDataStorageFeatures/Source/TedsAssetData/Private/Tests/TedsAssetDataTest.cpp
+++ b/Engine/Plugins/Experimental/EditorDataStorageFeatures/Source/TedsAssetData/Private/Tests/TedsAssetDataTest.cpp
@@ -66,7 +66,7 @@ bool FTedsAssetDataTest::RunTest(const FString& Parameters)
 
 	// Do a sanity check that the data from the asset registry exist in Teds
 	using namespace UE::Editor::DataStorage;
-	const IEditorDataStorageProvider* Database = GetMutableDataStorageFeature<IEditorDataStorageProvider>(StorageFeatureName);
+	const ICoreProvider* Database = GetMutableDataStorageFeature<ICoreProvider>(StorageFeatureName);
 	
 	for (const TPair<FName, int32>& PathAndAssetCount : AssetRegistryPathsAndAssetCount)
 	{
diff --git a/Engine/Plugins/Experimental/EditorDataStorageFeatures/Source/TedsAssetData/Private/Widgets/AssetDataDiskSizeWidget.cpp b/Engine/Plugins/Experimental/EditorDataStorageFeatures/Source/TedsAssetData/Private/Widgets/AssetDataDiskSizeWidget.cpp
--- a/Engine/Plugins/Experimental/EditorDataStorageFeatures/Source/TedsAssetData/Private/Widgets/AssetDataDiskSizeWidget.cpp
+++ b/Engine/Plugins/Experimental/EditorDataStorageFeatures/Source/TedsAssetData/Private/Widgets/AssetDataDiskSizeWidget.cpp
@@ -9,8 +9,8 @@
 #include "Widgets/Layout/SBox.h"
 #include "Widgets/Text/STextBlock.h"
 
-void UDiskSizeWidgetFactory::RegisterWidgetConstructors(IEditorDataStorageProvider& DataStorage,
-                                                        IEditorDataStorageUiProvider& DataStorageUi) const
+void UDiskSizeWidgetFactory::RegisterWidgetConstructors(UE::Editor::DataStorage::ICoreProvider& DataStorage,
+	UE::Editor::DataStorage::IUiProvider& DataStorageUi) const
 {
 	using namespace UE::Editor::DataStorage::Queries;
 	DataStorageUi.RegisterWidgetFactory<FDiskSizeWidgetConstructor>(TEXT("General.Cell"), TColumn<FDiskSizeColumn>());
@@ -21,8 +21,8 @@ FDiskSizeWidgetConstructor::FDiskSizeWidgetConstructor()
 {
 }
 
-TSharedPtr<SWidget> FDiskSizeWidgetConstructor::CreateWidget(IEditorDataStorageProvider* DataStorage,
-	IEditorDataStorageUiProvider* DataStorageUi, UE::Editor::DataStorage::RowHandle TargetRow, UE::Editor::DataStorage::RowHandle WidgetRow,
+TSharedPtr<SWidget> FDiskSizeWidgetConstructor::CreateWidget(UE::Editor::DataStorage::ICoreProvider* DataStorage,
+	UE::Editor::DataStorage::IUiProvider* DataStorageUi, UE::Editor::DataStorage::RowHandle TargetRow, UE::Editor::DataStorage::RowHandle WidgetRow,
 	const UE::Editor::DataStorage::FMetaDataView& Arguments)
 {
 	UE::Editor::DataStorage::FAttributeBinder Binder(TargetRow, DataStorage);
diff --git a/Engine/Plugins/Experimental/EditorDataStorageFeatures/Source/TedsAssetData/Private/Widgets/AssetDataDiskSizeWidget.h b/Engine/Plugins/Experimental/EditorDataStorageFeatures/Source/TedsAssetData/Private/Widgets/AssetDataDiskSizeWidget.h
--- a/Engine/Plugins/Experimental/EditorDataStorageFeatures/Source/TedsAssetData/Private/Widgets/AssetDataDiskSizeWidget.h
+++ b/Engine/Plugins/Experimental/EditorDataStorageFeatures/Source/TedsAssetData/Private/Widgets/AssetDataDiskSizeWidget.h
@@ -7,7 +7,6 @@
 
 #include "AssetDataDiskSizeWidget.generated.h"
 
-class IEditorDataStorageProvider;
 class UScriptStruct;
 
 UCLASS()
@@ -18,8 +17,8 @@ class UDiskSizeWidgetFactory : public UEditorDataStorageFactory
 public:
 	~UDiskSizeWidgetFactory() override = default;
 
-	TEDSASSETDATA_API void RegisterWidgetConstructors(IEditorDataStorageProvider& DataStorage,
-		IEditorDataStorageUiProvider& DataStorageUi) const override;
+	TEDSASSETDATA_API void RegisterWidgetConstructors(UE::Editor::DataStorage::ICoreProvider& DataStorage,
+		UE::Editor::DataStorage::IUiProvider& DataStorageUi) const override;
 };
 
 // Widget to show disk size in bytes
@@ -33,8 +32,8 @@ public:
 	~FDiskSizeWidgetConstructor() override = default;
 
 	TEDSASSETDATA_API virtual TSharedPtr<SWidget> CreateWidget(
-		IEditorDataStorageProvider* DataStorage,
-		IEditorDataStorageUiProvider* DataStorageUi,
+		UE::Editor::DataStorage::ICoreProvider* DataStorage,
+		UE::Editor::DataStorage::IUiProvider* DataStorageUi,
 		UE::Editor::DataStorage::RowHandle TargetRow,
 		UE::Editor::DataStorage::RowHandle WidgetRow,
 		const UE::Editor::DataStorage::FMetaDataView& Arguments) override;
diff --git a/Engine/Plugins/Experimental/EditorDataStorageFeatures/Source/TedsAssetData/Private/Widgets/AssetDataItemTypeWidget.cpp b/Engine/Plugins/Experimental/EditorDataStorageFeatures/Source/TedsAssetData/Private/Widgets/AssetDataItemTypeWidget.cpp
--- a/Engine/Plugins/Experimental/EditorDataStorageFeatures/Source/TedsAssetData/Private/Widgets/AssetDataItemTypeWidget.cpp
+++ b/Engine/Plugins/Experimental/EditorDataStorageFeatures/Source/TedsAssetData/Private/Widgets/AssetDataItemTypeWidget.cpp
@@ -12,8 +12,8 @@
 
 #define LOCTEXT_NAMESPACE "FAssetDataItemTypeWidgetConstructor"
 
-void UAssetDataItemTypeWidgetFactory::RegisterWidgetConstructors(IEditorDataStorageProvider& DataStorage,
-														IEditorDataStorageUiProvider& DataStorageUi) const
+void UAssetDataItemTypeWidgetFactory::RegisterWidgetConstructors(UE::Editor::DataStorage::ICoreProvider& DataStorage,
+	UE::Editor::DataStorage::IUiProvider& DataStorageUi) const
 {
 	using namespace UE::Editor::DataStorage::Queries;
 	DataStorageUi.RegisterWidgetFactory<FAssetDataItemTypeWidgetConstructor>(TEXT("General.Cell"), TColumn<FAssetClassColumn>() || TColumn<FFolderTag>());
@@ -24,8 +24,9 @@ FAssetDataItemTypeWidgetConstructor::FAssetDataItemTypeWidgetConstructor()
 {
 }
 
-TSharedPtr<SWidget> FAssetDataItemTypeWidgetConstructor::CreateWidget(IEditorDataStorageProvider* DataStorage,
-	IEditorDataStorageUiProvider* DataStorageUi, UE::Editor::DataStorage::RowHandle TargetRow, UE::Editor::DataStorage::RowHandle WidgetRow,
+TSharedPtr<SWidget> FAssetDataItemTypeWidgetConstructor::CreateWidget(UE::Editor::DataStorage::ICoreProvider* DataStorage,
+	UE::Editor::DataStorage::IUiProvider* DataStorageUi,
+	UE::Editor::DataStorage::RowHandle TargetRow, UE::Editor::DataStorage::RowHandle WidgetRow,
 	const UE::Editor::DataStorage::FMetaDataView& Arguments)
 {
 	UE::Editor::DataStorage::FAttributeBinder Binder(TargetRow, DataStorage);
