Adding an automated test for the batch modifcation for queries and also added a test for the normal deffered modification for a query.

The test highlited a few issues that were fixed,.
1- bad usage of the api to batch build the new rows with the batch modifications.
2- Fixed a issue when deleting a row with the bach modification.
3- Fixed an array view that wasn't always of the right size in the AddColumns function override DynamicColumnDescriptions

But the test aslo highlighted an issue with were the dynamic columns may not always behave as expected and seems to also have a race condition. This issue wan't fixed so the line that test those are commented for now.

diff --git a/Engine/Plugins/Experimental/EditorDataStorage/Source/TedsCore/Private/Processors/TypedElementProcessorAdaptors.cpp b/Engine/Plugins/Experimental/EditorDataStorage/Source/TedsCore/Private/Processors/TypedElementProcessorAdaptors.cpp
--- a/Engine/Plugins/Experimental/EditorDataStorage/Source/TedsCore/Private/Processors/TypedElementProcessorAdaptors.cpp
+++ b/Engine/Plugins/Experimental/EditorDataStorage/Source/TedsCore/Private/Processors/TypedElementProcessorAdaptors.cpp
@@ -373,113 +373,107 @@ namespace UE::Editor::DataStorage
 						InRowsToAdd = MoveTemp(RowsToAdd)
 						](FMassEntityManager& System) mutable
 						{
-							TSharedRef<FMassEntityManager::FEntityCreationContext> ObtainedContext = System.GetOrMakeCreationContext();
-
-							// Adds Rows
 							{
-								TArray<FMassArchetypeEntityCollectionWithPayload> EntityCollections;
+								TSharedRef<FMassEntityManager::FEntityCreationContext> ObtainedContext = System.GetOrMakeCreationContext();
+
+								// Adds Rows
 								for (TPair<FMassArchetypeHandle, TArray<FMassEntityHandle>>& ArchetypeAndRows : InRowsToAdd)
 								{
-									FMassArchetypeEntityCollectionWithPayload::CreateEntityRangesWithPayload(System, ArchetypeAndRows.Value, FMassArchetypeEntityCollection::NoDuplicates
-										, FMassGenericPayloadView(), EntityCollections);
+									FMassArchetypeEntityCollection EntityCollection(FMassArchetypeHandle(), ArchetypeAndRows.Value, FMassArchetypeEntityCollection::EDuplicatesHandling::FoldDuplicates);
 									const FMassArchetypeCompositionDescriptor& ArchetypeComposition = System.GetArchetypeComposition(ArchetypeAndRows.Key);
 
-									for (const FMassArchetypeEntityCollectionWithPayload& Collection : EntityCollections)
-									{
-										System.BatchBuildEntities(Collection, ArchetypeComposition);
-									}
-
-									EntityCollections.Reset();
+									FMassArchetypeEntityCollectionWithPayload EntityCollectionWithPayload(MoveTemp(EntityCollection));
+									System.BatchBuildEntities(EntityCollectionWithPayload, ArchetypeComposition);
 								}
-							}
-
-							using EntityHandleArray = TArray<FMassEntityHandle, TInlineAllocator<32>>;
-							using EntityArchetypeLookup = TMap<FMassArchetypeHandle, EntityHandleArray, TInlineSetAllocator<32>>;
-							using ArchetypeEntityArray = TArray<FMassArchetypeEntityCollection, TInlineAllocator<32>>;
 
-							EntityArchetypeLookup LookupTable;
-							ArchetypeEntityArray EntityCollections;
+								using EntityHandleArray = TArray<FMassEntityHandle, TInlineAllocator<32>>;
+								using EntityArchetypeLookup = TMap<FMassArchetypeHandle, EntityHandleArray, TInlineSetAllocator<32>>;
+								using ArchetypeEntityArray = TArray<FMassArchetypeEntityCollection, TInlineAllocator<32>>;
 
-							// Change the entity types
-							for (TPair<FRowColumnModifications, TArray<FMassEntityHandle>>& BatchedModification : InChangesAndRows)
-							{
-								for (const FDynamicColumnDescription& DynamicColumnAdded : BatchedModification.Key.AddedDynamicColumns)
-								{
-									BatchedModification.Key.AddedFragments.Add(*PtrToEnvironment->GenerateDynamicColumn(*DynamicColumnAdded.TemplateType, DynamicColumnAdded.Identifier));
-								}
+								EntityArchetypeLookup LookupTable;
+								ArchetypeEntityArray EntityCollections;
 
-								const FRowColumnModifications& ColumnModification = BatchedModification.Key;
-								if (!ColumnModification.AddedFragments.IsEmpty() || !ColumnModification.RemovedFragments.IsEmpty())
+								// Change the entity types
+								for (TPair<FRowColumnModifications, TArray<FMassEntityHandle>>& BatchedModification : InChangesAndRows)
 								{
-									LookupTable.Reset();
-									EntityCollections.Reset();
+									for (const FDynamicColumnDescription& DynamicColumnAdded : BatchedModification.Key.AddedDynamicColumns)
+									{
+										BatchedModification.Key.AddedFragments.Add(*PtrToEnvironment->GenerateDynamicColumn(*DynamicColumnAdded.TemplateType, DynamicColumnAdded.Identifier));
+									}
 
-									// Sort rows (entities) into to matching table (archetype) bucket.
-									for (FMassEntityHandle EntityHandle : BatchedModification.Value)
+									const FRowColumnModifications& ColumnModification = BatchedModification.Key;
+									if (!ColumnModification.AddedFragments.IsEmpty() || !ColumnModification.RemovedFragments.IsEmpty())
 									{
-										if (System.IsEntityValid(EntityHandle))
+										LookupTable.Reset();
+										EntityCollections.Reset();
+
+										// Sort rows (entities) into to matching table (archetype) bucket.
+										for (FMassEntityHandle EntityHandle : BatchedModification.Value)
 										{
-											FMassArchetypeHandle Archetype = System.GetArchetypeForEntity(EntityHandle);
-											EntityHandleArray& EntityCollection = LookupTable.FindOrAdd(Archetype);
-											EntityCollection.Add(EntityHandle);
+											if (System.IsEntityValid(EntityHandle))
+											{
+												FMassArchetypeHandle Archetype = System.GetArchetypeForEntity(EntityHandle);
+												EntityHandleArray& EntityCollection = LookupTable.FindOrAdd(Archetype);
+												EntityCollection.Add(EntityHandle);
+											}
 										}
-									}
 
-									// Construct table (archetype) specific row (entity) collections.
-									EntityCollections.Reserve(LookupTable.Num());
-									for (auto It = LookupTable.CreateConstIterator(); It; ++It)
-									{
-										// Since we use a map to combine all the operation on a row in one via a map go, we already know there won't be any duplicate
-										EntityCollections.Emplace(It.Key(), It.Value(), FMassArchetypeEntityCollection::EDuplicatesHandling::NoDuplicates);
+										// Construct table (archetype) specific row (entity) collections.
+										EntityCollections.Reserve(LookupTable.Num());
+										for (auto It = LookupTable.CreateConstIterator(); It; ++It)
+										{
+											// Since we use a map to combine all the operation on a row in one via a map go, we already know there won't be any duplicate
+											EntityCollections.Emplace(It.Key(), It.Value(), FMassArchetypeEntityCollection::EDuplicatesHandling::NoDuplicates);
+										}
+
+										// This could be improved by adding an operation that would both combine the Fragments and Tags change in one bath operation.
+										System.BatchChangeFragmentCompositionForEntities(EntityCollections, ColumnModification.AddedFragments, ColumnModification.RemovedFragments);
 									}
 
-									// This could be improved by adding an operation that would both combine the Fragments and Tags change in one bath operation.
-									System.BatchChangeFragmentCompositionForEntities(EntityCollections, ColumnModification.AddedFragments,  ColumnModification.RemovedFragments);
-								}
+									if (!ColumnModification.AddedTags.IsEmpty() || !ColumnModification.RemovedTags.IsEmpty())
+									{
+										LookupTable.Reset();
+										EntityCollections.Reset();
 
-								if (!ColumnModification.AddedTags.IsEmpty() || !ColumnModification.RemovedTags.IsEmpty())
-								{
-									LookupTable.Reset();
-									EntityCollections.Reset();
+										// Sort rows (entities) into to matching table (archetype) bucket.
+										for (FMassEntityHandle EntityHandle : BatchedModification.Value)
+										{
+											if (System.IsEntityValid(EntityHandle))
+											{
+												FMassArchetypeHandle Archetype = System.GetArchetypeForEntity(EntityHandle);
+												EntityHandleArray& EntityCollection = LookupTable.FindOrAdd(Archetype);
+												EntityCollection.Add(EntityHandle);
+											}
+										}
 
-									// Sort rows (entities) into to matching table (archetype) bucket.
-									for (FMassEntityHandle EntityHandle : BatchedModification.Value)
-									{
-										if (System.IsEntityValid(EntityHandle))
+										// Construct table (archetype) specific row (entity) collections.
+										EntityCollections.Reserve(LookupTable.Num());
+										for (auto It = LookupTable.CreateConstIterator(); It; ++It)
 										{
-											FMassArchetypeHandle Archetype = System.GetArchetypeForEntity(EntityHandle);
-											EntityHandleArray& EntityCollection = LookupTable.FindOrAdd(Archetype);
-											EntityCollection.Add(EntityHandle);
+											// Since we use a map to combine all the operation on a row in one via a map go, we already know there won't be any duplicate
+											EntityCollections.Emplace(It.Key(), It.Value(), FMassArchetypeEntityCollection::EDuplicatesHandling::NoDuplicates);
 										}
-									}
 
-									// Construct table (archetype) specific row (entity) collections.
-									EntityCollections.Reserve(LookupTable.Num());
-									for (auto It = LookupTable.CreateConstIterator(); It; ++It)
-									{
-										// Since we use a map to combine all the operation on a row in one via a map go, we already know there won't be any duplicate
-										EntityCollections.Emplace(It.Key(), It.Value(), FMassArchetypeEntityCollection::EDuplicatesHandling::NoDuplicates);
+										System.BatchChangeTagsForEntities(EntityCollections, ColumnModification.AddedTags, ColumnModification.RemovedTags);
 									}
-
-									System.BatchChangeTagsForEntities(EntityCollections, ColumnModification.AddedTags, ColumnModification.RemovedTags);
 								}
-							}
 
-							// Do the relocation
-							for (TPair<FMassEntityHandle, FRowOperations>& PerRowOperations : InRowOperations)
-							{
-								const FMassEntityHandle Handle = PerRowOperations.Key;
-								for (FRelocator& Relocator : PerRowOperations.Value.Relocators)
+								// Do the relocation
+								for (TPair<FMassEntityHandle, FRowOperations>& PerRowOperations : InRowOperations)
 								{
-									if (!Relocator.Identifier.IsNone())
+									const FMassEntityHandle Handle = PerRowOperations.Key;
+									for (FRelocator& Relocator : PerRowOperations.Value.Relocators)
 									{
-										Relocator.FragmentType = PtrToEnvironment->GenerateDynamicColumn(*Relocator.FragmentType, Relocator.Identifier);
-									}
+										if (!Relocator.Identifier.IsNone())
+										{
+											Relocator.FragmentType = PtrToEnvironment->GenerateDynamicColumn(*Relocator.FragmentType, Relocator.Identifier);
+										}
 
-									FStructView Fragment = System.GetFragmentDataStruct(Handle, Relocator.FragmentType);
-									Relocator.CopyOrMove(*Relocator.FragmentType, Fragment.GetMemory(), Relocator.Object);
+										FStructView Fragment = System.GetFragmentDataStruct(Handle, Relocator.FragmentType);
+										Relocator.CopyOrMove(*Relocator.FragmentType, Fragment.GetMemory(), Relocator.Object);
+									}
 								}
-							}
+							} // We use a scope here to trigger the notification. We can't delete a row while we are holding a EntityCreationContext.
 
 							// Remove Rows
 							if (!InEntitiesToDestroy.IsEmpty())
@@ -530,7 +524,7 @@ namespace UE::Editor::DataStorage
 				FAddDynamicColumnsBatch* CommandData = ScratchBuffer.Emplace<FAddDynamicColumnsBatch>();
 				TArrayView<RowHandle> ScratchRows = MakeArrayView(
 					ScratchBuffer.EmplaceArray<RowHandle>(Rows.Num()),
-					DynamicColumnDescriptions.Num());
+					Rows.Num());
 				TArrayView<FDynamicColumnDescription> ScratchDescriptions = MakeArrayView(
 					ScratchBuffer.EmplaceArray<FDynamicColumnDescription>(DynamicColumnDescriptions.Num()),
 					DynamicColumnDescriptions.Num());
diff --git a/Engine/Plugins/Experimental/EditorDataStorage/Source/TedsCore/Tests/TypedElementDataStorageInterfaceTests.cpp b/Engine/Plugins/Experimental/EditorDataStorage/Source/TedsCore/Tests/TypedElementDataStorageInterfaceTests.cpp
--- a/Engine/Plugins/Experimental/EditorDataStorage/Source/TedsCore/Tests/TypedElementDataStorageInterfaceTests.cpp
+++ b/Engine/Plugins/Experimental/EditorDataStorage/Source/TedsCore/Tests/TypedElementDataStorageInterfaceTests.cpp
@@ -2,12 +2,17 @@
 
 
 #if WITH_TESTS
+#include "Containers/ArrayView.h"
 #include "Elements/Common/EditorDataStorageFeatures.h"
-#include "Elements/Framework/TypedElementTestColumns.h"
 #include "Elements/Framework/TypedElementQueryBuilder.h"
+#include "Elements/Framework/TypedElementTestColumns.h"
 #include "Elements/Interfaces/TypedElementDataStorageInterface.h"
 #include "Misc/AutomationTest.h"
 
+#if WITH_EDITOR
+#include "TickableEditorObject.h"
+#endif
+
 namespace UE::Editor::DataStorage::Tests
 {
 	TableHandle RegisterTestTableA(ICoreProvider* TedsInterface)
@@ -122,6 +127,108 @@ namespace UE::Editor::DataStorage::Tests
 			}
 			CreatedRows.Empty();
 		}
+
+#if WITH_EDITOR
+		void TestDeferredModifications(bool bIsTestingBatchModifications, const int32 NumberOfTestRows)
+		{
+			const RowHandle AddedRowHandle = TedsInterface->AddRow(TestTableHandleA);
+			TArray<RowHandle> AddedRowHandles;
+			AddedRowHandles.Reserve(NumberOfTestRows);
+			TedsInterface->BatchAddRow(TestTableHandleA, NumberOfTestRows, [&AddedRowHandles](RowHandle Row)
+				{
+					AddedRowHandles.Add(Row);
+				});
+
+
+			TArray<RowHandle> RowHandlesAddedByQuery;
+			RowHandlesAddedByQuery.Reserve(NumberOfTestRows);
+			RowHandle RowHandleAddedByQuery;
+
+			FName TestIdentifier(TEXT("TestIdentifier"));
+
+			// We need to uncomment the part of the test that create and check the dynamic column description once the race condition has been fixed
+
+			using namespace Queries;
+			QueryHandle Query = TedsInterface->RegisterQuery(
+				Select(
+					bIsTestingBatchModifications ? TEXT("Test deferred modifications with batch modifications query") : TEXT("Test deferred modifications query"),
+					FProcessor(EQueryTickPhase::DuringPhysics, TedsInterface->GetQueryTickGroupName(EQueryTickGroups::Update))
+						.BatchModifications(bIsTestingBatchModifications),
+					[this, NumberOfTestRows, AddedRowHandle, &AddedRowHandles, &RowHandlesAddedByQuery, &RowHandleAddedByQuery, TestIdentifier](IQueryContext& Context, const RowHandle* Rows, const FTestColumnA* ColumnA, const FTestColumnB* ColumnB)
+					{
+						const int32 NumberOfRows = Context.GetRowCount();
+
+						TArrayView<const RowHandle> InRowsArrayView(Rows, NumberOfRows);
+
+						FDynamicColumnDescription DynamicColumnDescription;
+						DynamicColumnDescription.Identifier = TestIdentifier;
+						DynamicColumnDescription.TemplateType = FTestColumnDynamic::StaticStruct();
+						//Context.AddColumns(InRowsArrayView, TArrayView<FDynamicColumnDescription>(&DynamicColumnDescription, 1));
+						Context.RemoveColumns<FTestColumnB, FTestTagColumnB>(InRowsArrayView);
+
+						for (int32 Count = 0; Count < NumberOfTestRows; ++Count)
+						{
+							RowHandlesAddedByQuery.Add(Context.AddRow(TestTableHandleA));
+						}
+
+						Context.AddColumns<FTestColumnC, FTestTagColumnC>(RowHandlesAddedByQuery);
+
+						RowHandleAddedByQuery = Context.AddRow(TestTableHandleA);
+						Context.AddColumns<FTestColumnC, FTestTagColumnC>(RowHandleAddedByQuery);
+						FTestColumnInt IntColumn;
+						FTestColumnInt& IntColumnRef = Context.AddColumn<FTestColumnInt>(RowHandleAddedByQuery, MoveTemp(IntColumn));
+						IntColumnRef.TestInt = NumberOfTestRows;
+						Context.RemoveColumns<FTestTagColumnB>(RowHandleAddedByQuery);
+						//Context.AddColumn<FTestColumnDynamic>(RowHandleAddedByQuery, TestIdentifier);
+
+						Context.RemoveRow(AddedRowHandle);
+						Context.RemoveRows(AddedRowHandles);
+					})
+				.Where(TColumn<FTestTagColumnA>() && TColumn<FTestTagColumnB>())
+				.Compile());
+
+			// Tick Mass
+			FTickableEditorObject::TickObjects(-1.f);
+
+
+			for (const RowHandle& Row : CreatedRows)
+			{
+				//TestNotNull(TEXT("CreatedRows have the dynamic column"), TedsInterface->GetColumn<FTestColumnDynamic>(Row, TestIdentifier));
+				TestFalse(TEXT("CreatedRows don't have a FTestColumnB and FTestTagColumnB"), TedsInterface->HasColumns<FTestColumnB, FTestTagColumnB>(Row));
+			}
+
+			TestFalse(TEXT("The AddedRowHandle should be removed"), TedsInterface->IsRowAssigned(AddedRowHandle));
+			for (const RowHandle& Row : AddedRowHandles)
+			{
+				TestFalse(TEXT("The AddedRowHandles should be removed"), TedsInterface->IsRowAssigned(AddedRowHandle));
+			}
+
+			TestTrue(TEXT("The RowHandleAddedByQuery should be assigned"), TedsInterface->IsRowAssigned(RowHandleAddedByQuery));
+			TestTrue(TEXT("The RowHandleAddedByQuery should have a column FTestColumnC and a Tag ColumnC"), TedsInterface->HasColumns<FTestColumnC, FTestTagColumnC>(RowHandleAddedByQuery));
+			TestFalse(TEXT("The RowHandleAddedByQuery shouldn't have the tag FTestTagColumnB"), TedsInterface->HasColumns<FTestTagColumnB>(RowHandleAddedByQuery));
+			//TestNotNull(TEXT("The RowHandleAddedByQuery should have the dynamic column"), TedsInterface->GetColumn<FTestColumnDynamic>(RowHandleAddedByQuery, TestIdentifier));
+			{
+				const FTestColumnInt* IntColumn = TedsInterface->GetColumn<FTestColumnInt>(RowHandleAddedByQuery);
+				if (TestNotNull(TEXT("The RowHandleAddedByQuery should have a FTestColumnInt"), IntColumn))
+				{
+					TestEqual(TEXT("The RowHandleAddedByQuery FTestColumnInt should have the expected value"), IntColumn->TestInt, NumberOfTestRows);
+				}
+			}
+
+			for (const RowHandle& Row : RowHandlesAddedByQuery)
+			{
+				if (TestTrue(TEXT("Check that the RowHandlesAddedByQuery are assigned"), TedsInterface->IsRowAssigned(Row)))
+				{
+					TestTrue(TEXT("Check that the RowHandlesAddedByQuery have a FTestColumnC and FTestTagColumnC"), TedsInterface->HasColumns<FTestColumnC, FTestTagColumnC>(Row));
+				}
+			}
+
+			TedsInterface->RemoveRow(RowHandleAddedByQuery);
+			TedsInterface->BatchRemoveRows(RowHandlesAddedByQuery);
+			TedsInterface->UnregisterQuery(Query);
+		}
+#endif // WITH_EDITOR
+
 	END_DEFINE_SPEC(EditorDataStorageTestsFixture)
 
 	void EditorDataStorageTestsFixture::Define()
@@ -637,6 +744,31 @@ namespace UE::Editor::DataStorage::Tests
 							});
 						});
 					});
+
+#if WITH_EDITOR
+					Describe("", [this]()
+						{
+							constexpr int32 NumberOfTestRows = 4;
+
+							BeforeEach([this]()
+								{
+									CreateTestRows(NumberOfTestRows);
+								});
+
+							It("Processor Deferred Commands", [this]()
+								{
+									constexpr bool bIsTestingBatchModifications = false;
+									TestDeferredModifications(bIsTestingBatchModifications, NumberOfTestRows);
+								});
+
+							It("Processor Deferred Commands with Batch Modifications", [this]()
+								{
+									constexpr bool bIsTestingBatchModifications = true;
+									TestDeferredModifications(bIsTestingBatchModifications, NumberOfTestRows);
+								});
+						});
+#endif // WITH_EDITOR
+
 				});
 			
 			});
diff --git a/Engine/Source/Runtime/TypedElementFramework/Tests/Elements/Framework/TypedElementTestColumns.h b/Engine/Source/Runtime/TypedElementFramework/Tests/Elements/Framework/TypedElementTestColumns.h
--- a/Engine/Source/Runtime/TypedElementFramework/Tests/Elements/Framework/TypedElementTestColumns.h
+++ b/Engine/Source/Runtime/TypedElementFramework/Tests/Elements/Framework/TypedElementTestColumns.h
@@ -49,6 +49,14 @@ struct FTestColumnG final : public FEditorDataStorageColumn
 	GENERATED_BODY()
 };
 
+USTRUCT(meta = (DisplayName = "ColumnDynamic", EditorDataStorage_DynamicColumnTemplate))
+struct FTestColumnDynamic final : public FEditorDataStorageColumn
+{
+	GENERATED_BODY()
+
+	struct EditorDataStorage_DynamicColumnTemplate {};
+};
+
 USTRUCT(meta = (DisplayName = "ColumnInt"))
 struct FTestColumnInt final : public FEditorDataStorageColumn
 {
