TEDS UI: Refactor the widget purpose system and update all tools/widgets to replace deprecation function uses

- Widget purposes and widget factories are now stored directly in a TEDS table
- Widget purposes can now contain more fine grained info such as a namespace, name and frame
- Widget purposes now support hierarchies and widget construction methods can go up the parent chain if required
- Widget purposes can be identified by their row handle, or by a unique FName they are mapped to

diff --git a/Engine/Plugins/ChaosVD/Source/ChaosVD/Private/ChaosVDWorldOutlinerTab.cpp b/Engine/Plugins/ChaosVD/Source/ChaosVD/Private/ChaosVDWorldOutlinerTab.cpp
--- a/Engine/Plugins/ChaosVD/Source/ChaosVD/Private/ChaosVDWorldOutlinerTab.cpp
+++ b/Engine/Plugins/ChaosVD/Source/ChaosVD/Private/ChaosVDWorldOutlinerTab.cpp
@@ -119,7 +119,7 @@ void FChaosVDWorldOutlinerTab::CreateWorldOutlinerWidget()
 		return FSceneOutlinerTreeItemID(RowHandle);
 	}));
 
-	FSceneOutlinerTedsQueryBinder::GetInstance().AssignQuery(ColumnQuery, TedsOutliner, Params.CellWidgetPurposes);
+	FSceneOutlinerTedsQueryBinder::GetInstance().AssignQuery(ColumnQuery, TedsOutliner, Params.CellWidgetPurpose, Params.LabelWidgetPurpose);
 	
 	SceneOutlinerWidget = TedsOutliner;
 
diff --git a/Engine/Plugins/Experimental/EditorDataStorage/Source/TedsCore/Private/TypedElementDatabaseUI.cpp b/Engine/Plugins/Experimental/EditorDataStorage/Source/TedsCore/Private/TypedElementDatabaseUI.cpp
--- a/Engine/Plugins/Experimental/EditorDataStorage/Source/TedsCore/Private/TypedElementDatabaseUI.cpp
+++ b/Engine/Plugins/Experimental/EditorDataStorage/Source/TedsCore/Private/TypedElementDatabaseUI.cpp
@@ -5,9 +5,10 @@
 #include "Algo/BinarySearch.h"
 #include "Algo/Sort.h"
 #include "Algo/Unique.h"
+#include "Elements/Columns/TypedElementHiearchyColumns.h"
+#include "Elements/Columns/WidgetPurposeColumns.h"
 #include "Elements/Columns/TypedElementMiscColumns.h"
 #include "Elements/Columns/TypedElementSlateWidgetColumns.h"
-#include "Elements/Columns/TypedElementTypeInfoColumns.h"
 #include "Elements/Framework/TypedElementColumnUtils.h"
 #include "Elements/Interfaces/TypedElementDataStorageInterface.h"
 #include "Elements/Interfaces/TypedElementDataStorageCompatibilityInterface.h"
@@ -17,6 +18,8 @@
 
 DEFINE_LOG_CATEGORY(LogEditorDataStorageUI);
 
+#define LOCTEXT_NAMESPACE "TypedElementDatabaseUI"
+
 namespace Internal
 {
 	// Source: https://en.cppreference.com/w/cpp/utility/variant/visit
@@ -44,6 +47,13 @@ namespace Internal
 
 		return false;
 	}
+
+	static const UE::Editor::DataStorage::IUiProvider::FPurposeID DefaultWidgetPurposeID(
+		UE::Editor::DataStorage::IUiProvider::FPurposeInfo("General", "Cell", "Default").GeneratePurposeID());
+	
+	static const UE::Editor::DataStorage::IUiProvider::FPurposeID GeneralWidgetPurposeID(
+		UE::Editor::DataStorage::IUiProvider::FPurposeInfo("General", "Cell", NAME_None).GeneratePurposeID());
+	
 }
 
 void UEditorDataStorageUi::Initialize(
@@ -55,273 +65,341 @@ void UEditorDataStorageUi::Initialize(
 	Storage = StorageInterface;
 	StorageCompatibility = StorageCompatibilityInterface;
 	CreateStandardArchetypes();
+	RegisterQueries();
 }
 
 void UEditorDataStorageUi::Deinitialize()
 {
 }
 
-void UEditorDataStorageUi::RegisterWidgetPurpose(FName Purpose, EPurposeType Type, FText Description)
+UE::Editor::DataStorage::RowHandle UEditorDataStorageUi::RegisterWidgetPurpose(const FPurposeID& PurposeID, const FPurposeInfo& InPurposeInfo)
 {
-	FPurposeInfo* PurposeInfo = WidgetPurposes.Find(Purpose);
-	if (!PurposeInfo)
+	using namespace UE::Editor::DataStorage;
+
+	// If a purpose is already registered against this name, let the user know
+	const FMapKey Key = FMapKey(PurposeID);
+	RowHandle ExistingRow = Storage->LookupMappedRow(Key);
+	if (Storage->IsRowAvailable(ExistingRow))
+	{
+		ensureMsgf(false, TEXT("Existing purpose found registered with name: %s"), *PurposeID.ToString());
+		return InvalidRowHandle;
+	}
+
+	// Add the row and register the mapping
+	RowHandle PurposeRowHandle = Storage->AddRow(WidgetPurposeTable);
+	Storage->MapRow(Key, PurposeRowHandle);
+
+	// Setup the relevant columns
+	if (FWidgetPurposeColumn* PurposeColumn = Storage->GetColumn<FWidgetPurposeColumn>(PurposeRowHandle))
 	{
-		FPurposeInfo& NewInfo = WidgetPurposes.Add(Purpose);
-		NewInfo.Type = Type;
-		NewInfo.Description = MoveTemp(Description);
+		PurposeColumn->PurposeType = InPurposeInfo.Type;
+		PurposeColumn->PurposeID = PurposeID;
 	}
+
+	if (FWidgetPurposeNameColumn* PurposeNameColumn = Storage->GetColumn<FWidgetPurposeNameColumn>(PurposeRowHandle))
+	{
+		PurposeNameColumn->Namespace = InPurposeInfo.Namespace;
+		PurposeNameColumn->Name = InPurposeInfo.Name;
+		PurposeNameColumn->Frame = InPurposeInfo.Frame;
+	}
+
+	if (!InPurposeInfo.Description.IsEmpty())
+	{
+		Storage->AddColumn(PurposeRowHandle, FDescriptionColumn{ .Description = InPurposeInfo.Description });
+	}
+
+	// If the parent purpose already exists, simply reference it. Otherwise add an unresolved parent column to resolve it later
+	if (InPurposeInfo.ParentPurposeID.IsSet())
+	{
+		RowHandle ParentRowHandle = FindPurpose(InPurposeInfo.ParentPurposeID);
+
+		if (Storage->IsRowAvailable(ParentRowHandle))
+		{
+			Storage->AddColumn(PurposeRowHandle, FTableRowParentColumn{ .Parent = ParentRowHandle });
+		}
+		else
+		{
+			Storage->AddColumn(PurposeRowHandle, FUnresolvedTableRowParentColumn{ .ParentIdKey = InPurposeInfo.ParentPurposeID});
+		}
+	}
+	return PurposeRowHandle;
 }
 
-bool UEditorDataStorageUi::RegisterWidgetFactory(FName Purpose, const UScriptStruct* Constructor)
+UE::Editor::DataStorage::RowHandle UEditorDataStorageUi::RegisterWidgetPurpose(const FPurposeInfo& InPurposeInfo)
+{
+	return RegisterWidgetPurpose(InPurposeInfo.GeneratePurposeID(), InPurposeInfo);
+}
+
+bool UEditorDataStorageUi::RegisterWidgetFactory(UE::Editor::DataStorage::RowHandle PurposeRow, const UScriptStruct* Constructor)
 {
 	checkf(Constructor->IsChildOf(FTypedElementWidgetConstructor::StaticStruct()),
 		TEXT("Attempting to register a widget constructor '%s' that isn't derived from FTypedElementWidgetConstructor."),
 		*Constructor->GetFullName());
 	
-	if (FPurposeInfo* PurposeInfo = WidgetPurposes.Find(Purpose))
+	if (FWidgetPurposeColumn* PurposeInfo = Storage->GetColumn<FWidgetPurposeColumn>(PurposeRow))
 	{
-		switch (PurposeInfo->Type)
+		switch (PurposeInfo->PurposeType)
 		{
-		case IUiProvider::EPurposeType::Generic:
-			PurposeInfo->Factories.Emplace(Constructor);
-			PurposeInfo->bIsSorted = false;
-			return true;
-		case IUiProvider::EPurposeType::UniqueByName:
-			if (PurposeInfo->Factories.IsEmpty())
+		case EPurposeType::Generic:
 			{
-				PurposeInfo->Factories.Emplace(Constructor);
-				PurposeInfo->bIsSorted = false;
+				UE::Editor::DataStorage::RowHandle FactoryRow = RegisterWidgetFactoryRow(PurposeRow);
+				Storage->AddColumn(FactoryRow, FWidgetFactoryConstructorTypeInfoColumn{.Constructor = Constructor});
+				return true;
 			}
-			else
+		case EPurposeType::UniqueByName:
 			{
-				PurposeInfo->Factories.EmplaceAt(0, Constructor);
+				UE::Editor::DataStorage::RowHandle FactoryRow = RegisterUniqueWidgetFactoryRow(PurposeRow);
+				Storage->AddColumn(FactoryRow, FWidgetFactoryConstructorTypeInfoColumn{.Constructor = Constructor});
+				return true;
+			}
+		case EPurposeType::UniqueByNameAndColumn:
+			{
+				UE_LOG(LogEditorDataStorageUI, Warning,
+				TEXT("Unable to register widget factory '%s' as purpose '%llu' requires at least one column for matching."), 
+				*Constructor->GetName(), PurposeRow);
+				return false;
 			}
-			return true;
-		case IUiProvider::EPurposeType::UniqueByNameAndColumn:
-			UE_LOG(LogEditorDataStorageUI, Warning,
-				TEXT("Unable to register widget factory '%s' as purpose '%s' requires at least one column for matching."), 
-				*Constructor->GetName(), *Purpose.ToString());
-			return false;
 		default:
-			checkf(false, TEXT("Unexpected IUiProvider::EPurposeType found provided when registering widget factory."));
-			return false;
+			{
+				checkf(false, TEXT("Unexpected UE::Editor::DataStorage::IUiProvider::EPurposeType found provided when registering widget factory."));
+				return false;
+			}
 		}
 	}
-	else
-	{
-		UE_LOG(LogEditorDataStorageUI, Warning, 
-			TEXT("Unable to register widget factory '%s' as purpose '%s' isn't registered."), *Constructor->GetName(), *Purpose.ToString());
-		return false;
-	}
+	
+	UE_LOG(LogEditorDataStorageUI, Warning, 
+				TEXT("Unable to register widget factory '%s' as purpose '%llu' isn't registered."), *Constructor->GetName(), PurposeRow);
+	return false;
 }
 
-bool UEditorDataStorageUi::RegisterWidgetFactory(
-	FName Purpose, const UScriptStruct* Constructor, UE::Editor::DataStorage::Queries::FConditions Columns)
+bool UEditorDataStorageUi::RegisterWidgetFactory(UE::Editor::DataStorage::RowHandle PurposeRow, const UScriptStruct* Constructor,
+	UE::Editor::DataStorage::Queries::FConditions Columns)
 {
 	if (!Columns.IsEmpty())
 	{
 		checkf(Constructor->IsChildOf(FTypedElementWidgetConstructor::StaticStruct()),
 			TEXT("Attempting to register a widget constructor '%s' that isn't deriving from FTypedElementWidgetConstructor."),
 			*Constructor->GetFullName());
-
-		if (FPurposeInfo* PurposeInfo = WidgetPurposes.Find(Purpose))
+		
+		if (FWidgetPurposeColumn* PurposeInfo = Storage->GetColumn<FWidgetPurposeColumn>(PurposeRow))
 		{
-			switch (PurposeInfo->Type)
+			switch (PurposeInfo->PurposeType)
 			{
-			case IUiProvider::EPurposeType::Generic:
-				PurposeInfo->Factories.Emplace(Constructor);
-				PurposeInfo->bIsSorted = false;
-				return true;
-			case IUiProvider::EPurposeType::UniqueByName:
-				if (!Columns.IsEmpty())
+			case EPurposeType::Generic:
 				{
-					if (PurposeInfo->Factories.IsEmpty())
-					{
-						PurposeInfo->Factories.Emplace(Constructor);
-						PurposeInfo->bIsSorted = false;
-					}
-					else
-					{
-						PurposeInfo->Factories.EmplaceAt(0, Constructor);
-					}
+					UE::Editor::DataStorage::RowHandle FactoryRow = RegisterWidgetFactoryRow(PurposeRow);
+					Storage->AddColumn(FactoryRow, FWidgetFactoryConstructorTypeInfoColumn{.Constructor = Constructor});
 					return true;
 				}
-				else
-				{
-					return false;
-				}
-			case IUiProvider::EPurposeType::UniqueByNameAndColumn:
+			case EPurposeType::UniqueByName:
 				if (!Columns.IsEmpty())
 				{
-					PurposeInfo->Factories.Emplace(Constructor, MoveTemp(Columns));
-					PurposeInfo->bIsSorted = false;
-					return true;
+					UE::Editor::DataStorage::RowHandle FactoryRow = RegisterUniqueWidgetFactoryRow(PurposeRow);
+					Storage->AddColumn(FactoryRow, FWidgetFactoryConstructorTypeInfoColumn{.Constructor = Constructor});
 				}
 				else
 				{
 					return false;
 				}
+			case EPurposeType::UniqueByNameAndColumn:
+				{
+					if (!Columns.IsEmpty())
+					{
+						Columns.Compile(UE::Editor::DataStorage::Queries::FEditorStorageQueryConditionCompileContext(Storage));
+							
+						UE::Editor::DataStorage::RowHandle FactoryRow = RegisterWidgetFactoryRow(PurposeRow);
+						Storage->AddColumn(FactoryRow, FWidgetFactoryConstructorTypeInfoColumn{.Constructor = Constructor});
+						Storage->AddColumn(FactoryRow, FWidgetFactoryConditionsColumn{.Conditions = Columns});
+							
+						return true;
+					}
+					else
+					{
+						return false;
+					}
+				}
 			default:
-				checkf(false, TEXT("Unexpected IUiProvider::EPurposeType found provided when registering widget factory."));
+				checkf(false, TEXT("Unexpected UE::Editor::DataStorage::IUiProvider::EPurposeType found provided when registering widget factory."));
 				return false;
 			}
 		}
-		else
-		{
-			UE_LOG(LogEditorDataStorageUI, Warning,
-				TEXT("Unable to register widget factory '%s' as purpose '%s' isn't registered."), *Constructor->GetName(), *Purpose.ToString());
-			return false;
-		}
+		
+		UE_LOG(LogEditorDataStorageUI, Warning,
+				TEXT("Unable to register widget factory '%s' as purpose '%llu' isn't registered."), *Constructor->GetName(), PurposeRow);
+		return false;
 	}
 	else
 	{
-		return RegisterWidgetFactory(Purpose, Constructor);
+		return RegisterWidgetFactory(PurposeRow, Constructor);
 	}
 }
 
-bool UEditorDataStorageUi::RegisterWidgetFactory(FName Purpose, TUniquePtr<FTypedElementWidgetConstructor>&& Constructor)
+bool UEditorDataStorageUi::RegisterWidgetFactory(UE::Editor::DataStorage::RowHandle PurposeRow,
+	TUniquePtr<FTypedElementWidgetConstructor>&& Constructor)
 {
 	checkf(Constructor->GetTypeInfo(), TEXT("Widget constructor being registered that doesn't have valid type information."));
-	
-	if (FPurposeInfo* PurposeInfo = WidgetPurposes.Find(Purpose))
+
+	if (FWidgetPurposeColumn* PurposeInfo = Storage->GetColumn<FWidgetPurposeColumn>(PurposeRow))
 	{
-		switch (PurposeInfo->Type)
+		switch (PurposeInfo->PurposeType)
 		{
-		case IUiProvider::EPurposeType::Generic:
-			PurposeInfo->Factories.Emplace(MoveTemp(Constructor));
-			PurposeInfo->bIsSorted = false;
-			return true;
-		case IUiProvider::EPurposeType::UniqueByName:
-			if (PurposeInfo->Factories.IsEmpty())
+		case EPurposeType::Generic:
+			{
+				UE::Editor::DataStorage::RowHandle FactoryRow = RegisterWidgetFactoryRow(PurposeRow);
+				Storage->AddColumn(FactoryRow, FWidgetFactoryConstructorColumn{.Constructor = MoveTemp(Constructor)});
+				return true;
+			}
+		case EPurposeType::UniqueByName:
 			{
-				PurposeInfo->Factories.Emplace(MoveTemp(Constructor));
-				PurposeInfo->bIsSorted = false;
+				UE::Editor::DataStorage::RowHandle FactoryRow = RegisterUniqueWidgetFactoryRow(PurposeRow);
+				Storage->AddColumn(FactoryRow, FWidgetFactoryConstructorColumn{.Constructor = MoveTemp(Constructor)});
+				return true;
 			}
-			else
+		case EPurposeType::UniqueByNameAndColumn:
 			{
-				PurposeInfo->Factories.EmplaceAt(0, MoveTemp(Constructor));
+				UE_LOG(LogEditorDataStorageUI, Warning,
+				TEXT("Unable to register widget factory '%s' as purpose '%llu' requires at least one column for matching."),
+				*Constructor->GetTypeInfo()->GetName(), PurposeRow);
+				return false;
 			}
-			return true;
-		case IUiProvider::EPurposeType::UniqueByNameAndColumn:
-			UE_LOG(LogEditorDataStorageUI, Warning,
-				TEXT("Unable to register widget factory '%s' as purpose '%s' requires at least one column for matching."),
-				*Constructor->GetTypeInfo()->GetName(), *Purpose.ToString());
-			return false;
 		default:
-			checkf(false, TEXT("Unexpected IUiProvider::EPurposeType found provided when registering widget factory."));
-			return false;
+			{
+				checkf(false, TEXT("Unexpected UE::Editor::DataStorage::IUiProvider::EPurposeType found provided when registering widget factory."));
+				return false;
+			}
 		}
 	}
-	else
-	{
-		UE_LOG(LogEditorDataStorageUI, Warning, 
-			TEXT("Unable to register widget factory as purpose '%s' isn't registered."), *Purpose.ToString());
-		return false;
-	}
+
+	UE_LOG(LogEditorDataStorageUI, Warning,
+		TEXT("Unable to register widget factory '%s' as purpose '%llu' isn't registered."), *Constructor->GetTypeInfo()->GetName(), PurposeRow);
+
+	return false;
 }
 
-bool UEditorDataStorageUi::RegisterWidgetFactory(FName Purpose, TUniquePtr<FTypedElementWidgetConstructor>&& Constructor, 
-	UE::Editor::DataStorage::Queries::FConditions Columns)
+bool UEditorDataStorageUi::RegisterWidgetFactory(UE::Editor::DataStorage::RowHandle PurposeRow,
+	TUniquePtr<FTypedElementWidgetConstructor>&& Constructor, UE::Editor::DataStorage::Queries::FConditions Columns)
 {
+	checkf(Constructor->GetTypeInfo(), TEXT("Widget constructor being registered that doesn't have valid type information."));
+
 	if (!Columns.IsEmpty())
 	{
-		checkf(Constructor->GetTypeInfo(), TEXT("Widget constructor being registered that doesn't have valid type information."));
-
-		if (FPurposeInfo* PurposeInfo = WidgetPurposes.Find(Purpose))
+		if (FWidgetPurposeColumn* PurposeInfo = Storage->GetColumn<FWidgetPurposeColumn>(PurposeRow))
 		{
-			switch (PurposeInfo->Type)
+			switch (PurposeInfo->PurposeType)
 			{
-			case IUiProvider::EPurposeType::Generic:
-				PurposeInfo->Factories.Emplace(MoveTemp(Constructor));
-				PurposeInfo->bIsSorted = false;
+			case EPurposeType::Generic:
+				{
+					UE::Editor::DataStorage::RowHandle FactoryRow = RegisterWidgetFactoryRow(PurposeRow);
+					Storage->AddColumn(FactoryRow, FWidgetFactoryConstructorColumn{.Constructor = MoveTemp(Constructor)});
+				}
 				return true;
-			case IUiProvider::EPurposeType::UniqueByName:
+			case EPurposeType::UniqueByName:
 				if (!Columns.IsEmpty())
 				{
-					if (PurposeInfo->Factories.IsEmpty())
-					{
-						PurposeInfo->Factories.Emplace(MoveTemp(Constructor));
-						PurposeInfo->bIsSorted = false;
-					}
-					else
-					{
-						PurposeInfo->Factories.EmplaceAt(0, MoveTemp(Constructor));
-					}
-					return true;
+					UE::Editor::DataStorage::RowHandle FactoryRow = RegisterUniqueWidgetFactoryRow(PurposeRow);
+					Storage->AddColumn(FactoryRow, FWidgetFactoryConstructorColumn{.Constructor = MoveTemp(Constructor)});
 				}
 				else
 				{
 					return false;
 				}
-			case IUiProvider::EPurposeType::UniqueByNameAndColumn:
-				if (!Columns.IsEmpty())
-				{
-					PurposeInfo->Factories.Emplace(MoveTemp(Constructor), MoveTemp(Columns));
-					PurposeInfo->bIsSorted = false;
-					return true;
-				}
-				else
+			case EPurposeType::UniqueByNameAndColumn:
 				{
-					return false;
+					if (!Columns.IsEmpty())
+					{
+						Columns.Compile(UE::Editor::DataStorage::Queries::FEditorStorageQueryConditionCompileContext(Storage));
+							
+						UE::Editor::DataStorage::RowHandle FactoryRow = RegisterWidgetFactoryRow(PurposeRow);
+						Storage->AddColumn(FactoryRow, FWidgetFactoryConstructorColumn{.Constructor = MoveTemp(Constructor)});
+						Storage->AddColumn(FactoryRow, FWidgetFactoryConditionsColumn{.Conditions = Columns});
+
+						return true;
+					}
+					else
+					{
+						return false;
+					}
 				}
+					
 			default:
-				checkf(false, TEXT("Unexpected IUiProvider::EPurposeType found provided when registering widget factory."));
+				checkf(false, TEXT("Unexpected UE::Editor::DataStorage::IUiProvider::EPurposeType found provided when registering widget factory."));
 				return false;
 			}
 		}
-		else
-		{
-			UE_LOG(LogEditorDataStorageUI, Warning, TEXT("Unable to register widget factory '%s' as purpose '%s' isn't registered."), 
-				*Constructor->GetTypeInfo()->GetName(), *Purpose.ToString());
-			return false;
-		}
+		
+		UE_LOG(LogEditorDataStorageUI, Warning,
+		TEXT("Unable to register widget factory '%s' as purpose '%llu' isn't registered."), *Constructor->GetTypeInfo()->GetName(), PurposeRow);
+
+		return false;
 	}
 	else
 	{
-		return RegisterWidgetFactory(Purpose, MoveTemp(Constructor));
+		return RegisterWidgetFactory(PurposeRow, MoveTemp(Constructor));
 	}
+	
 }
 
-void UEditorDataStorageUi::CreateWidgetConstructors(FName Purpose,
-	const UE::Editor::DataStorage::FMetaDataView& Arguments,
-	const UE::Editor::DataStorage::IUiProvider::WidgetConstructorCallback& Callback)
+void UEditorDataStorageUi::CreateWidgetConstructors(UE::Editor::DataStorage::RowHandle PurposeRow,
+	const UE::Editor::DataStorage::FMetaDataView& Arguments, const WidgetConstructorCallback& Callback)
 {
-	if (FPurposeInfo* PurposeInfo = WidgetPurposes.Find(Purpose))
+	while (Storage->HasColumns<FWidgetPurposeColumn>(PurposeRow))
 	{
-		for (const FWidgetFactory& Factory : PurposeInfo->Factories)
+		TArray<UE::Editor::DataStorage::RowHandle> Factories;
+		GetFactories(PurposeRow, Factories);
+
+		// If no factories were found for this purpose, move on to the parent purpose
+		if (Factories.IsEmpty())
+		{
+			if (FTableRowParentColumn* ParentColumn = Storage->GetColumn<FTableRowParentColumn>(PurposeRow))
+			{
+				PurposeRow = ParentColumn->Parent;
+			}
+			else
+			{
+				PurposeRow = UE::Editor::DataStorage::InvalidRowHandle;
+			}
+			continue;
+		}
+
+		for (UE::Editor::DataStorage::RowHandle FactoryRow : Factories)
 		{
-			if (!CreateSingleWidgetConstructor(Factory.Constructor, Arguments, {}, Factory.GetConditions(Storage), Callback))
+			if (!CreateSingleWidgetConstructor(FactoryRow, Arguments, {}, Callback))
 			{
 				return;
 			}
 		}
+
+		// Don't want to go up the parent chain if we created any widgets for this purpose
+		break;
 	}
 }
 
-void UEditorDataStorageUi::CreateWidgetConstructors(FName Purpose, EMatchApproach MatchApproach, 
+void UEditorDataStorageUi::CreateWidgetConstructors(UE::Editor::DataStorage::RowHandle PurposeRow, EMatchApproach MatchApproach,
 	TArray<TWeakObjectPtr<const UScriptStruct>>& Columns, const UE::Editor::DataStorage::FMetaDataView& Arguments,
-	const UE::Editor::DataStorage::IUiProvider::WidgetConstructorCallback& Callback)
+	const WidgetConstructorCallback& Callback)
 {
-	if (FPurposeInfo* PurposeInfo = WidgetPurposes.Find(Purpose))
+	// Sort by name so that removing the matched columns can be done in a single pass.
+	// Sorting by ptr does not work because dynamic column ptrs are different than their base template
+	Columns.Sort(
+		[](const TWeakObjectPtr<const UScriptStruct>& Lhs, const TWeakObjectPtr<const UScriptStruct>& Rhs)
+		{
+			return Lhs->GetName() < Rhs->GetName();
+		});
+	
+	while (Storage->HasColumns<FWidgetPurposeColumn>(PurposeRow))
 	{
-		// Sort by name so that removing the matched columns can be done in a single pass.
-		// Sorting by ptr does not work because dynamic column ptrs are different than their base template
-		Columns.Sort(
-			[](const TWeakObjectPtr<const UScriptStruct>& Lhs, const TWeakObjectPtr<const UScriptStruct>& Rhs)
-			{
-				return Lhs->GetName() < Rhs->GetName();
-			});
+		TArray<UE::Editor::DataStorage::RowHandle> Factories;
+		GetFactories(PurposeRow, Factories);
 
-		if (!PurposeInfo->bIsSorted)
+		if (!Factories.IsEmpty())
 		{
-			// This is the only call that requires the array of factories to be sorted from largest to smallest number
-			// of columns, so lazily sort only when needed.
-			PurposeInfo->Factories.StableSort(
-				[this](const FWidgetFactory& Lhs, const FWidgetFactory& Rhs)
+			// There is currently no way to cache the sorted results back into TEDS, so we sort every time this function is called
+			Factories.StableSort(
+				[this](UE::Editor::DataStorage::RowHandle Lhs, UE::Editor::DataStorage::RowHandle Rhs)
 				{
-					const UE::Editor::DataStorage::Queries::FConditions& LhsConditions = Lhs.GetConditions(Storage);
-					const UE::Editor::DataStorage::Queries::FConditions& RhsConditions = Rhs.GetConditions(Storage);
-					
+					const UE::Editor::DataStorage::Queries::FConditions& LhsConditions = GetFactoryConditions(Lhs);
+					const UE::Editor::DataStorage::Queries::FConditions& RhsConditions = GetFactoryConditions(Rhs);
+				
 					int32 LeftSize = LhsConditions.MinimumColumnMatchRequired();
 					int32 RightSize = RhsConditions.MinimumColumnMatchRequired();
 
@@ -332,126 +410,230 @@ void UEditorDataStorageUi::CreateWidgetConstructors(FName Purpose, EMatchApproac
 					{
 						return !LhsConditions.UsesDynamicTemplates();
 					}
-					
+				
 					return LeftSize > RightSize;
 				});
-			PurposeInfo->bIsSorted = true;
+
+			switch (MatchApproach)
+			{
+			case EMatchApproach::LongestMatch:
+
+				// For longest match, we don't want to continue matching with the parent purpose if the user requested us to stop
+				if (!CreateWidgetConstructors_LongestMatch(Factories, Columns, Arguments, Callback))
+				{
+					return;
+				}
+				break;
+			case EMatchApproach::ExactMatch:
+				CreateWidgetConstructors_ExactMatch(Factories, Columns, Arguments, Callback);
+				break;
+			case EMatchApproach::SingleMatch:
+				CreateWidgetConstructors_SingleMatch(Factories, Columns, Arguments, Callback);
+				break;
+			default:
+				checkf(false, TEXT("Unsupported match type (%i) for CreateWidgetConstructors."), 
+					static_cast<std::underlying_type_t<EMatchApproach>>(MatchApproach));
+			}
 		}
 
-		switch (MatchApproach)
+		// No need to go up the parent chain if there are no more columns to match
+		if (Columns.IsEmpty())
 		{
-		case EMatchApproach::LongestMatch:
-			CreateWidgetConstructors_LongestMatch(PurposeInfo->Factories, Columns, Arguments, Callback);
-			break;
-		case EMatchApproach::ExactMatch:
-			CreateWidgetConstructors_ExactMatch(PurposeInfo->Factories, Columns, Arguments, Callback);
-			break;
-		case EMatchApproach::SingleMatch:
-			CreateWidgetConstructors_SingleMatch(PurposeInfo->Factories, Columns, Arguments, Callback);
-			break;
-		default:
-			checkf(false, TEXT("Unsupported match type (%i) for CreateWidgetConstructors."), 
-				static_cast<std::underlying_type_t<EMatchApproach>>(MatchApproach));
+			return;
+		}
+
+		// If we have a parent purpose, try matching against factories belonging to it next
+		if (FTableRowParentColumn* ParentColumn = Storage->GetColumn<FTableRowParentColumn>(PurposeRow))
+		{
+			PurposeRow = ParentColumn->Parent;
+		}
+		else
+		{
+			PurposeRow = UE::Editor::DataStorage::InvalidRowHandle;
 		}
 	}
 }
 
-void UEditorDataStorageUi::ConstructWidgets(FName Purpose, const UE::Editor::DataStorage::FMetaDataView& Arguments,
-	const UE::Editor::DataStorage::IUiProvider::WidgetCreatedCallback& ConstructionCallback)
+void UEditorDataStorageUi::ConstructWidgets(UE::Editor::DataStorage::RowHandle PurposeRow, const UE::Editor::DataStorage::FMetaDataView& Arguments,
+	const WidgetCreatedCallback& ConstructionCallback)
 {
-	if (FPurposeInfo* PurposeInfo = WidgetPurposes.Find(Purpose))
+	// Find the first purpose in the parent chain with at least one registered factory
+	TArray<UE::Editor::DataStorage::RowHandle> Factories;
+
+	while (Storage->HasColumns<FWidgetPurposeColumn>(PurposeRow) && Factories.IsEmpty())
 	{
-		for (const FWidgetFactory& Factory : PurposeInfo->Factories)
+		GetFactories(PurposeRow, Factories);
+		
+		// If no factories were found for this purpose, move on to the parent purpose
+		if (FTableRowParentColumn* ParentColumn = Storage->GetColumn<FTableRowParentColumn>(PurposeRow))
 		{
-			std::visit(Internal::TOverloaded
+			PurposeRow = ParentColumn->Parent;
+		}
+		else
+		{
+			PurposeRow = UE::Editor::DataStorage::InvalidRowHandle;
+		}
+	}
+	
+	for (UE::Editor::DataStorage::RowHandle FactoryRow : Factories)
+	{
+		if (FWidgetFactoryConstructorColumn* ConstructorColumn = Storage->GetColumn<FWidgetFactoryConstructorColumn>(FactoryRow))
+		{
+			CreateWidgetInstance(*ConstructorColumn->Constructor, Arguments, ConstructionCallback);
+		}
+		else if (FWidgetFactoryConstructorTypeInfoColumn* ConstructorTypeInfoColumn = Storage->GetColumn<FWidgetFactoryConstructorTypeInfoColumn>(FactoryRow))
+		{
+			if (const UScriptStruct* ConstructorType = ConstructorTypeInfoColumn->Constructor.Get())
+			{
+				FTypedElementWidgetConstructor* Constructor = reinterpret_cast<FTypedElementWidgetConstructor*>(
+						FMemory::Malloc(ConstructorType->GetStructureSize(), ConstructorType->GetMinAlignment()));
+				if (Constructor)
 				{
-					[this, &Arguments, &ConstructionCallback](const UScriptStruct* ConstructorType)
-					{ 
-						FTypedElementWidgetConstructor* Constructor = reinterpret_cast<FTypedElementWidgetConstructor*>(
-							FMemory_Alloca_Aligned(ConstructorType->GetStructureSize(), ConstructorType->GetMinAlignment()));
-						if (Constructor)
-						{
-							ConstructorType->InitializeStruct(Constructor);
-							CreateWidgetInstance(*Constructor, Arguments, ConstructionCallback);
-							ConstructorType->DestroyStruct(&Constructor);
-						}
-						else
-						{
-							checkf(false, TEXT("Remaining stack space is too small to create a widget constructor from a description."));
-						}
-					},
-					[this, &Arguments, &ConstructionCallback](const TUniquePtr<FTypedElementWidgetConstructor>& Constructor)
-					{
-						CreateWidgetInstance(*Constructor, Arguments, ConstructionCallback);
-					}
-				}, Factory.Constructor);
+					ConstructorType->InitializeStruct(Constructor);
+					CreateWidgetInstance(*Constructor, Arguments, ConstructionCallback);
+					ConstructorType->DestroyStruct(&Constructor);
+				}
+				else
+				{
+					checkf(false, TEXT("Remaining memory is too small to create a widget constructor from a description."));
+				}
+
+				FMemory::Free(Constructor);
+			}
 		}
 	}
 }
 
+void UEditorDataStorageUi::RegisterWidgetPurpose(FName Purpose, EPurposeType Type, FText Description)
+{
+	RegisterWidgetPurpose(UE::Editor::DataStorage::FMapKey(Purpose), FPurposeInfo(Purpose, Type, Description));
+}
+
+bool UEditorDataStorageUi::RegisterWidgetFactory(FName Purpose, const UScriptStruct* Constructor)
+{
+	UE::Editor::DataStorage::RowHandle PurposeRow = Storage->LookupMappedRow(UE::Editor::DataStorage::FMapKey(Purpose));
+
+	if (Storage->IsRowAvailable(PurposeRow))
+	{
+		return RegisterWidgetFactory(PurposeRow, Constructor);
+	}
+	
+	UE_LOG(LogEditorDataStorageUI, Warning, 
+		TEXT("Unable to register widget factory '%s' as purpose '%s' isn't registered."), *Constructor->GetName(), *Purpose.ToString());
+
+	return false;
+}
+
+bool UEditorDataStorageUi::RegisterWidgetFactory(
+	FName Purpose, const UScriptStruct* Constructor, UE::Editor::DataStorage::Queries::FConditions Columns)
+{
+	UE::Editor::DataStorage::RowHandle PurposeRow = Storage->LookupMappedRow(UE::Editor::DataStorage::FMapKey(Purpose));
+	
+	if (Storage->IsRowAvailable(PurposeRow))
+	{
+		return RegisterWidgetFactory(PurposeRow, Constructor, Columns);
+	}
+	
+	UE_LOG(LogEditorDataStorageUI, Warning, 
+			TEXT("Unable to register widget factory '%s' as purpose '%s' isn't registered."), *Constructor->GetName(), *Purpose.ToString());
+
+	return false;
+}
+
+bool UEditorDataStorageUi::RegisterWidgetFactory(FName Purpose, TUniquePtr<FTypedElementWidgetConstructor>&& Constructor)
+{
+	UE::Editor::DataStorage::RowHandle PurposeRow = Storage->LookupMappedRow(UE::Editor::DataStorage::FMapKey(Purpose));
+	if (Storage->IsRowAvailable(PurposeRow))
+	{
+		return RegisterWidgetFactory(PurposeRow, MoveTemp(Constructor));
+	}
+	
+	UE_LOG(LogEditorDataStorageUI, Warning, 
+			TEXT("Unable to register widget factory as purpose '%s' isn't registered."), *Purpose.ToString());
+	return false;
+}
+
+bool UEditorDataStorageUi::RegisterWidgetFactory(FName Purpose, TUniquePtr<FTypedElementWidgetConstructor>&& Constructor, 
+	UE::Editor::DataStorage::Queries::FConditions Columns)
+{
+	UE::Editor::DataStorage::RowHandle PurposeRow = Storage->LookupMappedRow(UE::Editor::DataStorage::FMapKey(Purpose));
+	
+	if (Storage->IsRowAvailable(PurposeRow))
+	{
+		return RegisterWidgetFactory(PurposeRow, MoveTemp(Constructor), Columns);
+	}
+		
+	UE_LOG(LogEditorDataStorageUI, Warning, TEXT("Unable to register widget factory '%s' as purpose '%s' isn't registered."), 
+			*Constructor->GetTypeInfo()->GetName(), *Purpose.ToString());
+	return false;
+}
+
+void UEditorDataStorageUi::CreateWidgetConstructors(FName Purpose,
+	const UE::Editor::DataStorage::FMetaDataView& Arguments, const WidgetConstructorCallback& Callback)
+{
+	UE::Editor::DataStorage::RowHandle PurposeRow = Storage->LookupMappedRow(UE::Editor::DataStorage::FMapKey(Purpose));
+	CreateWidgetConstructors(PurposeRow, Arguments, Callback);
+}
+
+void UEditorDataStorageUi::CreateWidgetConstructors(FName Purpose, EMatchApproach MatchApproach, 
+	TArray<TWeakObjectPtr<const UScriptStruct>>& Columns, const UE::Editor::DataStorage::FMetaDataView& Arguments,
+	const WidgetConstructorCallback& Callback)
+{
+	UE::Editor::DataStorage::RowHandle PurposeRow = Storage->LookupMappedRow(UE::Editor::DataStorage::FMapKey(Purpose));
+	CreateWidgetConstructors(PurposeRow, MatchApproach, Columns, Arguments, Callback);
+}
+
+void UEditorDataStorageUi::ConstructWidgets(FName Purpose, const UE::Editor::DataStorage::FMetaDataView& Arguments,
+	const WidgetCreatedCallback& ConstructionCallback)
+{
+	UE::Editor::DataStorage::RowHandle PurposeRow = Storage->LookupMappedRow(UE::Editor::DataStorage::FMapKey(Purpose));
+	ConstructWidgets(PurposeRow, Arguments, ConstructionCallback);
+}
+
 bool UEditorDataStorageUi::CreateSingleWidgetConstructor(
-	const FWidgetFactory::ConstructorType& Constructor,
+	UE::Editor::DataStorage::RowHandle FactoryRow,
 	const UE::Editor::DataStorage::FMetaDataView& Arguments,
 	TArray<TWeakObjectPtr<const UScriptStruct>> MatchedColumnTypes,
-	const UE::Editor::DataStorage::Queries::FConditions& QueryConditions,
-	const UE::Editor::DataStorage::IUiProvider::WidgetConstructorCallback& Callback)
+	const WidgetConstructorCallback& Callback)
 {
-	using namespace UE::Editor::DataStorage;
-	struct Visitor
-	{
-		Visitor(
-			TArray<TWeakObjectPtr<const UScriptStruct>>&& InMatchedColumnTypes,
-			const Queries::FConditions& InQueryConditions,
-			const FMetaDataView& InArguments,
-			const WidgetConstructorCallback& InCallback) 
-			: MatchedColumnTypes(MoveTemp(InMatchedColumnTypes))
-			, QueryConditions(InQueryConditions)
-			, Arguments(InArguments)
-			, Callback(InCallback)
-		{}
-
-		TArray<TWeakObjectPtr<const UScriptStruct>> MatchedColumnTypes;
-		const Queries::FConditions& QueryConditions;
-		const FMetaDataView& Arguments;
-		const WidgetConstructorCallback& Callback;
-
-		bool operator()(const UScriptStruct* Target)
+	if (FWidgetFactoryConstructorColumn* ConstructorColumn = Storage->GetColumn<FWidgetFactoryConstructorColumn>(FactoryRow))
+	{
+		const UScriptStruct* TargetType = ConstructorColumn->Constructor->GetTypeInfo();
+		checkf(TargetType, TEXT("Expected valid type information from a widget constructor."));
+		TUniquePtr<FTypedElementWidgetConstructor> Result(reinterpret_cast<FTypedElementWidgetConstructor*>(
+			FMemory::Malloc(TargetType->GetStructureSize(), TargetType->GetMinAlignment())));
+		if (Result)
+		{
+			TargetType->InitializeStruct(Result.Get());
+			TargetType->CopyScriptStruct(Result.Get(), ConstructorColumn->Constructor.Get());
+			Result->Initialize(Arguments, MoveTemp(MatchedColumnTypes), FactoryRow);
+			return Callback(MoveTemp(Result), Result->GetMatchedColumns());
+		}
+		return true;
+	}
+	else if (FWidgetFactoryConstructorTypeInfoColumn* ConstructorTypeInfoColumn = Storage->GetColumn<FWidgetFactoryConstructorTypeInfoColumn>(FactoryRow))
+	{
+		if (const UScriptStruct* Target = ConstructorTypeInfoColumn->Constructor.Get())
 		{
 			TUniquePtr<FTypedElementWidgetConstructor> Result(reinterpret_cast<FTypedElementWidgetConstructor*>(
 				FMemory::Malloc(Target->GetStructureSize(), Target->GetMinAlignment())));
 			if (Result)
 			{
 				Target->InitializeStruct(Result.Get());
-				Result->Initialize(Arguments, MoveTemp(MatchedColumnTypes), QueryConditions);
+				Result->Initialize(Arguments, MoveTemp(MatchedColumnTypes), FactoryRow);
 				const TArray<TWeakObjectPtr<const UScriptStruct>>& MatchedColumns = Result->GetMatchedColumns();
 				return Callback(MoveTemp(Result), MatchedColumns);
 			}
 			return true;
 		}
-
-		bool operator()(const TUniquePtr<FTypedElementWidgetConstructor>& Target)
-		{
-			const UScriptStruct* TargetType = Target->GetTypeInfo();
-			checkf(TargetType, TEXT("Expected valid type information from a widget constructor."));
-			TUniquePtr<FTypedElementWidgetConstructor> Result(reinterpret_cast<FTypedElementWidgetConstructor*>(
-				FMemory::Malloc(TargetType->GetStructureSize(), TargetType->GetMinAlignment())));
-			if (Result)
-			{
-				TargetType->InitializeStruct(Result.Get());
-				TargetType->CopyScriptStruct(Result.Get(), Target.Get());
-				Result->Initialize(Arguments, MoveTemp(MatchedColumnTypes), QueryConditions);
-				return Callback(MoveTemp(Result), Result->GetMatchedColumns());
-			}
-			return true;
-		}
-	};
-	return std::visit(Visitor(MoveTemp(MatchedColumnTypes), QueryConditions, Arguments, Callback), Constructor);
+	}
+	return false;
 }
 
 void UEditorDataStorageUi::CreateWidgetInstance(
 	FTypedElementWidgetConstructor& Constructor, 
 	const UE::Editor::DataStorage::FMetaDataView& Arguments,
-	const UE::Editor::DataStorage::IUiProvider::WidgetCreatedCallback& ConstructionCallback)
+	const WidgetCreatedCallback& ConstructionCallback)
 {
 	UE::Editor::DataStorage::RowHandle Row = Storage->AddRow(WidgetTable);
 	Storage->AddColumns(Row, Constructor.GetAdditionalColumnsList());
@@ -474,10 +656,16 @@ TSharedPtr<SWidget> UEditorDataStorageUi::ConstructWidget(UE::Editor::DataStorag
 
 void UEditorDataStorageUi::ListWidgetPurposes(const WidgetPurposeCallback& Callback) const
 {
-	for (auto&& It : WidgetPurposes)
+	using namespace UE::Editor::DataStorage::Queries;
+	
+	static QueryHandle PurposeQueryHandle =
+		Storage->RegisterQuery(Select().ReadOnly<FNameColumn, FWidgetPurposeColumn, FDescriptionColumn>().Compile());
+
+	Storage->RunQuery(PurposeQueryHandle, CreateDirectQueryCallbackBinding(
+		[Callback](const FNameColumn& NameColumn, const FWidgetPurposeColumn& PurposeColumn, const FDescriptionColumn& DescriptionColumn)
 	{
-		Callback(It.Key, It.Value.Type, It.Value.Description);
-	}
+			Callback(NameColumn.Name, PurposeColumn.PurposeType, DescriptionColumn.Description);
+	}));
 }
 
 bool UEditorDataStorageUi::SupportsExtension(FName Extension) const
@@ -501,6 +689,21 @@ UE::Editor::DataStorage::TableHandle UEditorDataStorageUi::GetWidgetTable() cons
 	return WidgetTable;
 }
 
+UE::Editor::DataStorage::IUiProvider::FPurposeID UEditorDataStorageUi::GetDefaultWidgetPurposeID() const
+{
+	return Internal::DefaultWidgetPurposeID;	
+}
+
+UE::Editor::DataStorage::IUiProvider::FPurposeID UEditorDataStorageUi::GetGeneralWidgetPurposeID() const
+{
+	return Internal::GeneralWidgetPurposeID;
+}
+
+UE::Editor::DataStorage::RowHandle UEditorDataStorageUi::FindPurpose(const FPurposeID& PurposeID) const
+{
+	return Storage->LookupMappedRow(UE::Editor::DataStorage::FMapKey(PurposeID));
+}
+
 void UEditorDataStorageUi::CreateStandardArchetypes()
 {
 	WidgetTable = Storage->RegisterTable(MakeArrayView(
@@ -509,16 +712,48 @@ void UEditorDataStorageUi::CreateStandardArchetypes()
 			FTypedElementSlateWidgetReferenceDeletesRowTag::StaticStruct(),
 			FSlateColorColumn::StaticStruct()
 		}), FName(TEXT("Editor_WidgetTable")));
+
+	WidgetPurposeTable = Storage->RegisterTable(UE::Editor::DataStorage::TTypedElementColumnTypeList<FWidgetPurposeColumn, FWidgetPurposeNameColumn>(), FName("Editor_WidgetPurposeTable"));
+	WidgetFactoryTable = Storage->RegisterTable(UE::Editor::DataStorage::TTypedElementColumnTypeList<FWidgetFactoryColumn>(), FName("Editor_WidgetFactoryTable"));
 }
 
-void UEditorDataStorageUi::CreateWidgetConstructors_LongestMatch(const TArray<FWidgetFactory>& WidgetFactories,
+void UEditorDataStorageUi::RegisterQueries()
+{
+	using namespace UE::Editor::DataStorage::Queries;
+	
+	Storage->RegisterQuery(
+		Select(
+			TEXT("Add display name to widget factory with constructor"),
+			FObserver::OnAdd<FWidgetFactoryConstructorTypeInfoColumn>().SetExecutionMode(EExecutionMode::GameThread),
+			[](IQueryContext& Context, RowHandle Row, const FWidgetFactoryConstructorTypeInfoColumn& ConstructorColumn)
+			{
+				Context.AddColumn(Row, FNameColumn{ .Name = ConstructorColumn.Constructor->GetFName() });
+			})
+		.Where()
+			.All<FWidgetFactoryColumn>()
+		.Compile());
+	
+	Storage->RegisterQuery(
+		Select(
+			TEXT("Add display name to widget factory with constructor type info"),
+			FObserver::OnAdd<FWidgetFactoryConstructorColumn>().SetExecutionMode(EExecutionMode::GameThread),
+			[this](IQueryContext& Context, RowHandle Row, const FWidgetFactoryConstructorColumn& ConstructorColumn)
+			{
+				Context.AddColumn(Row, FNameColumn{ .Name = ConstructorColumn.Constructor->GetTypeInfo()->GetFName() });
+			})
+		.Where()
+			.All<FWidgetFactoryColumn>()
+		.Compile());
+}
+
+bool UEditorDataStorageUi::CreateWidgetConstructors_LongestMatch(const TArray<UE::Editor::DataStorage::RowHandle>& WidgetFactories,
 	TArray<TWeakObjectPtr<const UScriptStruct>>& Columns, const UE::Editor::DataStorage::FMetaDataView& Arguments,
 	const WidgetConstructorCallback& Callback)
 {
 	TArray<TWeakObjectPtr<const UScriptStruct>> MatchedColumns;
 	for (auto FactoryIt = WidgetFactories.CreateConstIterator(); FactoryIt && !Columns.IsEmpty(); ++FactoryIt)
 	{
-		const UE::Editor::DataStorage::Queries::FConditions& Conditions = FactoryIt->GetConditions(Storage);
+		const UE::Editor::DataStorage::Queries::FConditions& Conditions = GetFactoryConditions(*FactoryIt);
 		
 		if (Conditions.MinimumColumnMatchRequired() > Columns.Num())
 		{
@@ -531,6 +766,11 @@ void UEditorDataStorageUi::CreateWidgetConstructors_LongestMatch(const TArray<FW
 		
 		if (Conditions.Verify(MatchedColumns, Columns))
 		{
+			// Empty conditions match against everything - so we update the matched columns list to reflect that
+			if (Conditions.IsEmpty())
+			{
+				MatchedColumns = Columns;
+			}
 			// Remove the found columns from the requested list.
 			MatchedColumns.Sort(
 				[](const TWeakObjectPtr<const UScriptStruct>& Lhs, const TWeakObjectPtr<const UScriptStruct>& Rhs)
@@ -557,7 +797,7 @@ void UEditorDataStorageUi::CreateWidgetConstructors_LongestMatch(const TArray<FW
 					if (ColumnsIt == ColumnsEnd)
 					{
 						ensureMsgf(false, TEXT("A previously found matching column can't be found in the original array."));
-						return;
+						return false;
 					}
 				}
 				RequestedColumnsThatMatched.Add(*ColumnsIt);
@@ -565,23 +805,25 @@ void UEditorDataStorageUi::CreateWidgetConstructors_LongestMatch(const TArray<FW
 				--ColumnsEnd;
 			}
 			
-			if (!CreateSingleWidgetConstructor(FactoryIt->Constructor, Arguments, MoveTemp(RequestedColumnsThatMatched), Conditions, Callback))
+			if (!CreateSingleWidgetConstructor(*FactoryIt, Arguments, MoveTemp(RequestedColumnsThatMatched), Callback))
 			{
-				return;
+				return false;
 			}
 		}
 	}
+
+	return true;
 }
 
-void UEditorDataStorageUi::CreateWidgetConstructors_ExactMatch(const TArray<FWidgetFactory>& WidgetFactories,
+void UEditorDataStorageUi::CreateWidgetConstructors_ExactMatch(const TArray<UE::Editor::DataStorage::RowHandle>& WidgetFactories,
 	TArray<TWeakObjectPtr<const UScriptStruct>>& Columns, const UE::Editor::DataStorage::FMetaDataView& Arguments,
-	const UE::Editor::DataStorage::IUiProvider::WidgetConstructorCallback& Callback)
+	const WidgetConstructorCallback& Callback)
 {
 	int32 ColumnCount = Columns.Num();
 	TArray<TWeakObjectPtr<const UScriptStruct>> MatchedColumns;
-	for (const FWidgetFactory& Factory : WidgetFactories)
+	for (UE::Editor::DataStorage::RowHandle FactoryRow : WidgetFactories)
 	{
-		const UE::Editor::DataStorage::Queries::FConditions& Conditions = Factory.GetConditions(Storage);
+		const UE::Editor::DataStorage::Queries::FConditions& Conditions = GetFactoryConditions(FactoryRow);
 		
 		// If there are more matches required that there are columns, then there will never be an exact match.
 		// Less than the column count can still result in a match that covers all columns.
@@ -594,21 +836,27 @@ void UEditorDataStorageUi::CreateWidgetConstructors_ExactMatch(const TArray<FWid
 
 		if (Conditions.Verify(MatchedColumns, Columns))
 		{
+			// Empty conditions match against everything - so we update the matched columns list to reflect that
+			if (Conditions.IsEmpty())
+			{
+				MatchedColumns = Columns;
+			}
+			
 			Algo::SortBy(MatchedColumns, [](const TWeakObjectPtr<const UScriptStruct>& Column) { return Column.Get(); });
 			MatchedColumns.SetNum(Algo::Unique(MatchedColumns), EAllowShrinking::No);
 			if (MatchedColumns.Num() == Columns.Num())
 			{
 				Columns.Reset();
-				CreateSingleWidgetConstructor(Factory.Constructor, Arguments, MoveTemp(MatchedColumns), Conditions, Callback);
+				CreateSingleWidgetConstructor(FactoryRow, Arguments, MoveTemp(MatchedColumns), Callback);
 				return;
 			}
 		}
 	}
 }
 
-void UEditorDataStorageUi::CreateWidgetConstructors_SingleMatch(const TArray<FWidgetFactory>& WidgetFactories,
+void UEditorDataStorageUi::CreateWidgetConstructors_SingleMatch(const TArray<UE::Editor::DataStorage::RowHandle>& WidgetFactories,
 	TArray<TWeakObjectPtr<const UScriptStruct>>& Columns, const UE::Editor::DataStorage::FMetaDataView& Arguments,
-	const UE::Editor::DataStorage::IUiProvider::WidgetConstructorCallback& Callback)
+	const WidgetConstructorCallback& Callback)
 {
 	auto FactoryIt = WidgetFactories.rbegin();
 	auto FactoryEnd = WidgetFactories.rend();
@@ -618,7 +866,7 @@ void UEditorDataStorageUi::CreateWidgetConstructors_SingleMatch(const TArray<FWi
 	{
 		for (; FactoryIt != FactoryEnd; ++FactoryIt)
 		{
-			const UE::Editor::DataStorage::Queries::FConditions& Conditions = (*FactoryIt).GetConditions(Storage);
+			const UE::Editor::DataStorage::Queries::FConditions& Conditions = GetFactoryConditions(*FactoryIt);
 			
 			TConstArrayView<TWeakObjectPtr<const UScriptStruct>> ColumnData = Conditions.GetColumns();
 			if (ColumnData.Num() > 1)
@@ -638,7 +886,7 @@ void UEditorDataStorageUi::CreateWidgetConstructors_SingleMatch(const TArray<FWi
 				// dynamic column and the widget constructor wants the latter
 				TWeakObjectPtr<const UScriptStruct> RequestedColumn = Columns[ColumnIndex];
 				Columns.RemoveAt(ColumnIndex);
-				CreateSingleWidgetConstructor((*FactoryIt).Constructor, Arguments, {RequestedColumn}, Conditions, Callback);
+				CreateSingleWidgetConstructor(*FactoryIt, Arguments, {RequestedColumn}, Callback);
 				
 				
 				// Match was found so move on to the next column in the column.
@@ -648,41 +896,80 @@ void UEditorDataStorageUi::CreateWidgetConstructors_SingleMatch(const TArray<FWi
 	}
 }
 
-//
-// FWidgetFactory
-//
-
-UEditorDataStorageUi::FWidgetFactory::FWidgetFactory(const UScriptStruct* InConstructor)
-	: Constructor(InConstructor)
+UE::Editor::DataStorage::RowHandle UEditorDataStorageUi::RegisterWidgetFactoryRow(UE::Editor::DataStorage::RowHandle PurposeRowHandle) const
 {
-}
+	UE::Editor::DataStorage::RowHandle FactoryRowHandle = Storage->AddRow(WidgetFactoryTable);
+	Storage->GetColumn<FWidgetFactoryColumn>(FactoryRowHandle)->PurposeRowHandle = PurposeRowHandle;
 
-UEditorDataStorageUi::FWidgetFactory::FWidgetFactory(TUniquePtr<FTypedElementWidgetConstructor>&& InConstructor)
-	: Constructor(MoveTemp(InConstructor))
-{
-	checkf(std::get<TUniquePtr<FTypedElementWidgetConstructor>>(Constructor)->GetTypeInfo(), 
-		TEXT("Widget constructor registered that didn't contain valid type information."));
+	return FactoryRowHandle;
 }
 
-UEditorDataStorageUi::FWidgetFactory::FWidgetFactory(const UScriptStruct* InConstructor, 
-	UE::Editor::DataStorage::Queries::FConditions&& InColumns)
-	: Constructor(InConstructor)
-	, Columns(MoveTemp(InColumns))
+UE::Editor::DataStorage::RowHandle UEditorDataStorageUi::RegisterUniqueWidgetFactoryRow(UE::Editor::DataStorage::RowHandle InPurposeRowHandle) const
 {
+	using namespace UE::Editor::DataStorage::Queries;
+	
+	static QueryHandle FactoryQueryHandle = Storage->RegisterQuery(Select().ReadOnly<FWidgetFactoryColumn>().Compile());
+
+	RowHandle FactoryRowHandle = InvalidRowHandle;
+
+	// Find the first matching factory belonging to this purpose, we know there is only going to be one
+	Storage->RunQuery(FactoryQueryHandle, CreateDirectQueryCallbackBinding(
+		[&FactoryRowHandle, InPurposeRowHandle](RowHandle FoundFactoryRowHandle, const FWidgetFactoryColumn& PurposeReferenceColumn)
+		{
+			if (PurposeReferenceColumn.PurposeRowHandle == InPurposeRowHandle)
+			{
+				FactoryRowHandle = FoundFactoryRowHandle;
+			}
+		}));
+
+	// If there was a factory already registered for this purpose, overwrite its information
+	if (Storage->IsRowAvailable(FactoryRowHandle))
+	{
+		Storage->RemoveColumns<FWidgetFactoryConstructorColumn, FWidgetFactoryConstructorTypeInfoColumn>(FactoryRowHandle);
+		Storage->GetColumn<FWidgetFactoryColumn>(FactoryRowHandle)->PurposeRowHandle = InPurposeRowHandle;
+		return FactoryRowHandle;
+	}
+	// Otherwise just register the factory row as usual
+	else
+	{
+		return RegisterWidgetFactoryRow(InPurposeRowHandle);
+	}
 }
 
-UEditorDataStorageUi::FWidgetFactory::FWidgetFactory(TUniquePtr<FTypedElementWidgetConstructor>&& InConstructor, 
-	UE::Editor::DataStorage::Queries::FConditions&& InColumns)
-	: Constructor(MoveTemp(InConstructor))
-	, Columns(MoveTemp(InColumns))
+void UEditorDataStorageUi::GetFactories(UE::Editor::DataStorage::RowHandle PurposeRowHandle,
+	TArray<UE::Editor::DataStorage::RowHandle>& OutFactories) const
 {
-	checkf(std::get<TUniquePtr<FTypedElementWidgetConstructor>>(Constructor)->GetTypeInfo(),
-		TEXT("Widget constructor registered that didn't contain valid type information."));
+	using namespace UE::Editor::DataStorage::Queries;
+	
+	static QueryHandle FactoryQueryHandle = Storage->RegisterQuery(Select().ReadOnly<FWidgetFactoryColumn>().Compile());
+	
+	Storage->RunQuery(FactoryQueryHandle, CreateDirectQueryCallbackBinding(
+		[PurposeRowHandle, &OutFactories](RowHandle RowHandle, const FWidgetFactoryColumn& PurposeReferenceColumn)
+		{
+			if (PurposeReferenceColumn.PurposeRowHandle == PurposeRowHandle)
+			{
+				OutFactories.Add(RowHandle);
+			}
+		}));
 }
 
-const UE::Editor::DataStorage::Queries::FConditions& UEditorDataStorageUi::FWidgetFactory::GetConditions(UE::Editor::DataStorage::ICoreProvider* DataStorage) const
+const UE::Editor::DataStorage::Queries::FConditions& UEditorDataStorageUi::GetFactoryConditions(UE::Editor::DataStorage::RowHandle FactoryRow) const
 {
-	Columns.Compile(UE::Editor::DataStorage::Queries::FEditorStorageQueryConditionCompileContext(DataStorage));
+	using namespace UE::Editor::DataStorage::Queries;
+
+	if (FWidgetFactoryConditionsColumn* FactoryColumn = Storage->GetColumn<FWidgetFactoryConditionsColumn>(FactoryRow))
+	{
+		FactoryColumn->Conditions.Compile(FEditorStorageQueryConditionCompileContext(Storage));
+		return FactoryColumn->Conditions;
+	}
+	// If this factory does not have any query conditions, just return a default empty FConditions struct
+	else
+	{
+		static FConditions DefaultConditions;
+		DefaultConditions.Compile(FEditorStorageQueryConditionCompileContext(Storage));
+		return DefaultConditions;
+	}
 	
-	return Columns;
 }
+
+#undef LOCTEXT_NAMESPACE // "TypedElementDatabaseUI"
diff --git a/Engine/Plugins/Experimental/EditorDataStorage/Source/TedsCore/Private/TypedElementDatabaseUI.h b/Engine/Plugins/Experimental/EditorDataStorage/Source/TedsCore/Private/TypedElementDatabaseUI.h
--- a/Engine/Plugins/Experimental/EditorDataStorage/Source/TedsCore/Private/TypedElementDatabaseUI.h
+++ b/Engine/Plugins/Experimental/EditorDataStorage/Source/TedsCore/Private/TypedElementDatabaseUI.h
@@ -2,7 +2,6 @@
 
 #pragma once
 
-#include <variant>
 #include "Containers/Map.h"
 #include "Elements/Common/TypedElementQueryConditions.h"
 #include "Elements/Interfaces/TypedElementDataStorageUiInterface.h"
@@ -17,6 +16,7 @@ namespace UE::Editor::DataStorage
 	class ICompatibilityProvider;
 } // namespace UE::Editor::DataStorage
 
+
 TEDSCORE_API DECLARE_LOG_CATEGORY_EXTERN(LogEditorDataStorageUI, Log, All);
 
 UCLASS()
@@ -34,24 +34,40 @@ public:
 		UE::Editor::DataStorage::ICompatibilityProvider* StorageCompatibilityInterface);
 	void Deinitialize();
 
-	void RegisterWidgetPurpose(FName Purpose, EPurposeType Type, FText Description) override;
+	virtual UE::Editor::DataStorage::RowHandle RegisterWidgetPurpose(const FPurposeID& PurposeID, const FPurposeInfo& InPurposeInfo) override;
+	virtual UE::Editor::DataStorage::RowHandle RegisterWidgetPurpose(const FPurposeInfo& InPurposeInfo) override;
 
+	virtual bool RegisterWidgetFactory(UE::Editor::DataStorage::RowHandle PurposeRow, const UScriptStruct* Constructor) override;
+	virtual bool RegisterWidgetFactory(UE::Editor::DataStorage::RowHandle PurposeRow, const UScriptStruct* Constructor,
+		UE::Editor::DataStorage::Queries::FConditions Columns) override;
+	virtual bool RegisterWidgetFactory(UE::Editor::DataStorage::RowHandle PurposeRow,
+		TUniquePtr<FTypedElementWidgetConstructor>&& Constructor) override;
+	virtual bool RegisterWidgetFactory(UE::Editor::DataStorage::RowHandle PurposeRow, TUniquePtr<FTypedElementWidgetConstructor>&& Constructor, 
+		UE::Editor::DataStorage::Queries::FConditions Columns) override;
+	
+	virtual void CreateWidgetConstructors(UE::Editor::DataStorage::RowHandle PurposeRow, 
+		const UE::Editor::DataStorage::FMetaDataView& Arguments, const WidgetConstructorCallback& Callback) override;
+	virtual void CreateWidgetConstructors(UE::Editor::DataStorage::RowHandle PurposeRow, EMatchApproach MatchApproach, TArray<TWeakObjectPtr<const UScriptStruct>>& Columns,
+		const UE::Editor::DataStorage::FMetaDataView& Arguments, const WidgetConstructorCallback& Callback) override;
+	virtual void ConstructWidgets(UE::Editor::DataStorage::RowHandle PurposeRow, const UE::Editor::DataStorage::FMetaDataView& Arguments,
+			const WidgetCreatedCallback& ConstructionCallback) override;
+
+	// Deprecated FName overloads
+	void RegisterWidgetPurpose(FName Purpose, EPurposeType Type, FText Description) override;
 	bool RegisterWidgetFactory(FName Purpose, const UScriptStruct* Constructor) override;
 	bool RegisterWidgetFactory(FName Purpose, const UScriptStruct* Constructor,
 		UE::Editor::DataStorage::Queries::FConditions Columns) override;
 	bool RegisterWidgetFactory(FName Purpose, TUniquePtr<FTypedElementWidgetConstructor>&& Constructor) override;
 	bool RegisterWidgetFactory(FName Purpose, TUniquePtr<FTypedElementWidgetConstructor>&& Constructor,
 		UE::Editor::DataStorage::Queries::FConditions Columns) override;
-
 	void CreateWidgetConstructors(FName Purpose,
-		const UE::Editor::DataStorage::FMetaDataView& Arguments,
-		const UE::Editor::DataStorage::IUiProvider::WidgetConstructorCallback& Callback) override;
+		const UE::Editor::DataStorage::FMetaDataView& Arguments, const WidgetConstructorCallback& Callback) override;
 	void CreateWidgetConstructors(FName Purpose, EMatchApproach MatchApproach, TArray<TWeakObjectPtr<const UScriptStruct>>& Columns,
-		const UE::Editor::DataStorage::FMetaDataView& Arguments,
-		const UE::Editor::DataStorage::IUiProvider::WidgetConstructorCallback& Callback) override;
-
+		const UE::Editor::DataStorage::FMetaDataView& Arguments, const WidgetConstructorCallback& Callback) override;
 	void ConstructWidgets(FName Purpose, const UE::Editor::DataStorage::FMetaDataView& Arguments,
-		const UE::Editor::DataStorage::IUiProvider::WidgetCreatedCallback& ConstructionCallback) override;
+		const WidgetCreatedCallback& ConstructionCallback) override;
+	// ~FName overloads
+	
 	TSharedPtr<SWidget> ConstructWidget(UE::Editor::DataStorage::RowHandle Row, FTypedElementWidgetConstructor& Constructor,
 		const UE::Editor::DataStorage::FMetaDataView& Arguments) override;
 
@@ -65,67 +81,62 @@ public:
 	
 	/** Get the table where TEDS UI widgets are stored */
 	virtual UE::Editor::DataStorage::TableHandle GetWidgetTable() const override;
+	
+	/** Get the ID of the default TEDS UI widget purpose used to register default widgets for different types of data (e.g FText -> STextBlock) */
+	virtual FPurposeID GetDefaultWidgetPurposeID() const;
+	
+	/** Get the ID of the general TEDS UI purpose used to register general purpose widgets for columns */
+	virtual FPurposeID GetGeneralWidgetPurposeID() const;
+	
+	/** Find the row handle for a purpose by looking it up using the purpose ID */
+	virtual UE::Editor::DataStorage::RowHandle FindPurpose(const FPurposeID& PurposeID) const override;
 
 private:
-	struct FWidgetFactory
-	{
-	public:
-		using ConstructorType = std::variant<const UScriptStruct*, TUniquePtr<FTypedElementWidgetConstructor>>;
-		ConstructorType Constructor;
-
-		FWidgetFactory() = default;
-		explicit FWidgetFactory(const UScriptStruct* InConstructor);
-		explicit FWidgetFactory(TUniquePtr<FTypedElementWidgetConstructor>&& InConstructor);
-		FWidgetFactory(const UScriptStruct* InConstructor, UE::Editor::DataStorage::Queries::FConditions&& InColumns);
-		FWidgetFactory(TUniquePtr<FTypedElementWidgetConstructor>&& InConstructor, UE::Editor::DataStorage::Queries::FConditions&& InColumns);
-		const UE::Editor::DataStorage::Queries::FConditions& GetConditions(UE::Editor::DataStorage::ICoreProvider* DataStorage) const;
-
-	private:
-		// Private and mutable so we can control access to it via GetConditions() to compile the conditions on demand
-		mutable UE::Editor::DataStorage::Queries::FConditions Columns;
-	};
-
-	struct FPurposeInfo
-	{	
-		TArray<FWidgetFactory> Factories;
-		FText Description;
-		EPurposeType Type;
-		bool bIsSorted{ false }; //< Whether or not the array of factories needs to be sorted. The factories themselves are already sorted.
-	};
-
 	void CreateStandardArchetypes();
+	void RegisterQueries();
 
 	bool CreateSingleWidgetConstructor(
-		const FWidgetFactory::ConstructorType& Constructor,
+		UE::Editor::DataStorage::RowHandle FactoryRow,
 		const UE::Editor::DataStorage::FMetaDataView& Arguments,
 		TArray<TWeakObjectPtr<const UScriptStruct>> MatchedColumnTypes,
-		const UE::Editor::DataStorage::Queries::FConditions& QueryConditions,
-		const UE::Editor::DataStorage::IUiProvider::WidgetConstructorCallback& Callback);
+		const WidgetConstructorCallback& Callback);
 
 	void CreateWidgetInstance(
 		FTypedElementWidgetConstructor& Constructor,
 		const UE::Editor::DataStorage::FMetaDataView& Arguments,
-		const UE::Editor::DataStorage::IUiProvider::WidgetCreatedCallback& ConstructionCallback);
+		const WidgetCreatedCallback& ConstructionCallback);
 
-	void CreateWidgetConstructors_LongestMatch(
-		const TArray<FWidgetFactory>& WidgetFactories, 
+	bool CreateWidgetConstructors_LongestMatch(
+		const TArray<UE::Editor::DataStorage::RowHandle>& WidgetFactories, 
 		TArray<TWeakObjectPtr<const UScriptStruct>>& Columns, 
 		const UE::Editor::DataStorage::FMetaDataView& Arguments,
-		const UE::Editor::DataStorage::IUiProvider::WidgetConstructorCallback& Callback);
+		const WidgetConstructorCallback& Callback);
 	void CreateWidgetConstructors_ExactMatch(
-		const TArray<FWidgetFactory>& WidgetFactories,
+		const TArray<UE::Editor::DataStorage::RowHandle>& WidgetFactories,
 		TArray<TWeakObjectPtr<const UScriptStruct>>& Columns,
 		const UE::Editor::DataStorage::FMetaDataView& Arguments,
-		const UE::Editor::DataStorage::IUiProvider::WidgetConstructorCallback& Callback);
+		const WidgetConstructorCallback& Callback);
 	void CreateWidgetConstructors_SingleMatch(
-		const TArray<FWidgetFactory>& WidgetFactories,
+		const TArray<UE::Editor::DataStorage::RowHandle>& WidgetFactories,
 		TArray<TWeakObjectPtr<const UScriptStruct>>& Columns,
 		const UE::Editor::DataStorage::FMetaDataView& Arguments,
-		const UE::Editor::DataStorage::IUiProvider::WidgetConstructorCallback& Callback);
+		const WidgetConstructorCallback& Callback);
+
+	// Register a widget factory for the provided purpose
+	UE::Editor::DataStorage::RowHandle RegisterWidgetFactoryRow(UE::Editor::DataStorage::RowHandle PurposeRowHandle) const;
+
+	// Register a unique factory for the provided purpose, clearing the info if there was any factory previously registered for the purpose
+	UE::Editor::DataStorage::RowHandle RegisterUniqueWidgetFactoryRow(UE::Editor::DataStorage::RowHandle PurposeRowHandle) const;
+
+	// Get all factories for the provided purpose
+	void GetFactories(UE::Editor::DataStorage::RowHandle PurposeRowHandle, TArray<UE::Editor::DataStorage::RowHandle>& OutFactories) const;
+
+	// The the query conditions for the provided factory
+	const UE::Editor::DataStorage::Queries::FConditions& GetFactoryConditions(UE::Editor::DataStorage::RowHandle FactoryRow) const ;
 
 	UE::Editor::DataStorage::TableHandle WidgetTable{ UE::Editor::DataStorage::InvalidTableHandle };
-	
-	TMap<FName, FPurposeInfo> WidgetPurposes;
+	UE::Editor::DataStorage::TableHandle WidgetPurposeTable{ UE::Editor::DataStorage::InvalidTableHandle };
+	UE::Editor::DataStorage::TableHandle WidgetFactoryTable{ UE::Editor::DataStorage::InvalidTableHandle };
 	
 	UE::Editor::DataStorage::ICoreProvider* Storage{ nullptr };
 	UE::Editor::DataStorage::ICompatibilityProvider* StorageCompatibility{ nullptr };
diff --git a/Engine/Plugins/Experimental/EditorDataStorage/Source/TedsUI/Private/Widgets/CounterWidget.cpp b/Engine/Plugins/Experimental/EditorDataStorage/Source/TedsUI/Private/Widgets/CounterWidget.cpp
--- a/Engine/Plugins/Experimental/EditorDataStorage/Source/TedsUI/Private/Widgets/CounterWidget.cpp
+++ b/Engine/Plugins/Experimental/EditorDataStorage/Source/TedsUI/Private/Widgets/CounterWidget.cpp
@@ -32,7 +32,9 @@ FAutoConsoleCommand EnableCounterWidgetsConsoleCommand(
 // UCounterWidgetFactory
 //
 
-FName UCounterWidgetFactory::WigetPurpose(TEXT("LevelEditor.StatusBar.ToolBar"));
+UE::Editor::DataStorage::IUiProvider::FPurposeID UCounterWidgetFactory::LevelEditorWidgetPurpose(
+	UE::Editor::DataStorage::IUiProvider::FPurposeInfo("LevelEditor", "StatusBar", "Toolbar").GeneratePurposeID());
+
 bool UCounterWidgetFactory::bAreCounterWidgetsEnabled{ false };
 bool UCounterWidgetFactory::bHasBeenSetup{ false };
 
@@ -80,9 +82,10 @@ void UCounterWidgetFactory::RegisterQueries(UE::Editor::DataStorage::ICoreProvid
 
 void UCounterWidgetFactory::RegisterWidgetPurposes(UE::Editor::DataStorage::IUiProvider& DataStorageUi) const
 {
-	using namespace UE::Editor::DataStorage;
-	DataStorageUi.RegisterWidgetPurpose(WigetPurpose, IUiProvider::EPurposeType::Generic,
-		LOCTEXT("ToolBarPurposeDescription", "Widgets added to the status bar at the bottom editor of the main editor window."));
+	DataStorageUi.RegisterWidgetPurpose(
+		UE::Editor::DataStorage::IUiProvider::FPurposeInfo("LevelEditor", "StatusBar", "Toolbar",
+			UE::Editor::DataStorage::IUiProvider::EPurposeType::Generic,
+			LOCTEXT("ToolBarPurposeDescription", "Widgets added to the status bar at the bottom editor of the main editor window.")));
 }
 
 void UCounterWidgetFactory::RegisterWidgetConstructors(UE::Editor::DataStorage::ICoreProvider& DataStorage,
@@ -100,7 +103,7 @@ void UCounterWidgetFactory::RegisterWidgetConstructors(UE::Editor::DataStorage::
 		Where().
 			All<FTypedElementSlateWidgetReferenceColumn>().
 		Compile());
-	DataStorageUi.RegisterWidgetFactory(WigetPurpose, MoveTemp(WidgetCounter));
+	DataStorageUi.RegisterWidgetFactory(DataStorageUi.FindPurpose(LevelEditorWidgetPurpose), MoveTemp(WidgetCounter));
 }
 
 void UCounterWidgetFactory::EnableCounterWidgets()
@@ -115,15 +118,15 @@ void UCounterWidgetFactory::SetupMainWindowIntegrations(TSharedPtr<SWindow> Pare
 	{
 		using namespace UE::Editor::DataStorage;
 
-		IUiProvider* UiInterface = GetMutableDataStorageFeature<IUiProvider>(UiFeatureName);
+		UE::Editor::DataStorage::IUiProvider* UiInterface = GetMutableDataStorageFeature<UE::Editor::DataStorage::IUiProvider>(UiFeatureName);
 		checkf(UiInterface, TEXT(
 			"FEditorDataStorageUiModule tried to integrate with the main window before the "
 			"TEDS UI interface is available."));
 
-		UToolMenu* Menu = UToolMenus::Get()->ExtendMenu(WigetPurpose);
+		UToolMenu* Menu = UToolMenus::Get()->ExtendMenu("LevelEditor.StatusBar.Toolbar");
 
 		TArray<TSharedRef<SWidget>> Widgets;
-		UiInterface->ConstructWidgets(WigetPurpose, {},
+		UiInterface->ConstructWidgets(UiInterface->FindPurpose(LevelEditorWidgetPurpose), {},
 			[&Widgets](const TSharedRef<SWidget>& NewWidget, UE::Editor::DataStorage::RowHandle Row)
 			{
 				Widgets.Add(NewWidget);
@@ -160,8 +163,7 @@ FCounterWidgetConstructor::FCounterWidgetConstructor()
 
 TConstArrayView<const UScriptStruct*> FCounterWidgetConstructor::GetAdditionalColumnsList() const
 {
-	using namespace UE::Editor::DataStorage;
-	static TTypedElementColumnTypeList<FCounterWidgetColumn, FTypedElementU32IntValueCacheColumn> Columns;
+	static UE::Editor::DataStorage::TTypedElementColumnTypeList<FCounterWidgetColumn, FTypedElementU32IntValueCacheColumn> Columns;
 	return Columns;
 }
 
diff --git a/Engine/Plugins/Experimental/EditorDataStorage/Source/TedsUI/Private/Widgets/ExportedTextWidget.cpp b/Engine/Plugins/Experimental/EditorDataStorage/Source/TedsUI/Private/Widgets/ExportedTextWidget.cpp
--- a/Engine/Plugins/Experimental/EditorDataStorage/Source/TedsUI/Private/Widgets/ExportedTextWidget.cpp
+++ b/Engine/Plugins/Experimental/EditorDataStorage/Source/TedsUI/Private/Widgets/ExportedTextWidget.cpp
@@ -18,7 +18,7 @@
 void UExportedTextWidgetFactory::RegisterWidgetConstructors(UE::Editor::DataStorage::ICoreProvider& DataStorage,
 	UE::Editor::DataStorage::IUiProvider& DataStorageUi) const
 {
-	DataStorageUi.RegisterWidgetFactory(FName(TEXT("General.Cell.Default")), FExportedTextWidgetConstructor::StaticStruct());
+	DataStorageUi.RegisterWidgetFactory(DataStorageUi.FindPurpose(DataStorageUi.GetDefaultWidgetPurposeID()), FExportedTextWidgetConstructor::StaticStruct());
 }
 
 //
@@ -41,7 +41,7 @@ TConstArrayView<const UScriptStruct*> FExportedTextWidgetConstructor::GetAdditio
     	return Columns;
 }
 
-const UE::Editor::DataStorage::Queries::FConditions* FExportedTextWidgetConstructor::GetQueryConditions() const
+const UE::Editor::DataStorage::Queries::FConditions* FExportedTextWidgetConstructor::GetQueryConditions(const UE::Editor::DataStorage::ICoreProvider* Storage) const
 {
 	// For the exported text widget, the query condition we are matched against is the column we are exporting text for
 	if(MatchedColumn.IsCompiled() && !MatchedColumn.IsEmpty())
diff --git a/Engine/Plugins/Experimental/EditorDataStorage/Source/TedsUI/Private/Widgets/GeneralWidgetRegistrationFactory.cpp b/Engine/Plugins/Experimental/EditorDataStorage/Source/TedsUI/Private/Widgets/GeneralWidgetRegistrationFactory.cpp
--- a/Engine/Plugins/Experimental/EditorDataStorage/Source/TedsUI/Private/Widgets/GeneralWidgetRegistrationFactory.cpp
+++ b/Engine/Plugins/Experimental/EditorDataStorage/Source/TedsUI/Private/Widgets/GeneralWidgetRegistrationFactory.cpp
@@ -6,29 +6,33 @@
 
 #define LOCTEXT_NAMESPACE "TypedElementsUI_GeneralRegistration"
 
-const FName UGeneralWidgetRegistrationFactory::CellPurpose(TEXT("General.Cell"));
 const FName UGeneralWidgetRegistrationFactory::LargeCellPurpose(TEXT("General.Cell.Large"));
 const FName UGeneralWidgetRegistrationFactory::HeaderPurpose(TEXT("General.Header"));
-const FName UGeneralWidgetRegistrationFactory::CellDefaultPurpose(TEXT("General.Cell.Default"));
-const FName UGeneralWidgetRegistrationFactory::HeaderDefaultPurpose(TEXT("General.Header.Default"));
 
 void UGeneralWidgetRegistrationFactory::RegisterWidgetPurposes(UE::Editor::DataStorage::IUiProvider& DataStorageUi) const
 {
-	using namespace UE::Editor::DataStorage;
-	DataStorageUi.RegisterWidgetPurpose(CellPurpose, IUiProvider::EPurposeType::UniqueByNameAndColumn,
-		LOCTEXT("GeneralCellPurpose", "General purpose widgets that can be used as cells for specific columns or column combinations."));
-
-	DataStorageUi.RegisterWidgetPurpose(LargeCellPurpose, IUiProvider::EPurposeType::UniqueByNameAndColumn,
-		LOCTEXT("GeneralLargeCellPurpose", "General purpose widgets that are specifically designed to be embedded in a space larger than a single cell."));
-
-	DataStorageUi.RegisterWidgetPurpose(HeaderPurpose, IUiProvider::EPurposeType::UniqueByNameAndColumn,
-		LOCTEXT("GeneralHeaderPurpose", "General purpose widget that can be used as a header."));
-
-	DataStorageUi.RegisterWidgetPurpose(HeaderDefaultPurpose, IUiProvider::EPurposeType::UniqueByName,
-		LOCTEXT("GeneralHeaderDefaultPurpose", "The default widget to use in headers if no other specialization is provided."));
+	DataStorageUi.RegisterWidgetPurpose(
+		DataStorageUi.GetDefaultWidgetPurposeID(),
+		UE::Editor::DataStorage::IUiProvider::FPurposeInfo("General", "Cell", "Default",
+			UE::Editor::DataStorage::IUiProvider::EPurposeType::Generic,
+			LOCTEXT("GeneralCellDefaultPurpose", "The default widget to use in cells if no other specialization is provide.")));
+	
+	DataStorageUi.RegisterWidgetPurpose(
+		DataStorageUi.GetGeneralWidgetPurposeID(),
+		UE::Editor::DataStorage::IUiProvider::FPurposeInfo("General", "Cell", NAME_None,
+			UE::Editor::DataStorage::IUiProvider::EPurposeType::UniqueByNameAndColumn,
+			LOCTEXT("GeneralCellPurpose", "General purpose widgets that can be used as cells for specific columns or column combinations.")));
+
+	DataStorageUi.RegisterWidgetPurpose(
+		UE::Editor::DataStorage::IUiProvider::FPurposeInfo("General", "Cell", "Large",
+			UE::Editor::DataStorage::IUiProvider::EPurposeType::UniqueByNameAndColumn,
+			LOCTEXT("GeneralLargeCellPurpose", "General purpose widgets that are specifically designed to be embedded in a space larger than a single cell."),
+			DataStorageUi.GetGeneralWidgetPurposeID()));
 	
-	DataStorageUi.RegisterWidgetPurpose(CellDefaultPurpose, IUiProvider::EPurposeType::UniqueByName,
-		LOCTEXT("GeneralCellDefaultPurpose", "The default widget to use in cells if no other specialization is provided."));
+	DataStorageUi.RegisterWidgetPurpose(
+		UE::Editor::DataStorage::IUiProvider::FPurposeInfo("General", "Header", NAME_None,
+			UE::Editor::DataStorage::IUiProvider::EPurposeType::UniqueByNameAndColumn,
+			LOCTEXT("GeneralHeaderPurpose", "General purpose widget that can be used as a header.")));
 }
 
 #undef LOCTEXT_NAMESPACE
diff --git a/Engine/Plugins/Experimental/EditorDataStorage/Source/TedsUI/Private/Widgets/LabelWidget.cpp b/Engine/Plugins/Experimental/EditorDataStorage/Source/TedsUI/Private/Widgets/LabelWidget.cpp
--- a/Engine/Plugins/Experimental/EditorDataStorage/Source/TedsUI/Private/Widgets/LabelWidget.cpp
+++ b/Engine/Plugins/Experimental/EditorDataStorage/Source/TedsUI/Private/Widgets/LabelWidget.cpp
@@ -29,21 +29,22 @@ void ULabelWidgetFactory::RegisterWidgetConstructors(UE::Editor::DataStorage::IC
 {
 	using namespace UE::Editor::DataStorage::Queries;
 
-	DataStorageUi.RegisterWidgetFactory<FLabelWidgetConstructor>(FName(TEXT("General.Cell")), 
+	DataStorageUi.RegisterWidgetFactory<FLabelWidgetConstructor>(
+		DataStorageUi.FindPurpose(DataStorageUi.GetGeneralWidgetPurposeID()), 
 		TColumn<FTypedElementLabelColumn>() || (TColumn<FTypedElementLabelColumn>() && TColumn<FTypedElementLabelHashColumn>()));
 
-	DataStorageUi.RegisterWidgetFactory<FLabelWidgetConstructor>(FName(TEXT("General.RowLabel")), 
+	DataStorageUi.RegisterWidgetFactory<FLabelWidgetConstructor>(
+		DataStorageUi.FindPurpose(UE::Editor::DataStorage::IUiProvider::FPurposeInfo("General", "RowLabel", NAME_None).GeneratePurposeID()), 
 		TColumn<FTypedElementLabelColumn>() || (TColumn<FTypedElementLabelColumn>() && TColumn<FTypedElementLabelHashColumn>()));
 }
 
 void ULabelWidgetFactory::RegisterWidgetPurposes(UE::Editor::DataStorage::IUiProvider& DataStorageUi) const
 {
-	using namespace UE::Editor::DataStorage;
-	DataStorageUi.RegisterWidgetPurpose(FName(TEXT("General.RowLabel")), IUiProvider::EPurposeType::UniqueByNameAndColumn,
-	LOCTEXT("GeneralRowLabelPurpose", "Specific purpose to request a widget to display a user facing display name for a row."));
-	
-	DataStorageUi.RegisterWidgetPurpose(FName(TEXT("General.RowLabel.Default")), IUiProvider::EPurposeType::UniqueByName,
-	LOCTEXT("GeneralRowLabelDefaultPurpose", "Default purpose to request a widget to display a user facing display name for a row."));
+	DataStorageUi.RegisterWidgetPurpose(
+		UE::Editor::DataStorage::IUiProvider::FPurposeInfo("General", "RowLabel", NAME_None,
+			UE::Editor::DataStorage::IUiProvider::EPurposeType::UniqueByNameAndColumn,
+			LOCTEXT("GeneralRowLabelPurpose", "Specific purpose to request a widget to display a user facing display name for a row."),
+			DataStorageUi.GetGeneralWidgetPurposeID()));
 }
 
 
@@ -58,8 +59,7 @@ FLabelWidgetConstructor::FLabelWidgetConstructor()
 
 TConstArrayView<const UScriptStruct*> FLabelWidgetConstructor::GetAdditionalColumnsList() const
 {
-	using namespace UE::Editor::DataStorage;
-	static const TTypedElementColumnTypeList<
+	static const UE::Editor::DataStorage::TTypedElementColumnTypeList<
 		FTypedElementRowReferenceColumn,
 		FTypedElementU64IntValueCacheColumn,
 		FExternalWidgetSelectionColumn> Columns;
@@ -67,7 +67,7 @@ TConstArrayView<const UScriptStruct*> FLabelWidgetConstructor::GetAdditionalColu
 }
 
 TSharedPtr<SWidget> FLabelWidgetConstructor::CreateWidget(UE::Editor::DataStorage::ICoreProvider* DataStorage,
-	UE::Editor::DataStorage::IUiProvider* DataStorageUi, UE::Editor::DataStorage::RowHandle TargetRow,
+	UE::Editor::DataStorage::IUiProvider* DataStorageUi, UE::Editor::DataStorage::RowHandle TargetRow, 
 	UE::Editor::DataStorage::RowHandle WidgetRow, const UE::Editor::DataStorage::FMetaDataView& Arguments)
 {
 	UE::Editor::DataStorage::FAttributeBinder Binder(TargetRow, DataStorage);
diff --git a/Engine/Plugins/Experimental/EditorDataStorage/Source/TedsUI/Private/Widgets/PackagePathWidget.cpp b/Engine/Plugins/Experimental/EditorDataStorage/Source/TedsUI/Private/Widgets/PackagePathWidget.cpp
--- a/Engine/Plugins/Experimental/EditorDataStorage/Source/TedsUI/Private/Widgets/PackagePathWidget.cpp
+++ b/Engine/Plugins/Experimental/EditorDataStorage/Source/TedsUI/Private/Widgets/PackagePathWidget.cpp
@@ -16,9 +16,9 @@ void UPackagePathWidgetFactory::RegisterWidgetConstructors(UE::Editor::DataStora
 {
 	using namespace UE::Editor::DataStorage::Queries;
 
-	DataStorageUi.RegisterWidgetFactory<FPackagePathWidgetConstructor>(FName(TEXT("General.Cell")),
+	DataStorageUi.RegisterWidgetFactory<FPackagePathWidgetConstructor>(DataStorageUi.FindPurpose(DataStorageUi.GetGeneralWidgetPurposeID()),
 		TColumn<FTypedElementPackagePathColumn>());
-	DataStorageUi.RegisterWidgetFactory<FLoadedPackagePathWidgetConstructor>(FName(TEXT("General.Cell")),
+	DataStorageUi.RegisterWidgetFactory<FLoadedPackagePathWidgetConstructor>(DataStorageUi.FindPurpose(DataStorageUi.GetGeneralWidgetPurposeID()),
 		TColumn<FTypedElementPackageLoadedPathColumn>());
 }
 
diff --git a/Engine/Plugins/Experimental/EditorDataStorage/Source/TedsUI/Private/Widgets/SlateColorWidget.cpp b/Engine/Plugins/Experimental/EditorDataStorage/Source/TedsUI/Private/Widgets/SlateColorWidget.cpp
--- a/Engine/Plugins/Experimental/EditorDataStorage/Source/TedsUI/Private/Widgets/SlateColorWidget.cpp
+++ b/Engine/Plugins/Experimental/EditorDataStorage/Source/TedsUI/Private/Widgets/SlateColorWidget.cpp
@@ -44,7 +44,7 @@ void USlateColorWidgetFactory::RegisterWidgetConstructors(UE::Editor::DataStorag
 	UE::Editor::DataStorage::IUiProvider& DataStorageUi) const
 {
 	using namespace UE::Editor::DataStorage::Queries;
-	DataStorageUi.RegisterWidgetFactory<FSlateColorWidgetConstructor>(TEXT("General.Cell"), TColumn<FSlateColorColumn>());
+	DataStorageUi.RegisterWidgetFactory<FSlateColorWidgetConstructor>(DataStorageUi.FindPurpose(DataStorageUi.GetGeneralWidgetPurposeID()), TColumn<FSlateColorColumn>());
 }
 
 FSlateColorWidgetConstructor::FSlateColorWidgetConstructor()
diff --git a/Engine/Plugins/Experimental/EditorDataStorage/Source/TedsUI/Private/Widgets/UrlWidget.cpp b/Engine/Plugins/Experimental/EditorDataStorage/Source/TedsUI/Private/Widgets/UrlWidget.cpp
--- a/Engine/Plugins/Experimental/EditorDataStorage/Source/TedsUI/Private/Widgets/UrlWidget.cpp
+++ b/Engine/Plugins/Experimental/EditorDataStorage/Source/TedsUI/Private/Widgets/UrlWidget.cpp
@@ -12,7 +12,7 @@ void UUrlWidgetFactory::RegisterWidgetConstructors(UE::Editor::DataStorage::ICor
 	UE::Editor::DataStorage::IUiProvider& DataStorageUi) const
 {
 	using namespace UE::Editor::DataStorage::Queries;
-	DataStorageUi.RegisterWidgetFactory<FUrlWidgetConstructor>(TEXT("General.Cell"), TColumn<FUrlColumn>());
+	DataStorageUi.RegisterWidgetFactory<FUrlWidgetConstructor>(DataStorageUi.FindPurpose(DataStorageUi.GetGeneralWidgetPurposeID()), TColumn<FUrlColumn>());
 }
 
 FUrlWidgetConstructor::FUrlWidgetConstructor()
diff --git a/Engine/Plugins/Experimental/EditorDataStorage/Source/TedsUI/Public/Widgets/CounterWidget.h b/Engine/Plugins/Experimental/EditorDataStorage/Source/TedsUI/Public/Widgets/CounterWidget.h
--- a/Engine/Plugins/Experimental/EditorDataStorage/Source/TedsUI/Public/Widgets/CounterWidget.h
+++ b/Engine/Plugins/Experimental/EditorDataStorage/Source/TedsUI/Public/Widgets/CounterWidget.h
@@ -24,7 +24,7 @@ class TEDSUI_API UCounterWidgetFactory : public UEditorDataStorageFactory
 	GENERATED_BODY()
 
 public:
-	static FName WigetPurpose;
+	static UE::Editor::DataStorage::IUiProvider::FPurposeID LevelEditorWidgetPurpose;
 
 	UCounterWidgetFactory();
 	~UCounterWidgetFactory() override = default;
diff --git a/Engine/Plugins/Experimental/EditorDataStorage/Source/TedsUI/Public/Widgets/ExportedTextWidget.h b/Engine/Plugins/Experimental/EditorDataStorage/Source/TedsUI/Public/Widgets/ExportedTextWidget.h
--- a/Engine/Plugins/Experimental/EditorDataStorage/Source/TedsUI/Public/Widgets/ExportedTextWidget.h
+++ b/Engine/Plugins/Experimental/EditorDataStorage/Source/TedsUI/Public/Widgets/ExportedTextWidget.h
@@ -37,7 +37,7 @@ public:
 	virtual ~FExportedTextWidgetConstructor() override = default;
 
 	virtual TConstArrayView<const UScriptStruct*> GetAdditionalColumnsList() const override;
-	virtual const UE::Editor::DataStorage::Queries::FConditions* GetQueryConditions() const override;
+	virtual const UE::Editor::DataStorage::Queries::FConditions* GetQueryConditions(const UE::Editor::DataStorage::ICoreProvider* Storage) const override;
 
 	virtual FText CreateWidgetDisplayNameText(
 		UE::Editor::DataStorage::ICoreProvider* DataStorage, UE::Editor::DataStorage::RowHandle Row) const override;
diff --git a/Engine/Plugins/Experimental/EditorDataStorage/Source/TedsUI/Public/Widgets/GeneralWidgetRegistrationFactory.h b/Engine/Plugins/Experimental/EditorDataStorage/Source/TedsUI/Public/Widgets/GeneralWidgetRegistrationFactory.h
--- a/Engine/Plugins/Experimental/EditorDataStorage/Source/TedsUI/Public/Widgets/GeneralWidgetRegistrationFactory.h
+++ b/Engine/Plugins/Experimental/EditorDataStorage/Source/TedsUI/Public/Widgets/GeneralWidgetRegistrationFactory.h
@@ -13,11 +13,8 @@ class TEDSUI_API UGeneralWidgetRegistrationFactory : public UEditorDataStorageFa
 	GENERATED_BODY()
 
 public:
-	static const FName CellPurpose;
 	static const FName LargeCellPurpose;
 	static const FName HeaderPurpose;
-	static const FName CellDefaultPurpose;
-	static const FName HeaderDefaultPurpose;
 
 	~UGeneralWidgetRegistrationFactory() override = default;
 
diff --git a/Engine/Plugins/Experimental/EditorDataStorageFeatures/Source/TedsAlerts/Private/TedsAlertWidget.cpp b/Engine/Plugins/Experimental/EditorDataStorageFeatures/Source/TedsAlerts/Private/TedsAlertWidget.cpp
--- a/Engine/Plugins/Experimental/EditorDataStorageFeatures/Source/TedsAlerts/Private/TedsAlertWidget.cpp
+++ b/Engine/Plugins/Experimental/EditorDataStorageFeatures/Source/TedsAlerts/Private/TedsAlertWidget.cpp
@@ -155,10 +155,11 @@ void UAlertWidgetFactory::RegisterWidgetConstructors(
 	using namespace UE::Editor::DataStorage::Columns;
 	using namespace UE::Editor::DataStorage::Queries;
 
-	DataStorageUi.RegisterWidgetFactory<FAlertWidgetConstructor>(FName(TEXT("General.Cell")),
+	DataStorageUi.RegisterWidgetFactory<FAlertWidgetConstructor>(DataStorageUi.FindPurpose(DataStorageUi.GetGeneralWidgetPurposeID()),
 		TColumn<FAlertColumn>() || TColumn<FChildAlertColumn>());
-
-	DataStorageUi.RegisterWidgetFactory<FAlertHeaderWidgetConstructor>(FName(TEXT("General.Header")),
+	
+	DataStorageUi.RegisterWidgetFactory<FAlertHeaderWidgetConstructor>(
+		DataStorageUi.FindPurpose(IUiProvider::FPurposeInfo("General", "Header", NAME_None).GeneratePurposeID()),
 		TColumn<FAlertColumn>() || TColumn<FChildAlertColumn>());
 }
 
diff --git a/Engine/Plugins/Experimental/EditorDataStorageFeatures/Source/TedsAssetData/Private/Widgets/AssetDataDiskSizeWidget.cpp b/Engine/Plugins/Experimental/EditorDataStorageFeatures/Source/TedsAssetData/Private/Widgets/AssetDataDiskSizeWidget.cpp
--- a/Engine/Plugins/Experimental/EditorDataStorageFeatures/Source/TedsAssetData/Private/Widgets/AssetDataDiskSizeWidget.cpp
+++ b/Engine/Plugins/Experimental/EditorDataStorageFeatures/Source/TedsAssetData/Private/Widgets/AssetDataDiskSizeWidget.cpp
@@ -13,7 +13,7 @@ void UDiskSizeWidgetFactory::RegisterWidgetConstructors(UE::Editor::DataStorage:
 	UE::Editor::DataStorage::IUiProvider& DataStorageUi) const
 {
 	using namespace UE::Editor::DataStorage::Queries;
-	DataStorageUi.RegisterWidgetFactory<FDiskSizeWidgetConstructor>(TEXT("General.Cell"), TColumn<FDiskSizeColumn>());
+	DataStorageUi.RegisterWidgetFactory<FDiskSizeWidgetConstructor>(DataStorageUi.FindPurpose(DataStorageUi.GetGeneralWidgetPurposeID()), TColumn<FDiskSizeColumn>());
 }
 
 FDiskSizeWidgetConstructor::FDiskSizeWidgetConstructor()
diff --git a/Engine/Plugins/Experimental/EditorDataStorageFeatures/Source/TedsAssetData/Private/Widgets/AssetDataItemTypeWidget.cpp b/Engine/Plugins/Experimental/EditorDataStorageFeatures/Source/TedsAssetData/Private/Widgets/AssetDataItemTypeWidget.cpp
--- a/Engine/Plugins/Experimental/EditorDataStorageFeatures/Source/TedsAssetData/Private/Widgets/AssetDataItemTypeWidget.cpp
+++ b/Engine/Plugins/Experimental/EditorDataStorageFeatures/Source/TedsAssetData/Private/Widgets/AssetDataItemTypeWidget.cpp
@@ -16,7 +16,8 @@ void UAssetDataItemTypeWidgetFactory::RegisterWidgetConstructors(UE::Editor::Dat
 	UE::Editor::DataStorage::IUiProvider& DataStorageUi) const
 {
 	using namespace UE::Editor::DataStorage::Queries;
-	DataStorageUi.RegisterWidgetFactory<FAssetDataItemTypeWidgetConstructor>(TEXT("General.Cell"), TColumn<FAssetClassColumn>() || TColumn<FFolderTag>());
+	DataStorageUi.RegisterWidgetFactory<FAssetDataItemTypeWidgetConstructor>(DataStorageUi.FindPurpose(DataStorageUi.GetGeneralWidgetPurposeID()),
+		TColumn<FAssetClassColumn>() || TColumn<FFolderTag>());
 }
 
 FAssetDataItemTypeWidgetConstructor::FAssetDataItemTypeWidgetConstructor()
diff --git a/Engine/Plugins/Experimental/EditorDataStorageFeatures/Source/TedsAssetData/Private/Widgets/AssetDataLabelWidget.cpp b/Engine/Plugins/Experimental/EditorDataStorageFeatures/Source/TedsAssetData/Private/Widgets/AssetDataLabelWidget.cpp
--- a/Engine/Plugins/Experimental/EditorDataStorageFeatures/Source/TedsAssetData/Private/Widgets/AssetDataLabelWidget.cpp
+++ b/Engine/Plugins/Experimental/EditorDataStorageFeatures/Source/TedsAssetData/Private/Widgets/AssetDataLabelWidget.cpp
@@ -23,7 +23,7 @@ void UAssetDataLabelWidgetFactory::RegisterWidgetConstructors(UE::Editor::DataSt
 {
 	using namespace UE::Editor::DataStorage::Queries;
 	DataStorageUi.RegisterWidgetFactory<FAssetDataLabelWidgetConstructor>(
-		TEXT("General.RowLabel"),
+		DataStorageUi.FindPurpose(IUiProvider::FPurposeInfo("General", "RowLabel", NAME_None).GeneratePurposeID()),
 		TColumn<FNameColumn>() && (TColumn<FAssetTag>() || TColumn<FAssetPathColumn_Experimental>()));
 }
 
@@ -53,7 +53,7 @@ TSharedPtr<SWidget> FAssetDataLabelWidgetConstructor::CreateWidget(UE::Editor::D
 		TAttribute<const FSlateBrush*> FolderImage;
 
 		// Create the LabelWidget through its General Widget
-		static FName LabelPurposeName = TEXT("General.Cell");
+		UE::Editor::DataStorage::RowHandle LabelPurposeName = DataStorageUi->FindPurpose(DataStorageUi->GetGeneralWidgetPurposeID());
 
 		TSharedPtr<FTypedElementWidgetConstructor> OutLabelWidgetConstructorPtr;
 
diff --git a/Engine/Plugins/Experimental/EditorDataStorageFeatures/Source/TedsAssetData/Private/Widgets/AssetDataThumbnailWidget.cpp b/Engine/Plugins/Experimental/EditorDataStorageFeatures/Source/TedsAssetData/Private/Widgets/AssetDataThumbnailWidget.cpp
--- a/Engine/Plugins/Experimental/EditorDataStorageFeatures/Source/TedsAssetData/Private/Widgets/AssetDataThumbnailWidget.cpp
+++ b/Engine/Plugins/Experimental/EditorDataStorageFeatures/Source/TedsAssetData/Private/Widgets/AssetDataThumbnailWidget.cpp
@@ -19,7 +19,7 @@ void UAssetThumbnailWidgetFactory::RegisterWidgetConstructors(UE::Editor::DataSt
 {
 	using namespace UE::Editor::DataStorage::Queries;
 	DataStorageUi.RegisterWidgetFactory<FAssetThumbnailWidgetConstructor>(
-		"General.Cell",
+		DataStorageUi.FindPurpose(DataStorageUi.GetGeneralWidgetPurposeID()),
 		TColumn<FAssetTag>() || TColumn<FFolderTag>());
 }
 
diff --git a/Engine/Plugins/Experimental/EditorDataStorageFeatures/Source/TedsAssetData/Private/Widgets/AssetDataVirtualPathWidget.cpp b/Engine/Plugins/Experimental/EditorDataStorageFeatures/Source/TedsAssetData/Private/Widgets/AssetDataVirtualPathWidget.cpp
--- a/Engine/Plugins/Experimental/EditorDataStorageFeatures/Source/TedsAssetData/Private/Widgets/AssetDataVirtualPathWidget.cpp
+++ b/Engine/Plugins/Experimental/EditorDataStorageFeatures/Source/TedsAssetData/Private/Widgets/AssetDataVirtualPathWidget.cpp
@@ -15,7 +15,7 @@ void UAssetDataVirtualPathWidgetFactory::RegisterWidgetConstructors(UE::Editor::
 {
 	using namespace UE::Editor::DataStorage::Queries;
 	DataStorageUi.RegisterWidgetFactory<FAssetDataVirtualPathWidgetConstructor>(
-		TEXT("General.Cell"),
+		DataStorageUi.FindPurpose(DataStorageUi.GetGeneralWidgetPurposeID()),
 		TColumn<FVirtualPathColumn_Experimental>() || TColumn<FAssetPathColumn_Experimental>());
 }
 
diff --git a/Engine/Plugins/Experimental/EditorDataStorageFeatures/Source/TedsAssetData/Private/Widgets/AssetLabelWidget.cpp b/Engine/Plugins/Experimental/EditorDataStorageFeatures/Source/TedsAssetData/Private/Widgets/AssetLabelWidget.cpp
--- a/Engine/Plugins/Experimental/EditorDataStorageFeatures/Source/TedsAssetData/Private/Widgets/AssetLabelWidget.cpp
+++ b/Engine/Plugins/Experimental/EditorDataStorageFeatures/Source/TedsAssetData/Private/Widgets/AssetLabelWidget.cpp
@@ -15,7 +15,7 @@ void UAssetLabelWidgetFactory::RegisterWidgetConstructors(UE::Editor::DataStorag
 {
 	using namespace UE::Editor::DataStorage::Queries;
 	DataStorageUi.RegisterWidgetFactory<FAssetLabelWidgetConstructor>(
-		"General.Cell",
+		DataStorageUi.FindPurpose(DataStorageUi.GetGeneralWidgetPurposeID()),
 		TColumn<FNameColumn>());
 }
 
diff --git a/Engine/Plugins/Experimental/EditorDataStorageFeatures/Source/TedsAssetData/Private/Widgets/ContentBrowserListViewNameWidget.cpp b/Engine/Plugins/Experimental/EditorDataStorageFeatures/Source/TedsAssetData/Private/Widgets/ContentBrowserListViewNameWidget.cpp
--- a/Engine/Plugins/Experimental/EditorDataStorageFeatures/Source/TedsAssetData/Private/Widgets/ContentBrowserListViewNameWidget.cpp
+++ b/Engine/Plugins/Experimental/EditorDataStorageFeatures/Source/TedsAssetData/Private/Widgets/ContentBrowserListViewNameWidget.cpp
@@ -16,13 +16,18 @@
 
 #define LOCTEXT_NAMESPACE "FContentBrowserLabelAssetWidgetConstructor"
 
-FName UContentBrowserListViewNameWidgetFactory::WidgetPurpose(TEXT("ContentBrowser.RowLabel"));
+UE::Editor::DataStorage::IUiProvider::FPurposeID UContentBrowserListViewNameWidgetFactory::WidgetPurpose(UE::Editor::DataStorage::IUiProvider::FPurposeInfo(
+			"ContentBrowser", "RowLabel", NAME_None).GeneratePurposeID());
 
 void UContentBrowserListViewNameWidgetFactory::RegisterWidgetPurposes(UE::Editor::DataStorage::IUiProvider& DataStorageUi) const
 {
 	using namespace UE::Editor::DataStorage::Queries;
-	DataStorageUi.RegisterWidgetPurpose(WidgetPurpose, IUiProvider::EPurposeType::UniqueByNameAndColumn,
-		LOCTEXT("ContentBrowserLabelWidget_PurposeDescription", "Widget that display a Label + Thumbnail in the Content Browser"));
+
+	DataStorageUi.RegisterWidgetPurpose(
+		IUiProvider::FPurposeInfo("ContentBrowser", "RowLabel", NAME_None,
+			IUiProvider::EPurposeType::UniqueByNameAndColumn,
+			LOCTEXT("ContentBrowserLabelWidget_PurposeDescription", "Widget that display a Label + Thumbnail in the Content Browser."),
+			IUiProvider::FPurposeInfo("General", "RowLabel", NAME_None).GeneratePurposeID()));
 }
 
 void UContentBrowserListViewNameWidgetFactory::RegisterWidgetConstructors(UE::Editor::DataStorage::ICoreProvider& DataStorage,
@@ -30,7 +35,7 @@ void UContentBrowserListViewNameWidgetFactory::RegisterWidgetConstructors(UE::Ed
 {
 	using namespace UE::Editor::DataStorage::Queries;
 	DataStorageUi.RegisterWidgetFactory<FContentBrowserListViewNameWidgetConstructor>(
-		WidgetPurpose,
+		DataStorageUi.FindPurpose(WidgetPurpose),
 		TColumn<FNameColumn>() && (TColumn<FAssetTag>() || TColumn<FFolderTag>()));
 }
 
@@ -52,8 +57,6 @@ TSharedPtr<SWidget> FContentBrowserListViewNameWidgetConstructor::CreateWidget(U
 
 	if (DataStorage->IsRowAvailable(TargetRow))
 	{
-		static FName LabelThumbnailPurposeName = TEXT("General.Cell");
-
 		TSharedPtr<FTypedElementWidgetConstructor> OutThumbnailWidgetConstructorPtr;
 		TSharedPtr<FTypedElementWidgetConstructor> OutLabelWidgetConstructorPtr;
 
@@ -71,10 +74,12 @@ TSharedPtr<SWidget> FContentBrowserListViewNameWidgetConstructor::CreateWidget(U
 
 		// TODO: Arguments need to be developed further to customize the Thumbnail Config if needed
 		TArray<TWeakObjectPtr<const UScriptStruct>> ThumbnailColumns = GetThumbnailColumns();
-		DataStorageUi->CreateWidgetConstructors(LabelThumbnailPurposeName, IUiProvider::EMatchApproach::ExactMatch, ThumbnailColumns, Arguments, AssignThumbnailWidgetToColumn);
+		DataStorageUi->CreateWidgetConstructors(DataStorageUi->FindPurpose(DataStorageUi->GetGeneralWidgetPurposeID()),
+			IUiProvider::EMatchApproach::ExactMatch, ThumbnailColumns, Arguments, AssignThumbnailWidgetToColumn);
 
 		TArray<TWeakObjectPtr<const UScriptStruct>> LabelColumns = GetLabelColumns();
-		DataStorageUi->CreateWidgetConstructors(LabelThumbnailPurposeName, IUiProvider::EMatchApproach::ExactMatch, LabelColumns, Arguments, AssignLabelWidgetToColumn);
+		DataStorageUi->CreateWidgetConstructors(DataStorageUi->FindPurpose(DataStorageUi->GetGeneralWidgetPurposeID()),
+			IUiProvider::EMatchApproach::ExactMatch, LabelColumns, Arguments, AssignLabelWidgetToColumn);
 
 		static const FMargin ColumnItemPadding(8, 0);
 		constexpr float ThumbnailNameHorizontalPadding = 8.f;
diff --git a/Engine/Plugins/Experimental/EditorDataStorageFeatures/Source/TedsAssetData/Private/Widgets/ContentBrowserListViewNameWidget.h b/Engine/Plugins/Experimental/EditorDataStorageFeatures/Source/TedsAssetData/Private/Widgets/ContentBrowserListViewNameWidget.h
--- a/Engine/Plugins/Experimental/EditorDataStorageFeatures/Source/TedsAssetData/Private/Widgets/ContentBrowserListViewNameWidget.h
+++ b/Engine/Plugins/Experimental/EditorDataStorageFeatures/Source/TedsAssetData/Private/Widgets/ContentBrowserListViewNameWidget.h
@@ -21,7 +21,7 @@ class UContentBrowserListViewNameWidgetFactory : public UEditorDataStorageFactor
 	GENERATED_BODY()
 
 public:
-	static FName WidgetPurpose;
+	static UE::Editor::DataStorage::IUiProvider::FPurposeID WidgetPurpose;
 
 	~UContentBrowserListViewNameWidgetFactory() override = default;
 
diff --git a/Engine/Plugins/Experimental/EditorDataStorageFeatures/Source/TedsAssetData/Private/Widgets/DynamicAssetDataColumnBaseWidget.cpp b/Engine/Plugins/Experimental/EditorDataStorageFeatures/Source/TedsAssetData/Private/Widgets/DynamicAssetDataColumnBaseWidget.cpp
--- a/Engine/Plugins/Experimental/EditorDataStorageFeatures/Source/TedsAssetData/Private/Widgets/DynamicAssetDataColumnBaseWidget.cpp
+++ b/Engine/Plugins/Experimental/EditorDataStorageFeatures/Source/TedsAssetData/Private/Widgets/DynamicAssetDataColumnBaseWidget.cpp
@@ -36,11 +36,11 @@ void UDynamicAssetDataColumnBaseWidgetFactory::RegisterWidgetConstructors(UE::Ed
 	UE::Editor::DataStorage::IUiProvider& DataStorageUi) const
 {
 	using namespace UE::Editor::DataStorage::Queries;
-	DataStorageUi.RegisterWidgetFactory<FDynamicAssetDataColumnBaseWidgetConstructor>(TEXT("General.Cell"), TColumn<FItemStringAttributeColumn_Experimental>("ParentClass"));
-	DataStorageUi.RegisterWidgetFactory<FDynamicAssetDataColumnBaseWidgetConstructor>(TEXT("General.Cell"), TColumn<FItemStringAttributeColumn_Experimental>("Skeleton"));
-	DataStorageUi.RegisterWidgetFactory<FDynamicAssetDataColumnBaseWidgetConstructor>(TEXT("General.Cell"), TColumn<FItemStringAttributeColumn_Experimental>("SourceTexture"));
-	DataStorageUi.RegisterWidgetFactory<FDynamicAssetDataColumnBaseWidgetConstructor>(TEXT("General.Cell"), TColumn<FItemStringAttributeColumn_Experimental>("PhysicsAsset"));
-	DataStorageUi.RegisterWidgetFactory<FDynamicAssetDataColumnBaseWidgetConstructor>(TEXT("General.Cell"), TColumn<FItemStringAttributeColumn_Experimental>("ShadowPhysicsAsset"));
+	DataStorageUi.RegisterWidgetFactory<FDynamicAssetDataColumnBaseWidgetConstructor>(DataStorageUi.FindPurpose(DataStorageUi.GetGeneralWidgetPurposeID()), TColumn<FItemStringAttributeColumn_Experimental>("ParentClass"));
+	DataStorageUi.RegisterWidgetFactory<FDynamicAssetDataColumnBaseWidgetConstructor>(DataStorageUi.FindPurpose(DataStorageUi.GetGeneralWidgetPurposeID()), TColumn<FItemStringAttributeColumn_Experimental>("Skeleton"));
+	DataStorageUi.RegisterWidgetFactory<FDynamicAssetDataColumnBaseWidgetConstructor>(DataStorageUi.FindPurpose(DataStorageUi.GetGeneralWidgetPurposeID()), TColumn<FItemStringAttributeColumn_Experimental>("SourceTexture"));
+	DataStorageUi.RegisterWidgetFactory<FDynamicAssetDataColumnBaseWidgetConstructor>(DataStorageUi.FindPurpose(DataStorageUi.GetGeneralWidgetPurposeID()), TColumn<FItemStringAttributeColumn_Experimental>("PhysicsAsset"));
+	DataStorageUi.RegisterWidgetFactory<FDynamicAssetDataColumnBaseWidgetConstructor>(DataStorageUi.FindPurpose(DataStorageUi.GetGeneralWidgetPurposeID()), TColumn<FItemStringAttributeColumn_Experimental>("ShadowPhysicsAsset"));
 
 }
 
diff --git a/Engine/Plugins/Experimental/EditorDataStorageFeatures/Source/TedsAssetData/Private/Widgets/StaticMeshTrianglesWidget.cpp b/Engine/Plugins/Experimental/EditorDataStorageFeatures/Source/TedsAssetData/Private/Widgets/StaticMeshTrianglesWidget.cpp
--- a/Engine/Plugins/Experimental/EditorDataStorageFeatures/Source/TedsAssetData/Private/Widgets/StaticMeshTrianglesWidget.cpp
+++ b/Engine/Plugins/Experimental/EditorDataStorageFeatures/Source/TedsAssetData/Private/Widgets/StaticMeshTrianglesWidget.cpp
@@ -12,7 +12,7 @@ void UStaticMeshTrianglesWidgetFactory::RegisterWidgetConstructors(UE::Editor::D
 	UE::Editor::DataStorage::IUiProvider& DataStorageUi) const
 {
 	using namespace UE::Editor::DataStorage::Queries;
-	DataStorageUi.RegisterWidgetFactory<FStaticMeshTrianglesWidgetConstructor>(TEXT("General.Cell"), TColumn<FItemStringAttributeColumn_Experimental>("Triangles"));
+	DataStorageUi.RegisterWidgetFactory<FStaticMeshTrianglesWidgetConstructor>(DataStorageUi.FindPurpose(DataStorageUi.GetGeneralWidgetPurposeID()), TColumn<FItemStringAttributeColumn_Experimental>("Triangles"));
 }
 
 FStaticMeshTrianglesWidgetConstructor::FStaticMeshTrianglesWidgetConstructor()
diff --git a/Engine/Plugins/Experimental/EditorDataStorageFeatures/Source/TedsContentBrowser/Private/TedsContentBrowserModule.cpp b/Engine/Plugins/Experimental/EditorDataStorageFeatures/Source/TedsContentBrowser/Private/TedsContentBrowserModule.cpp
--- a/Engine/Plugins/Experimental/EditorDataStorageFeatures/Source/TedsContentBrowser/Private/TedsContentBrowserModule.cpp
+++ b/Engine/Plugins/Experimental/EditorDataStorageFeatures/Source/TedsContentBrowser/Private/TedsContentBrowserModule.cpp
@@ -220,7 +220,8 @@ namespace UE::Editor::ContentBrowser
 		// Create the table viewer widget
 		TableViewer = SNew(STedsTableViewer)
 					.QueryStack(RowQueryStack)
-					.CellWidgetPurposes({TEXT("ContentBrowser.RowLabel"), TEXT("General.RowLabel"), TEXT("General.Cell")})
+					.CellWidgetPurpose(IUiProvider::FPurposeInfo("ContentBrowser", "RowLabel", NAME_None).GeneratePurposeID())
+					.HeaderWidgetPurpose(IUiProvider::FPurposeInfo("General", "Header", NAME_None).GeneratePurposeID())
 					// Default list of columns to display
 					.Columns({ FNameColumn::StaticStruct(), FTedsAlertColumn::StaticStruct(),
 						FAssetClassColumn::StaticStruct(), FAssetTag::StaticStruct(), FFolderTag::StaticStruct(), FAssetPathColumn_Experimental::StaticStruct(),
@@ -245,7 +246,7 @@ namespace UE::Editor::ContentBrowser
 		// Create the table viewer widget
 		TableViewer = SNew(STedsTileViewer)
 					.QueryStack(RowQueryStack)
-					.WidgetPurpose(TEXT("ContentBrowser.RowLabel"))
+					.WidgetPurpose(IUiProvider::FPurposeInfo("ContentBrowser", "RowLabel", NAME_None).GeneratePurposeID())
 					.SelectionMode(ESelectionMode::Multi)
 					.OnSelectionChanged_Lambda([this](RowHandle Row)
 					{
@@ -329,7 +330,9 @@ namespace UE::Editor::ContentBrowser
 			FTedsAlertColumn::StaticStruct(), FTedsChildAlertColumn::StaticStruct() };
 
 		// Same widget purposes as the Outliner for now
-		OutInitParams.WidgetPurposes = { FName("SceneOutliner.Cell"), FName("General.Cell") };
+		
+		OutInitParams.CellWidgetPurpose = IUiProvider::FPurposeID(IUiProvider::FPurposeInfo(
+			"SceneOutliner", "Cell", NAME_None).GeneratePurposeID());
 
 	}
 
diff --git a/Engine/Plugins/Experimental/EditorDataStorageFeatures/Source/TedsContentBrowser/Private/Widgets/TedsContentBrowserAssetViewWidget.cpp b/Engine/Plugins/Experimental/EditorDataStorageFeatures/Source/TedsContentBrowser/Private/Widgets/TedsContentBrowserAssetViewWidget.cpp
--- a/Engine/Plugins/Experimental/EditorDataStorageFeatures/Source/TedsContentBrowser/Private/Widgets/TedsContentBrowserAssetViewWidget.cpp
+++ b/Engine/Plugins/Experimental/EditorDataStorageFeatures/Source/TedsContentBrowser/Private/Widgets/TedsContentBrowserAssetViewWidget.cpp
@@ -42,16 +42,14 @@ public:
 		case ETableViewMode::Type::List:
 			ChildWidget = SNew(UE::Editor::DataStorage::STedsTableViewer)
 					.QueryStack(RowView)
-					.CellWidgetPurposes(Arguments._InitParams.WidgetPurposes)
+					.CellWidgetPurpose(Arguments._InitParams.CellWidgetPurpose)
 					.Columns(Arguments._InitParams.Columns)
 					.ListSelectionMode(ESelectionMode::Multi);
 			break;
 		case ETableViewMode::Type::Tile:
 			ChildWidget = SNew(UE::Editor::DataStorage::STedsTileViewer)
 					.QueryStack(RowView)
-					// The Tile View only accepts one widget purposes intentionally, all other table viewers will do the same with the updated
-					// purposes system so all TArray args will be removed
-					.WidgetPurpose(Arguments._InitParams.WidgetPurposes[0])
+					.WidgetPurpose(Arguments._InitParams.CellWidgetPurpose)
 					.Columns(Arguments._InitParams.Columns)
 					.SelectionMode(ESelectionMode::Multi);
 			break;
@@ -77,7 +75,7 @@ public:
 
 		Storage->RunQuery(QueryHandle,
 			Queries::CreateDirectQueryCallbackBinding([&NewRows]
-				(const IDirectQueryContext& Context, const RowHandle*)
+				(const UE::Editor::DataStorage::IDirectQueryContext& Context, const RowHandle*)
 				{
 					NewRows.Append(Context.GetRowHandles());
 				}));
@@ -111,20 +109,24 @@ private:
 
 namespace TedsContentBrowser::Private
 {
-	static const FName WidgetPurpose(TEXT("ContentBrowser.AssetView"));
+	static const UE::Editor::DataStorage::IUiProvider::FPurposeID
+		Purpose(UE::Editor::DataStorage::IUiProvider::FPurposeInfo("ContentBrowser", "AssetView", NAME_None).GeneratePurposeID());
+
 }
 
 void UContentBrowserAssetViewWidgetFactory::RegisterWidgetPurposes(UE::Editor::DataStorage::IUiProvider& DataStorageUi) const
 {
-	DataStorageUi.RegisterWidgetPurpose(TedsContentBrowser::Private::WidgetPurpose, UE::Editor::DataStorage::IUiProvider::EPurposeType::UniqueByName,
-		LOCTEXT("ContentBrowserAssetView_PurposeDescription", "Widget that displays a table viewer in the content browser"));
+	DataStorageUi.RegisterWidgetPurpose(
+		UE::Editor::DataStorage::IUiProvider::FPurposeInfo("ContentBrowser", "AssetView", NAME_None,
+			UE::Editor::DataStorage::IUiProvider::EPurposeType::UniqueByName,
+			LOCTEXT("ContentBrowserAssetView_PurposeDescription", "Widget that displays a table viewer in the content browser")));
 }
 
 void UContentBrowserAssetViewWidgetFactory::RegisterWidgetConstructors(UE::Editor::DataStorage::ICoreProvider& DataStorage,
 	UE::Editor::DataStorage::IUiProvider& DataStorageUi) const
 {
 	using namespace UE::Editor::DataStorage::Queries;
-	DataStorageUi.RegisterWidgetFactory<FContentBrowserAssetViewWidgetConstructor>(TedsContentBrowser::Private::WidgetPurpose);
+	DataStorageUi.RegisterWidgetFactory<FContentBrowserAssetViewWidgetConstructor>(DataStorageUi.FindPurpose(TedsContentBrowser::Private::Purpose));
 }
 
 FContentBrowserAssetViewWidgetConstructor::FContentBrowserAssetViewWidgetConstructor()
@@ -138,8 +140,7 @@ FContentBrowserAssetViewWidgetConstructor::FContentBrowserAssetViewWidgetConstru
 }
 
 TSharedPtr<SWidget> FContentBrowserAssetViewWidgetConstructor::CreateWidget(UE::Editor::DataStorage::ICoreProvider* DataStorage,
-	UE::Editor::DataStorage::IUiProvider* DataStorageUi,
-	UE::Editor::DataStorage::RowHandle TargetRow, UE::Editor::DataStorage::RowHandle WidgetRow,
+	UE::Editor::DataStorage::IUiProvider* DataStorageUi, UE::Editor::DataStorage::RowHandle TargetRow, UE::Editor::DataStorage::RowHandle WidgetRow,
 	const UE::Editor::DataStorage::FMetaDataView& Arguments)
 {
 	using namespace UE::Editor::DataStorage::Queries;
diff --git a/Engine/Plugins/Experimental/EditorDataStorageFeatures/Source/TedsDebugger/Private/Widgets/QueryEditor/TedsQueryEditorResultsView.cpp b/Engine/Plugins/Experimental/EditorDataStorageFeatures/Source/TedsDebugger/Private/Widgets/QueryEditor/TedsQueryEditorResultsView.cpp
--- a/Engine/Plugins/Experimental/EditorDataStorageFeatures/Source/TedsDebugger/Private/Widgets/QueryEditor/TedsQueryEditorResultsView.cpp
+++ b/Engine/Plugins/Experimental/EditorDataStorageFeatures/Source/TedsDebugger/Private/Widgets/QueryEditor/TedsQueryEditorResultsView.cpp
@@ -12,6 +12,7 @@
 #include "Modules/ModuleManager.h"
 #include "QueryEditor/TedsQueryEditorModel.h"
 #include "TedsRowViewNode.h"
+#include "Elements/Columns/TypedElementHiearchyColumns.h"
 #include "Widgets/STedsTableViewer.h"
 #include "Widgets/SRowDetails.h"
 #include "Widgets/Text/STextBlock.h"
@@ -29,6 +30,9 @@ namespace UE::Editor::DataStorage::Debug::QueryEditor
 	{
 		Model = &InModel;
 		ModelChangedDelegateHandle = Model->GetModelChangedDelegate().AddRaw(this, &SResultsView::OnModelChanged);
+		UiProvider = GetMutableDataStorageFeature<IUiProvider>(UiFeatureName);
+		checkf(UiProvider, TEXT("Cannot create SResultsView without TEDS UI"));
+
 		// RowQueryHandle = InvalidQueryHandle;
 
 		// Create a custom column for the table viewer to display row handles
@@ -106,6 +110,7 @@ namespace UE::Editor::DataStorage::Debug::QueryEditor
 
 		ICoreProvider& TedsInterface = Model->GetTedsInterface();
 
+		
 		if(bModelDirty)
 		{
 			{
@@ -122,9 +127,18 @@ namespace UE::Editor::DataStorage::Debug::QueryEditor
 			{
 				FQueryDescription TableViewerQueryDescription = Model->GenerateQueryDescription();
 
+				RowHandle GeneralPurposeRowHandle = UiProvider->FindPurpose(UiProvider->GetGeneralWidgetPurposeID());
+				RowHandle DefaultPurposeRowHandle = UiProvider->FindPurpose(UiProvider->GetDefaultWidgetPurposeID());
+
+				// Temporarily add the default purpose as a parent of the general purpose so the debugger can support both
+				TedsInterface.AddColumn(GeneralPurposeRowHandle, FTableRowParentColumn{.Parent = DefaultPurposeRowHandle});
+
 				// Update the columns in the table viewer using the selection types from the query description
 				TableViewer->SetColumns(TArray<TWeakObjectPtr<const UScriptStruct>>(TableViewerQueryDescription.SelectionTypes));
 
+				// Remove the parenting chain after we have used it to generate widgets
+				TedsInterface.RemoveColumns<FTableRowParentColumn>(GeneralPurposeRowHandle);
+
 				// Since setting the columns clears all columns, we are going to re-add the custom column back
 				if(RowHandleColumn)
 				{
@@ -189,10 +203,10 @@ namespace UE::Editor::DataStorage::Debug::QueryEditor
 			return false;
 		};
 	
-		IUiProvider* StorageUi = GetMutableDataStorageFeature<IUiProvider>(UiFeatureName);
-		checkf(StorageUi, TEXT("SResultsView created before data storage interfaces were initialized."))
 
-		StorageUi->CreateWidgetConstructors(TEXT("General.Cell.RowHandle"), FMetaDataView(), AssignWidgetToColumn);
+		UiProvider->CreateWidgetConstructors(
+			UiProvider->FindPurpose(IUiProvider::FPurposeInfo("General", "Cell", "RowHandle").GeneratePurposeID()),
+			FMetaDataView(), AssignWidgetToColumn);
 	}
 } // namespace UE::Editor::DataStorage::Debug::QueryEditor
 
diff --git a/Engine/Plugins/Experimental/EditorDataStorageFeatures/Source/TedsDebugger/Private/Widgets/QueryEditor/TedsQueryEditorResultsView.h b/Engine/Plugins/Experimental/EditorDataStorageFeatures/Source/TedsDebugger/Private/Widgets/QueryEditor/TedsQueryEditorResultsView.h
--- a/Engine/Plugins/Experimental/EditorDataStorageFeatures/Source/TedsDebugger/Private/Widgets/QueryEditor/TedsQueryEditorResultsView.h
+++ b/Engine/Plugins/Experimental/EditorDataStorageFeatures/Source/TedsDebugger/Private/Widgets/QueryEditor/TedsQueryEditorResultsView.h
@@ -4,6 +4,7 @@
 #include "Elements/Interfaces/TypedElementDataStorageInterface.h"
 #include "Widgets/SCompoundWidget.h"
 
+
 class ISceneOutliner;
 class SSceneOutliner;
 class SHorizontalBox;
@@ -14,6 +15,7 @@ namespace UE::Editor::DataStorage
 	class STedsTableViewer;
 	class FTedsTableViewerColumn;
 	class SRowDetails;
+	class IUiProvider;
 
 	namespace QueryStack
 	{
@@ -58,6 +60,8 @@ namespace UE::Editor::DataStorage
 
 			// Widget that displays details of a row
 			TSharedPtr<SRowDetails> RowDetailsWidget;
+			
+			IUiProvider* UiProvider = nullptr;
 		};
 
 	} // namespace Debug::QueryEditor
diff --git a/Engine/Plugins/Experimental/EditorDataStorageFeatures/Source/TedsDebugger/Private/Widgets/RowReferenceWidget.cpp b/Engine/Plugins/Experimental/EditorDataStorageFeatures/Source/TedsDebugger/Private/Widgets/RowReferenceWidget.cpp
--- a/Engine/Plugins/Experimental/EditorDataStorageFeatures/Source/TedsDebugger/Private/Widgets/RowReferenceWidget.cpp
+++ b/Engine/Plugins/Experimental/EditorDataStorageFeatures/Source/TedsDebugger/Private/Widgets/RowReferenceWidget.cpp
@@ -62,7 +62,8 @@ void URowReferenceWidgetFactory::RegisterWidgetConstructors(UE::Editor::DataStor
 	using namespace UE::Editor::DataStorage::Queries;
 
 	// TEDS UI TODO: We can re-use this widget for FTypedElementParentColumn
-	DataStorageUi.RegisterWidgetFactory<FRowReferenceWidgetConstructor>(FName(TEXT("SceneOutliner.Cell")),
+	DataStorageUi.RegisterWidgetFactory<FRowReferenceWidgetConstructor>(
+		DataStorageUi.FindPurpose(IUiProvider::FPurposeInfo("SceneOutliner", "Cell", NAME_None).GeneratePurposeID()),
 		TColumn<FTypedElementRowReferenceColumn>());
 }
 
diff --git a/Engine/Plugins/Experimental/EditorDataStorageFeatures/Source/TedsDebugger/Private/Widgets/SlateVisualizationWidget.cpp b/Engine/Plugins/Experimental/EditorDataStorageFeatures/Source/TedsDebugger/Private/Widgets/SlateVisualizationWidget.cpp
--- a/Engine/Plugins/Experimental/EditorDataStorageFeatures/Source/TedsDebugger/Private/Widgets/SlateVisualizationWidget.cpp
+++ b/Engine/Plugins/Experimental/EditorDataStorageFeatures/Source/TedsDebugger/Private/Widgets/SlateVisualizationWidget.cpp
@@ -13,7 +13,7 @@ void USlateVisualizationWidgetFactory::RegisterWidgetConstructors(UE::Editor::Da
 {
 	using namespace UE::Editor::DataStorage::Queries;
 
-	DataStorageUi.RegisterWidgetFactory<FSlateVisualizationWidgetConstructor>(FName(TEXT("General.Cell")),
+	DataStorageUi.RegisterWidgetFactory<FSlateVisualizationWidgetConstructor>(DataStorageUi.FindPurpose(DataStorageUi.GetGeneralWidgetPurposeID()),
 		TColumn<FTypedElementSlateWidgetReferenceColumn>());
 }
 
diff --git a/Engine/Plugins/Experimental/EditorDataStorageFeatures/Source/TedsOutliner/Private/Compatibility/SceneOutlinerRowHandleColumn.cpp b/Engine/Plugins/Experimental/EditorDataStorageFeatures/Source/TedsOutliner/Private/Compatibility/SceneOutlinerRowHandleColumn.cpp
--- a/Engine/Plugins/Experimental/EditorDataStorageFeatures/Source/TedsOutliner/Private/Compatibility/SceneOutlinerRowHandleColumn.cpp
+++ b/Engine/Plugins/Experimental/EditorDataStorageFeatures/Source/TedsOutliner/Private/Compatibility/SceneOutlinerRowHandleColumn.cpp
@@ -37,7 +37,9 @@ void FSceneOutlinerRowHandleColumn::CreateWidgetConstructor()
 	IUiProvider* StorageUi = GetMutableDataStorageFeature<IUiProvider>(UiFeatureName);
 	checkf(StorageUi, TEXT("FSceneOutlinerRowHandleColumn created before data storage interfaces were initialized."))
 	
-	StorageUi->CreateWidgetConstructors(TEXT("General.Cell.RowHandle"), FMetaDataView(), AssignWidgetToColumn);
+	StorageUi->CreateWidgetConstructors(
+		StorageUi->FindPurpose(IUiProvider::FPurposeInfo("General", "Cell", "RowHandle").GeneratePurposeID()),
+		FMetaDataView(), AssignWidgetToColumn);
 }
 
 
diff --git a/Engine/Plugins/Experimental/EditorDataStorageFeatures/Source/TedsOutliner/Private/Compatibility/SceneOutlinerTedsBridge.cpp b/Engine/Plugins/Experimental/EditorDataStorageFeatures/Source/TedsOutliner/Private/Compatibility/SceneOutlinerTedsBridge.cpp
--- a/Engine/Plugins/Experimental/EditorDataStorageFeatures/Source/TedsOutliner/Private/Compatibility/SceneOutlinerTedsBridge.cpp
+++ b/Engine/Plugins/Experimental/EditorDataStorageFeatures/Source/TedsOutliner/Private/Compatibility/SceneOutlinerTedsBridge.cpp
@@ -39,12 +39,10 @@ FAutoConsoleCommand BindColumnsToSceneOutlinerConsoleCommand(
 	FConsoleCommandWithArgsDelegate::CreateLambda([](const TArray<FString>& Args)
 		{
 			using namespace UE::Editor::DataStorage::Queries;
-			
-		    const FName WidgetPurposes[] = {TEXT("SceneOutliner.Cell"), TEXT("General.Cell")};
 
 			if (ICoreProvider* DataStorage = GetMutableDataStorageFeature<ICoreProvider>(StorageFeatureName))
 			{
-				static UE::Editor::DataStorage::QueryHandle Queries[] =
+				static QueryHandle Queries[] =
 				{
 					DataStorage->RegisterQuery(Select().ReadWrite<FTypedElementLabelColumn>().Compile()),
 					DataStorage->RegisterQuery(Select().ReadOnly<FTypedElementLocalTransformColumn>().Compile()),
@@ -71,9 +69,9 @@ FAutoConsoleCommand BindColumnsToSceneOutlinerConsoleCommand(
 						if (Args[0].IsNumeric())
 						{
 							int32 QueryIndex = FCString::Atoi(*Args[0]);
-							if (QueryIndex < sizeof(Queries) / sizeof(UE::Editor::DataStorage::QueryHandle))
+							if (QueryIndex < sizeof(Queries) / sizeof(QueryHandle))
 							{
-								Binder.AssignQuery(Queries[QueryIndex], SceneOutliner, WidgetPurposes);
+								Binder.AssignQuery(Queries[QueryIndex], SceneOutliner);
 								return;
 							}
 						}
@@ -89,7 +87,7 @@ FAutoConsoleCommand BindColumnsToSceneOutlinerConsoleCommand(
 								if (!Arg.IsEmpty() && Arg[0] == '/' && Path.TrySetPath(Arg))
 								{
 									const UScriptStruct* ColumnType = TypeOptional(Path);
-									if (ColumnType && ColumnType->IsChildOf(UE::Editor::DataStorage::FColumn::StaticStruct()))
+									if (ColumnType && ColumnType->IsChildOf(FColumn::StaticStruct()))
 									{
 										Query.ReadOnly(ColumnType);
 										++AdditionCount;
@@ -98,18 +96,18 @@ FAutoConsoleCommand BindColumnsToSceneOutlinerConsoleCommand(
 							}
 							if (AdditionCount > 0)
 							{
-								static UE::Editor::DataStorage::QueryHandle CustomQuery = UE::Editor::DataStorage::InvalidQueryHandle;
-								if (CustomQuery != UE::Editor::DataStorage::InvalidQueryHandle)
+								static QueryHandle CustomQuery = InvalidQueryHandle;
+								if (CustomQuery != InvalidQueryHandle)
 								{
 									DataStorage->UnregisterQuery(CustomQuery);
 								}
 								CustomQuery = DataStorage->RegisterQuery(Query.Compile());
-								Binder.AssignQuery(CustomQuery, SceneOutliner, WidgetPurposes);
+								Binder.AssignQuery(CustomQuery, SceneOutliner);
 								return;
 							}
 						}
 					}
-					Binder.AssignQuery(UE::Editor::DataStorage::InvalidQueryHandle, SceneOutliner, WidgetPurposes);
+					Binder.AssignQuery(InvalidQueryHandle, SceneOutliner);
 				}
 			}
 		}));
@@ -124,8 +122,10 @@ public:
 		UE::Editor::DataStorage::IUiProvider& InStorageUi,
 		UE::Editor::DataStorage::ICompatibilityProvider& InStorageCompatibility,
 		const TSharedPtr<ISceneOutliner>& InOutliner);
-
-	void AssignQuery(UE::Editor::DataStorage::QueryHandle Query, const TConstArrayView<FName> CellWidgetPurposes);
+	
+	void AssignQuery(UE::Editor::DataStorage::QueryHandle Query,
+		const UE::Editor::DataStorage::IUiProvider::FPurposeID& CellWidgetPurpose, const UE::Editor::DataStorage::IUiProvider::FPurposeID& HeaderWidgetPurpose);
+	
 	void RegisterDealiaser(const FTreeItemIDDealiaser& InDealiaser);
 	FTreeItemIDDealiaser GetDealiaser();
 	
@@ -138,7 +138,7 @@ private:
 	UE::Editor::DataStorage::IUiProvider* StorageUi{ nullptr };
 	UE::Editor::DataStorage::ICompatibilityProvider* StorageCompatibility{ nullptr };
 	FTreeItemIDDealiaser Dealiaser;
-	TArray<FName> CellWidgetPurposes;
+	UE::Editor::DataStorage::IUiProvider::FPurposeID CellWidgetPurpose;
 };
 
 class FOutlinerColumn : public ISceneOutlinerColumn
@@ -185,7 +185,7 @@ public:
 		return NameId;
 	}
 
-	virtual void Tick(double InCurrentTime, float InDeltaTime)
+	virtual void Tick(double InCurrentTime, float InDeltaTime) override
 	{
 		TableViewerColumnImpl->Tick();
 		if (FallbackColumn)
@@ -295,7 +295,7 @@ public:
 		if(Storage.IsRowAssigned(TargetRowHandle))
 		{
 			RowWidget = TableViewerColumnImpl->ConstructRowWidget(TargetRowHandle,
-				[&](ICoreProvider& DataStorage, const RowHandle& WidgetRow)
+				[&](UE::Editor::DataStorage::ICoreProvider& DataStorage, const RowHandle& WidgetRow)
 				{
 					DataStorage.AddColumn(WidgetRow, FTedsOutlinerColumn{ .Outliner = OwningOutliner });
 				});
@@ -342,25 +342,30 @@ public:
 
 void USceneOutlinerTedsBridgeFactory::RegisterWidgetPurposes(UE::Editor::DataStorage::IUiProvider& DataStorageUi) const
 {
-	using namespace UE::Editor::DataStorage;
-	using PurposeType = IUiProvider::EPurposeType;
-
-	DataStorageUi.RegisterWidgetPurpose(FSceneOutlinerTedsQueryBinder::HeaderWidgetPurpose, PurposeType::UniqueByNameAndColumn,
-		LOCTEXT("HeaderWidgetPurpose", "Widgets for headers in any Scene Outliner for specific columns or column combinations."));
-	DataStorageUi.RegisterWidgetPurpose(FSceneOutlinerTedsQueryBinder::DefaultHeaderWidgetPurpose, PurposeType::UniqueByName,
-		LOCTEXT("DefaultHeaderWidgetPurpose", "The default widget to use in headers for the Scene Outliner."));
-	
-	DataStorageUi.RegisterWidgetPurpose(FSceneOutlinerTedsQueryBinder::CellWidgetPurpose, PurposeType::UniqueByNameAndColumn,
-		LOCTEXT("CellWidgetPurpose", "Widgets for cells in any Scene Outliner for specific columns or column combinations."));
-	DataStorageUi.RegisterWidgetPurpose(FSceneOutlinerTedsQueryBinder::DefaultCellWidgetPurpose, PurposeType::UniqueByName,
-		LOCTEXT("DefaultCellWidgetPurpose", "The default widget to use in cells for the Scene Outliner."));
-
-	DataStorageUi.RegisterWidgetPurpose(FSceneOutlinerTedsQueryBinder::ItemLabelCellWidgetPurpose, PurposeType::UniqueByNameAndColumn,
-		LOCTEXT("ItemCellWidgetPurpose", "Widgets for cells in any Scene Outliner that are specific to the Item label column."));
-	DataStorageUi.RegisterWidgetPurpose(FSceneOutlinerTedsQueryBinder::DefaultItemLabelCellWidgetPurpose, PurposeType::UniqueByName,
-		LOCTEXT("DefaultItemCellWidgetPurpose", "The default widget to use in cells for the Scene Outliner specific to the Item label column."));
+	UE::Editor::DataStorage::IUiProvider::FPurposeID GeneralRowLabelPurposeID =
+		UE::Editor::DataStorage::IUiProvider::FPurposeInfo("General", "RowLabel", NAME_None).GeneratePurposeID();
 
+	UE::Editor::DataStorage::IUiProvider::FPurposeID GeneralHeaderPurposeID =
+		UE::Editor::DataStorage::IUiProvider::FPurposeInfo("General", "Header", NAME_None).GeneratePurposeID();
 
+	DataStorageUi.RegisterWidgetPurpose(
+		UE::Editor::DataStorage::IUiProvider::FPurposeInfo("SceneOutliner", "Header", NAME_None,
+			UE::Editor::DataStorage::IUiProvider::EPurposeType::UniqueByNameAndColumn,
+			LOCTEXT("HeaderWidgetPurpose", "Widgets for headers in any Scene Outliner for specific columns or column combinations."),
+			GeneralHeaderPurposeID));
+	
+	DataStorageUi.RegisterWidgetPurpose(
+		UE::Editor::DataStorage::IUiProvider::FPurposeInfo("SceneOutliner", "Cell", NAME_None,
+			UE::Editor::DataStorage::IUiProvider::EPurposeType::UniqueByNameAndColumn,
+			LOCTEXT("CellWidgetPurpose", "Widgets for cells in any Scene Outliner for specific columns or column combinations."),
+			DataStorageUi.GetGeneralWidgetPurposeID()));
+
+
+	DataStorageUi.RegisterWidgetPurpose(
+		UE::Editor::DataStorage::IUiProvider::FPurposeInfo("SceneOutliner", "RowLabel", NAME_None,
+			UE::Editor::DataStorage::IUiProvider::EPurposeType::UniqueByNameAndColumn,
+			LOCTEXT("ItemCellWidgetPurpose", "Widgets for cells in any Scene Outliner that are specific to the Item label column."),
+			GeneralRowLabelPurposeID));
 }
 
 
@@ -371,18 +376,15 @@ void USceneOutlinerTedsBridgeFactory::RegisterWidgetPurposes(UE::Editor::DataSto
 
 const FName FSceneOutlinerTedsQueryBinder::CellWidgetTableName(TEXT("Editor_SceneOutlinerCellWidgetTable"));
 const FName FSceneOutlinerTedsQueryBinder::HeaderWidgetPurpose(TEXT("SceneOutliner.Header"));
-const FName FSceneOutlinerTedsQueryBinder::DefaultHeaderWidgetPurpose(TEXT("SceneOutliner.Header.Default"));
 const FName FSceneOutlinerTedsQueryBinder::CellWidgetPurpose(TEXT("SceneOutliner.Cell"));
-const FName FSceneOutlinerTedsQueryBinder::DefaultCellWidgetPurpose(TEXT("SceneOutliner.Cell.Default"));
 const FName FSceneOutlinerTedsQueryBinder::ItemLabelCellWidgetPurpose(TEXT("SceneOutliner.RowLabel"));
-const FName FSceneOutlinerTedsQueryBinder::DefaultItemLabelCellWidgetPurpose(TEXT("SceneOutliner.RowLabel.Default"));
 
 FSceneOutlinerTedsQueryBinder::FSceneOutlinerTedsQueryBinder()
 {
 	using namespace UE::Editor::DataStorage;
-	Storage = GetMutableDataStorageFeature<ICoreProvider>(StorageFeatureName);
-	StorageUi = GetMutableDataStorageFeature<IUiProvider>(UiFeatureName);
-	StorageCompatibility = GetMutableDataStorageFeature<ICompatibilityProvider>(CompatibilityFeatureName);
+	Storage = GetMutableDataStorageFeature<UE::Editor::DataStorage::ICoreProvider>(StorageFeatureName);
+	StorageUi = GetMutableDataStorageFeature<UE::Editor::DataStorage::IUiProvider>(UiFeatureName);
+	StorageCompatibility = GetMutableDataStorageFeature<UE::Editor::DataStorage::ICompatibilityProvider>(CompatibilityFeatureName);
 
 	SetupDefaultColumnMapping();
 }
@@ -431,12 +433,26 @@ TSharedPtr<FSceneOutlinerTedsBridge>* FSceneOutlinerTedsQueryBinder::FindQueryMa
 	return SceneOutliners.Find(Outliner);
 }
 
-void FSceneOutlinerTedsQueryBinder::AssignQuery(UE::Editor::DataStorage::QueryHandle Query, const TSharedPtr<ISceneOutliner>& Outliner, TConstArrayView<FName> CellWidgetPurposes)
+
+void FSceneOutlinerTedsQueryBinder::AssignQuery(UE::Editor::DataStorage::QueryHandle Query, const TSharedPtr<ISceneOutliner>& Outliner,
+	const UE::Editor::DataStorage::IUiProvider::FPurposeID& InWidgetPurpose, const UE::Editor::DataStorage::IUiProvider::FPurposeID& InHeaderPurpose)
 {
 	CleanupStaleOutliners();
 
 	TSharedPtr<FSceneOutlinerTedsBridge>* QueryMapping = FindOrAddQueryMapping(Outliner);
-	(*QueryMapping)->AssignQuery(Query, CellWidgetPurposes);
+	(*QueryMapping)->AssignQuery(Query, InWidgetPurpose, InHeaderPurpose);
+}
+
+void FSceneOutlinerTedsQueryBinder::AssignQuery(UE::Editor::DataStorage::QueryHandle Query, const TSharedPtr<ISceneOutliner>& Outliner)
+{
+	UE::Editor::DataStorage::IUiProvider::FPurposeID CellPurpose =
+		UE::Editor::DataStorage::IUiProvider::FPurposeInfo("SceneOutliner", "Cell", NAME_None).GeneratePurposeID();
+	
+	UE::Editor::DataStorage::IUiProvider::FPurposeID HeaderPurpose =
+		UE::Editor::DataStorage::IUiProvider::FPurposeInfo("SceneOutliner", "Header", NAME_None).GeneratePurposeID();
+
+	AssignQuery(Query, Outliner, CellPurpose, HeaderPurpose);
+
 }
 
 void FSceneOutlinerTedsQueryBinder::RegisterTreeItemIDDealiaser(const TSharedPtr<ISceneOutliner>& Outliner, const FTreeItemIDDealiaser& InDealiaser)
@@ -505,23 +521,24 @@ FTreeItemIDDealiaser FSceneOutlinerTedsBridge::GetDealiaser()
 }
 
 
-void FSceneOutlinerTedsBridge::AssignQuery(UE::Editor::DataStorage::QueryHandle Query, const TConstArrayView<FName> InCellWidgetPurposes)
+void FSceneOutlinerTedsBridge::AssignQuery(UE::Editor::DataStorage::QueryHandle Query,
+	const UE::Editor::DataStorage::IUiProvider::FPurposeID& InCellWidgetPurposes, const UE::Editor::DataStorage::IUiProvider::FPurposeID& HeaderWidgetPurpose)
 {
-	using namespace UE::Editor::DataStorage;
-	using MatchApproach = IUiProvider::EMatchApproach;
+	using MatchApproach = UE::Editor::DataStorage::IUiProvider::EMatchApproach;
 	constexpr uint8 DefaultPriorityIndex = 100;
 	FSceneOutlinerTedsQueryBinder& Binder = FSceneOutlinerTedsQueryBinder::GetInstance();
-	CellWidgetPurposes = InCellWidgetPurposes;
-
+	CellWidgetPurpose = InCellWidgetPurposes;
 
 	if (TSharedPtr<ISceneOutliner> OutlinerPinned = Outliner.Pin())
 	{
-		const FQueryDescription& Description = Storage->GetQueryDescription(Query);
+		const UE::Editor::DataStorage::FQueryDescription& Description = Storage->GetQueryDescription(Query);
 		UE::Editor::DataStorage::FQueryMetaDataView MetaDataView(Description);
-
+		UE::Editor::DataStorage::RowHandle CellPurposeRow = StorageUi->FindPurpose(CellWidgetPurpose);
+		UE::Editor::DataStorage::RowHandle HeaderPurposeRow = StorageUi->FindPurpose(HeaderWidgetPurpose);
+		
 		ClearColumns(*OutlinerPinned);
 
-		if (Description.Action == FQueryDescription::EActionType::Select)
+		if (Description.Action == UE::Editor::DataStorage::FQueryDescription::EActionType::Select)
 		{
 			int32 SelectionCount = Description.SelectionTypes.Num();
 			AddedColumns.Reset(SelectionCount);
@@ -529,7 +546,7 @@ void FSceneOutlinerTedsBridge::AssignQuery(UE::Editor::DataStorage::QueryHandle
 			TArray<TWeakObjectPtr<const UScriptStruct>> ColumnTypes = UE::Editor::DataStorage::TableViewerUtils::CreateVerifiedColumnTypeArray(Description.SelectionTypes);
 
 			int32 IndexOffset = 0;
-			auto ColumnConstructor = [this, Query, MetaDataView, &IndexOffset, &OutlinerPinned, Binder](
+			auto ColumnConstructor = [this, Query, MetaDataView, &IndexOffset, &OutlinerPinned, Binder, HeaderPurposeRow](
 				TUniquePtr<FTypedElementWidgetConstructor> Constructor, TConstArrayView<TWeakObjectPtr<const UScriptStruct>> ColumnTypes)
 				{
 					TSharedPtr<FTypedElementWidgetConstructor> CellConstructor(Constructor.Release());
@@ -552,10 +569,10 @@ void FSceneOutlinerTedsBridge::AssignQuery(UE::Editor::DataStorage::QueryHandle
 							ESceneOutlinerColumnVisibility::Visible, 
 							ColumnPriority,
 							FCreateSceneOutlinerColumn::CreateLambda(
-								[this, Query, MetaDataView, NameId, &ColumnTypes, CellConstructor, &OutlinerPinned, FallbackColumn](ISceneOutliner&)
+								[this, Query, MetaDataView, NameId, &ColumnTypes, CellConstructor, &OutlinerPinned, FallbackColumn, HeaderPurposeRow](ISceneOutliner&)
 								{
 									TSharedPtr<FTypedElementWidgetConstructor> HeaderConstructor = 
-										UE::Editor::DataStorage::TableViewerUtils::CreateHeaderWidgetConstructor(*StorageUi, MetaDataView, ColumnTypes, CellWidgetPurposes);
+										UE::Editor::DataStorage::TableViewerUtils::CreateHeaderWidgetConstructor(*StorageUi, MetaDataView, ColumnTypes, HeaderPurposeRow);
 									return MakeShared<FOutlinerColumn>(
 										Query, *Storage, *StorageUi, *StorageCompatibility, NameId,
 										TArray<TWeakObjectPtr<const UScriptStruct>>(ColumnTypes.GetData(), ColumnTypes.Num()), 
@@ -571,61 +588,8 @@ void FSceneOutlinerTedsBridge::AssignQuery(UE::Editor::DataStorage::QueryHandle
 					return true;
 				};
 
-			for(const FName& WidgetPurpose : CellWidgetPurposes)
-			{
-				StorageUi->CreateWidgetConstructors(WidgetPurpose, MatchApproach::LongestMatch, ColumnTypes, 
+			StorageUi->CreateWidgetConstructors(CellPurposeRow, MatchApproach::LongestMatch, ColumnTypes, 
 				MetaDataView, ColumnConstructor);
-			}
-
-			for (TWeakObjectPtr<const UScriptStruct>& ColumnType : ColumnTypes)
-			{
-				FName FallbackColumn = Binder.FindOutlinerColumnFromTEDSColumns({ColumnType});
-
-				auto AssignWidgetToColumn = [this, Query, MetaDataView, ColumnType, &IndexOffset, &OutlinerPinned, FallbackColumn](
-					TUniquePtr<FTypedElementWidgetConstructor> Constructor, TConstArrayView<TWeakObjectPtr<const UScriptStruct>>)
-				{
-					TSharedPtr<FTypedElementWidgetConstructor> CellConstructor(Constructor.Release());
-					FName NameId = FName(ColumnType->GetDisplayNameText().ToString());
-					uint8 ColumnPriority = static_cast<uint8>(FMath::Clamp(DefaultPriorityIndex + static_cast<uint8>(FMath::Clamp(IndexOffset, 0, 255)), 0, 255));
-					AddedColumns.Add(NameId);
-					OutlinerPinned->AddColumn(NameId,
-						FSceneOutlinerColumnInfo(ESceneOutlinerColumnVisibility::Visible, ColumnPriority,
-							FCreateSceneOutlinerColumn::CreateLambda(
-								[this, Query, MetaDataView, NameId, ColumnType, CellConstructor, &OutlinerPinned, FallbackColumn](ISceneOutliner&)
-								{
-									TArray<TWeakObjectPtr<const UScriptStruct>> ColumnTypesStored;
-									ColumnTypesStored.Add(ColumnType);
-									TSharedPtr<FTypedElementWidgetConstructor> HeaderConstructor =
-										UE::Editor::DataStorage::TableViewerUtils::CreateHeaderWidgetConstructor(*StorageUi, MetaDataView, { ColumnType }, CellWidgetPurposes);
-									return MakeShared<FOutlinerColumn>(
-										Query, *Storage, *StorageUi, *StorageCompatibility, NameId, MoveTemp(ColumnTypesStored),
-										HeaderConstructor, CellConstructor, FallbackColumn, OutlinerPinned, Dealiaser);
-
-								})
-						)
-					);
-					++IndexOffset;
-					return false;
-				};
-
-				const int32 BeforeIndexOffset = IndexOffset;
-				for(const FName& WidgetPurpose : CellWidgetPurposes)
-				{
-					const FName DefaultWidgetPurpose(WidgetPurpose.ToString() + TEXT(".Default"));
-
-					StorageUi->CreateWidgetConstructors(DefaultWidgetPurpose, MetaDataView, AssignWidgetToColumn);
-					
-					if (BeforeIndexOffset != IndexOffset)
-					{
-						break;
-					}
-				}
-
-				if (BeforeIndexOffset == IndexOffset)
-				{
-					++IndexOffset;
-				}
-			}
 		}
 	}
 }
diff --git a/Engine/Plugins/Experimental/EditorDataStorageFeatures/Source/TedsOutliner/Private/TedsOutlinerImpl.cpp b/Engine/Plugins/Experimental/EditorDataStorageFeatures/Source/TedsOutliner/Private/TedsOutlinerImpl.cpp
--- a/Engine/Plugins/Experimental/EditorDataStorageFeatures/Source/TedsOutliner/Private/TedsOutlinerImpl.cpp
+++ b/Engine/Plugins/Experimental/EditorDataStorageFeatures/Source/TedsOutliner/Private/TedsOutlinerImpl.cpp
@@ -79,9 +79,33 @@ namespace QueryUtils
 		return false;
 	};
 }
+
+
+FTedsOutlinerParams::FTedsOutlinerParams(SSceneOutliner* InSceneOutliner)
+	: SceneOutliner(InSceneOutliner)
+	, QueryDescription()
+	, bUseDefaultTedsFilters(false)
+	, bShowRowHandleColumn(true)
+	, bForceShowParents(true)
+	, HierarchyData(FTedsOutlinerHierarchyData::GetDefaultHierarchyData())
+{
+	CellWidgetPurpose =
+		DataStorage::IUiProvider::FPurposeID(DataStorage::IUiProvider::FPurposeInfo(
+			"SceneOutliner", "Cell", NAME_None).GeneratePurposeID());
+	
+	HeaderWidgetPurpose =
+		DataStorage::IUiProvider::FPurposeID(DataStorage::IUiProvider::FPurposeInfo(
+			"SceneOutliner", "Header", NAME_None).GeneratePurposeID());
+
+	LabelWidgetPurpose =
+		DataStorage::IUiProvider::FPurposeID(DataStorage::IUiProvider::FPurposeInfo(
+			"SceneOutliner", "RowLabel", NAME_None).GeneratePurposeID());
+}
+	
 FTedsOutlinerImpl::FTedsOutlinerImpl(const FTedsOutlinerParams& InParams, ISceneOutlinerMode* InMode)
 	: CreationParams(InParams)
-	, CellWidgetPurposes(InParams.CellWidgetPurposes)
+	, CellWidgetPurpose(InParams.CellWidgetPurpose)
+	, LabelWidgetPurpose(InParams.LabelWidgetPurpose)
 	, InitialQueryDescription(InParams.QueryDescription)
 	, HierarchyData(InParams.HierarchyData)
 	, SelectionSetName(InParams.SelectionSetOverride)
@@ -91,9 +115,9 @@ FTedsOutlinerImpl::FTedsOutlinerImpl(const FTedsOutlinerParams& InParams, IScene
 {
 	// Initialize the TEDS constructs
 	using namespace UE::Editor::DataStorage;
-	Storage = GetMutableDataStorageFeature<ICoreProvider>(StorageFeatureName);
-	StorageUi = GetMutableDataStorageFeature<IUiProvider>(UiFeatureName);
-	StorageCompatibility = GetMutableDataStorageFeature<ICompatibilityProvider>(CompatibilityFeatureName);
+	Storage = GetMutableDataStorageFeature<DataStorage::ICoreProvider>(StorageFeatureName);
+	StorageUi = GetMutableDataStorageFeature<DataStorage::IUiProvider>(UiFeatureName);
+	StorageCompatibility = GetMutableDataStorageFeature<DataStorage::ICompatibilityProvider>(CompatibilityFeatureName);
 }
 
 void FTedsOutlinerImpl::CreateFilterQueries()
@@ -192,33 +216,19 @@ TSharedRef<SWidget> FTedsOutlinerImpl::CreateLabelWidgetForItem(DataStorage::Row
 	// Find the best matching label widget for a given list of columns
 	auto CreateWidgetConstructor = [this](TArray<TWeakObjectPtr<const UScriptStruct>>& ColumnTypes) -> TSharedPtr<FTypedElementWidgetConstructor>
 	{
-		using namespace UE::Editor::DataStorage;
-		TSharedPtr<FTypedElementWidgetConstructor> OutWidgetConstructorPtr;
-
-		bool bFoundWidget = false;
+		TSharedPtr<FTypedElementWidgetConstructor> OutWidgetConstructorPtr = nullptr;
 
-		static TArray<FName> Purposes{TEXT("SceneOutliner.RowLabel"), TEXT("General.RowLabel")};
-
-		for(const FName& Purpose : Purposes)
-		{
-			StorageUi->CreateWidgetConstructors(Purpose, IUiProvider::EMatchApproach::LongestMatch, ColumnTypes, {},
-					[&OutWidgetConstructorPtr, ColumnTypes, &bFoundWidget](
+		StorageUi->CreateWidgetConstructors(StorageUi->FindPurpose(LabelWidgetPurpose), DataStorage::IUiProvider::EMatchApproach::LongestMatch, ColumnTypes, {},
+					[&OutWidgetConstructorPtr, ColumnTypes](
 					TUniquePtr<FTypedElementWidgetConstructor> CreatedConstructor, 
 					TConstArrayView<TWeakObjectPtr<const UScriptStruct>> MatchedColumnTypes)
 					{
 						OutWidgetConstructorPtr = TSharedPtr<FTypedElementWidgetConstructor>(CreatedConstructor.Release());
-						bFoundWidget = true;
-					
 						// Either this was the exact match so no need to search further or the longest possible chain didn't match so the next ones will 
 						// always be shorter in both cases just return.
 						return false;
 					});
-
-			if(bFoundWidget)
-			{
-				break;
-			}
-		}
+		
 		return OutWidgetConstructorPtr;
 	};
 	
diff --git a/Engine/Plugins/Experimental/EditorDataStorageFeatures/Source/TedsOutliner/Private/TedsOutlinerModule.cpp b/Engine/Plugins/Experimental/EditorDataStorageFeatures/Source/TedsOutliner/Private/TedsOutlinerModule.cpp
--- a/Engine/Plugins/Experimental/EditorDataStorageFeatures/Source/TedsOutliner/Private/TedsOutlinerModule.cpp
+++ b/Engine/Plugins/Experimental/EditorDataStorageFeatures/Source/TedsOutliner/Private/TedsOutlinerModule.cpp
@@ -96,7 +96,7 @@ TSharedRef<ISceneOutliner> FTedsOutlinerModule::CreateTedsOutliner(const FSceneO
 	
 	TSharedRef<ISceneOutliner> TedsOutlinerShared = SNew(SSceneOutliner, InitOptions);
 	
-	FSceneOutlinerTedsQueryBinder::GetInstance().AssignQuery(ColumnQuery, TedsOutlinerShared, InitTedsOptions.CellWidgetPurposes);
+	FSceneOutlinerTedsQueryBinder::GetInstance().AssignQuery(ColumnQuery, TedsOutlinerShared, InInitTedsOptions.CellWidgetPurpose, InInitTedsOptions.LabelWidgetPurpose);
 	
 	return TedsOutlinerShared;
 }
diff --git a/Engine/Plugins/Experimental/EditorDataStorageFeatures/Source/TedsOutliner/Private/Widgets/OutlinerIconWidget.cpp b/Engine/Plugins/Experimental/EditorDataStorageFeatures/Source/TedsOutliner/Private/Widgets/OutlinerIconWidget.cpp
--- a/Engine/Plugins/Experimental/EditorDataStorageFeatures/Source/TedsOutliner/Private/Widgets/OutlinerIconWidget.cpp
+++ b/Engine/Plugins/Experimental/EditorDataStorageFeatures/Source/TedsOutliner/Private/Widgets/OutlinerIconWidget.cpp
@@ -97,7 +97,8 @@ void UOutlinerIconWidgetFactory::RegisterWidgetConstructors(UE::Editor::DataStor
 {
 	using namespace UE::Editor::DataStorage::Queries;
 
-	DataStorageUi.RegisterWidgetFactory<FOutlinerIconWidgetConstructor>(TEXT("SceneOutliner.RowLabel.Icon"),
+	DataStorageUi.RegisterWidgetFactory<FOutlinerIconWidgetConstructor>(
+		DataStorageUi.FindPurpose(IUiProvider::FPurposeInfo("SceneOutliner", "RowLabel", "Icon").GeneratePurposeID()),
 		TColumn<FTypedElementClassTypeInfoColumn>() || TColumn<FTypedElementScriptStructTypeInfoColumn>());
 }
 
diff --git a/Engine/Plugins/Experimental/EditorDataStorageFeatures/Source/TedsOutliner/Private/Widgets/OutlinerLabelWidget.cpp b/Engine/Plugins/Experimental/EditorDataStorageFeatures/Source/TedsOutliner/Private/Widgets/OutlinerLabelWidget.cpp
--- a/Engine/Plugins/Experimental/EditorDataStorageFeatures/Source/TedsOutliner/Private/Widgets/OutlinerLabelWidget.cpp
+++ b/Engine/Plugins/Experimental/EditorDataStorageFeatures/Source/TedsOutliner/Private/Widgets/OutlinerLabelWidget.cpp
@@ -11,27 +11,26 @@
 
 #define LOCTEXT_NAMESPACE "OutlinerLabelWidget"
 
-static const FName IconWidgetPurpose = TEXT("SceneOutliner.RowLabel.Icon");
-static const FName TextWidgetPurpose = TEXT("SceneOutliner.RowLabel.Text");
-
 void UOutlinerLabelWidgetFactory::RegisterWidgetConstructors(UE::Editor::DataStorage::ICoreProvider& DataStorage,
 	UE::Editor::DataStorage::IUiProvider& DataStorageUi) const
 {
 	using namespace UE::Editor::DataStorage::Queries;
-	DataStorageUi.RegisterWidgetFactory<FOutlinerLabelWidgetConstructor>(TEXT("SceneOutliner.RowLabel"),
+	DataStorageUi.RegisterWidgetFactory<FOutlinerLabelWidgetConstructor>(
+		DataStorageUi.FindPurpose(UE::Editor::DataStorage::IUiProvider::FPurposeInfo("SceneOutliner", "RowLabel", NAME_None).GeneratePurposeID()),
 		TColumn<FTypedElementLabelColumn>() && (TColumn<FTypedElementClassTypeInfoColumn>() || TColumn<FTypedElementScriptStructTypeInfoColumn>()));
 }
 
 void UOutlinerLabelWidgetFactory::RegisterWidgetPurposes(UE::Editor::DataStorage::IUiProvider& DataStorageUi) const
 {
-	using namespace UE::Editor::DataStorage;
-	DataStorageUi.RegisterWidgetPurpose(IconWidgetPurpose, IUiProvider::EPurposeType::UniqueByNameAndColumn,
-		LOCTEXT("IconItemCellWidgetPurpose", "The icon widget to use in cells for the Scene Outliner specific to the Item label column.")
-	);
-
-	DataStorageUi.RegisterWidgetPurpose(TextWidgetPurpose, IUiProvider::EPurposeType::UniqueByNameAndColumn,
-		LOCTEXT("TextItemCellWidgetPurpose", "The text widget to use in cells for the Scene Outliner specific to the Item label column.")
-	);
+	DataStorageUi.RegisterWidgetPurpose(
+	UE::Editor::DataStorage::IUiProvider::FPurposeInfo("SceneOutliner", "RowLabel", "Icon",
+		UE::Editor::DataStorage::IUiProvider::EPurposeType::UniqueByNameAndColumn,
+		LOCTEXT("IconItemCellWidgetPurpose", "The icon widget to use in cells for the Scene Outliner specific to the Item label column.")));
+
+	DataStorageUi.RegisterWidgetPurpose(
+	UE::Editor::DataStorage::IUiProvider::FPurposeInfo("SceneOutliner", "RowLabel", "Text",
+		UE::Editor::DataStorage::IUiProvider::EPurposeType::UniqueByNameAndColumn,
+		LOCTEXT("TextItemCellWidgetPurpose", "The text widget to use in cells for the Scene Outliner specific to the Item label column.")));
 }
 
 static void GetAllColumns(TArray<TWeakObjectPtr<const UScriptStruct>>& OutColumns, const UE::Editor::DataStorage::ICoreProvider& DataStorage, UE::Editor::DataStorage::RowHandle Row)
@@ -50,11 +49,9 @@ FOutlinerLabelWidgetConstructor::FOutlinerLabelWidgetConstructor()
 }
 
 TSharedPtr<SWidget> FOutlinerLabelWidgetConstructor::CreateWidget(UE::Editor::DataStorage::ICoreProvider* DataStorage,
-	UE::Editor::DataStorage::IUiProvider* DataStorageUi,
-	UE::Editor::DataStorage::RowHandle TargetRow,
+	UE::Editor::DataStorage::IUiProvider* DataStorageUi, UE::Editor::DataStorage::RowHandle TargetRow,
 	UE::Editor::DataStorage::RowHandle WidgetRow, const UE::Editor::DataStorage::FMetaDataView& Arguments)
 {
-	using namespace UE::Editor::DataStorage;
 	if (DataStorage->IsRowAvailable(TargetRow))
 	{
 		TSharedPtr<FTypedElementWidgetConstructor> OutIconWidgetConstructorPtr;
@@ -74,9 +71,13 @@ TSharedPtr<SWidget> FOutlinerLabelWidgetConstructor::CreateWidget(UE::Editor::Da
 
 		TArray<TWeakObjectPtr<const UScriptStruct>> Columns;
 		GetAllColumns(Columns, *DataStorage, TargetRow);
-		DataStorageUi->CreateWidgetConstructors(IconWidgetPurpose, IUiProvider::EMatchApproach::LongestMatch, Columns, Arguments, AssignIconWidgetToColumn);
+		DataStorageUi->CreateWidgetConstructors(
+			DataStorageUi->FindPurpose(UE::Editor::DataStorage::IUiProvider::FPurposeInfo("SceneOutliner", "RowLabel", "Icon").GeneratePurposeID()),
+			UE::Editor::DataStorage::IUiProvider::EMatchApproach::LongestMatch, Columns, Arguments, AssignIconWidgetToColumn);
 		GetAllColumns(Columns, *DataStorage, TargetRow);
-		DataStorageUi->CreateWidgetConstructors(TextWidgetPurpose, IUiProvider::EMatchApproach::LongestMatch, Columns, Arguments, AssignTextWidgetToColumn);
+		DataStorageUi->CreateWidgetConstructors(
+			DataStorageUi->FindPurpose(UE::Editor::DataStorage::IUiProvider::FPurposeInfo("SceneOutliner", "RowLabel", "Text").GeneratePurposeID()),
+			UE::Editor::DataStorage::IUiProvider::EMatchApproach::LongestMatch, Columns, Arguments, AssignTextWidgetToColumn);
 
 		static const FMargin ColumnItemPadding(8, 0);
 		constexpr float IconNameHorizontalPadding = 8.f;
diff --git a/Engine/Plugins/Experimental/EditorDataStorageFeatures/Source/TedsOutliner/Private/Widgets/OutlinerTextWidget.cpp b/Engine/Plugins/Experimental/EditorDataStorageFeatures/Source/TedsOutliner/Private/Widgets/OutlinerTextWidget.cpp
--- a/Engine/Plugins/Experimental/EditorDataStorageFeatures/Source/TedsOutliner/Private/Widgets/OutlinerTextWidget.cpp
+++ b/Engine/Plugins/Experimental/EditorDataStorageFeatures/Source/TedsOutliner/Private/Widgets/OutlinerTextWidget.cpp
@@ -20,7 +20,8 @@ void UOutlinerTextWidgetFactory::RegisterWidgetConstructors(UE::Editor::DataStor
 {
 	using namespace UE::Editor::DataStorage::Queries;
 
-	DataStorageUi.RegisterWidgetFactory<FOutlinerTextWidgetConstructor>(TEXT("SceneOutliner.RowLabel.Text"),
+	DataStorageUi.RegisterWidgetFactory<FOutlinerTextWidgetConstructor>(DataStorageUi.FindPurpose(
+		IUiProvider::FPurposeInfo("SceneOutliner", "RowLabel", "Text").GeneratePurposeID()),
 		TColumn<FTypedElementLabelColumn>());
 }
 
diff --git a/Engine/Plugins/Experimental/EditorDataStorageFeatures/Source/TedsOutliner/Private/Widgets/OutlinerVisibilityWidget.cpp b/Engine/Plugins/Experimental/EditorDataStorageFeatures/Source/TedsOutliner/Private/Widgets/OutlinerVisibilityWidget.cpp
--- a/Engine/Plugins/Experimental/EditorDataStorageFeatures/Source/TedsOutliner/Private/Widgets/OutlinerVisibilityWidget.cpp
+++ b/Engine/Plugins/Experimental/EditorDataStorageFeatures/Source/TedsOutliner/Private/Widgets/OutlinerVisibilityWidget.cpp
@@ -21,9 +21,6 @@
 
 #define LOCTEXT_NAMESPACE "OutlinerVisibilityWidget"
 
-FName UOutlinerVisibilityWidgetFactory::VisibilityWidgetPurpose = FName(TEXT("SceneOutliner.Cell"));
-FName UOutlinerVisibilityHeaderFactory::VisibilityHeaderPurpose = FName(TEXT("SceneOutliner.Header"));
-
 //
 // Cell Factory
 //
@@ -37,7 +34,9 @@ void UOutlinerVisibilityWidgetFactory::RegisterWidgetConstructors(UE::Editor::Da
 	UE::Editor::DataStorage::IUiProvider& DataStorageUi) const
 {
 	using namespace UE::Editor::DataStorage::Queries;
-	DataStorageUi.RegisterWidgetFactory<FOutlinerVisibilityWidgetConstructor>(VisibilityWidgetPurpose, TColumn<FVisibleInEditorColumn>());
+	DataStorageUi.RegisterWidgetFactory<FOutlinerVisibilityWidgetConstructor>(
+		DataStorageUi.FindPurpose(IUiProvider::FPurposeInfo("SceneOutliner", "Cell", NAME_None).GeneratePurposeID()),
+		TColumn<FVisibleInEditorColumn>());
 	
 }
 
@@ -68,7 +67,8 @@ void UOutlinerVisibilityHeaderFactory::RegisterWidgetConstructors(UE::Editor::Da
 	UE::Editor::DataStorage::IUiProvider& DataStorageUi) const
 {
 	using namespace UE::Editor::DataStorage::Queries;
-	DataStorageUi.RegisterWidgetFactory<FOutlinerVisibilityHeaderConstructor>(VisibilityHeaderPurpose, TColumn<FVisibleInEditorColumn>());
+	DataStorageUi.RegisterWidgetFactory<FOutlinerVisibilityHeaderConstructor>(
+		DataStorageUi.FindPurpose(IUiProvider::FPurposeInfo("SceneOutliner", "Header", NAME_None).GeneratePurposeID()), TColumn<FVisibleInEditorColumn>());
 }
 
 TSharedPtr<SWidget> FOutlinerVisibilityHeaderConstructor::CreateWidget(UE::Editor::DataStorage::ICoreProvider* DataStorage,
diff --git a/Engine/Plugins/Experimental/EditorDataStorageFeatures/Source/TedsOutliner/Private/Widgets/TransformHeadsUpWidget.cpp b/Engine/Plugins/Experimental/EditorDataStorageFeatures/Source/TedsOutliner/Private/Widgets/TransformHeadsUpWidget.cpp
--- a/Engine/Plugins/Experimental/EditorDataStorageFeatures/Source/TedsOutliner/Private/Widgets/TransformHeadsUpWidget.cpp
+++ b/Engine/Plugins/Experimental/EditorDataStorageFeatures/Source/TedsOutliner/Private/Widgets/TransformHeadsUpWidget.cpp
@@ -190,7 +190,8 @@ void UTransformHeadsUpWidgetFactory::RegisterWidgetConstructors(UE::Editor::Data
 {
 	using namespace UE::Editor::DataStorage::Queries;
 
-	DataStorageUi.RegisterWidgetFactory<FTransformHeadsUpWidgetConstructor>(FName(TEXT("SceneOutliner.Cell")),
+	DataStorageUi.RegisterWidgetFactory<FTransformHeadsUpWidgetConstructor>(
+		DataStorageUi.FindPurpose(IUiProvider::FPurposeInfo("SceneOutliner", "Cell", NAME_None).GeneratePurposeID()),
 		TColumn<FTypedElementLocalTransformColumn>());
 }
 
diff --git a/Engine/Plugins/Experimental/EditorDataStorageFeatures/Source/TedsOutliner/Private/Widgets/TypeInfoWidget.cpp b/Engine/Plugins/Experimental/EditorDataStorageFeatures/Source/TedsOutliner/Private/Widgets/TypeInfoWidget.cpp
--- a/Engine/Plugins/Experimental/EditorDataStorageFeatures/Source/TedsOutliner/Private/Widgets/TypeInfoWidget.cpp
+++ b/Engine/Plugins/Experimental/EditorDataStorageFeatures/Source/TedsOutliner/Private/Widgets/TypeInfoWidget.cpp
@@ -18,7 +18,7 @@ void UTypeInfoWidgetFactory::RegisterWidgetConstructors(UE::Editor::DataStorage:
 {
 	using namespace UE::Editor::DataStorage::Queries;
 
-	DataStorageUi.RegisterWidgetFactory<FTypeInfoWidgetConstructor>(FName(TEXT("General.Cell")),
+	DataStorageUi.RegisterWidgetFactory<FTypeInfoWidgetConstructor>(DataStorageUi.FindPurpose(DataStorageUi.GetGeneralWidgetPurposeID()),
 		TColumn<FTypedElementClassTypeInfoColumn>());
 
 }
diff --git a/Engine/Plugins/Experimental/EditorDataStorageFeatures/Source/TedsOutliner/Public/Compatibility/SceneOutlinerTedsBridge.h b/Engine/Plugins/Experimental/EditorDataStorageFeatures/Source/TedsOutliner/Public/Compatibility/SceneOutlinerTedsBridge.h
--- a/Engine/Plugins/Experimental/EditorDataStorageFeatures/Source/TedsOutliner/Public/Compatibility/SceneOutlinerTedsBridge.h
+++ b/Engine/Plugins/Experimental/EditorDataStorageFeatures/Source/TedsOutliner/Public/Compatibility/SceneOutlinerTedsBridge.h
@@ -34,16 +34,17 @@ class TEDSOUTLINER_API FSceneOutlinerTedsQueryBinder
 public:
 	static const FName CellWidgetTableName;
 	static const FName HeaderWidgetPurpose;
-	static const FName DefaultHeaderWidgetPurpose;
 	static const FName CellWidgetPurpose;
-	static const FName DefaultCellWidgetPurpose;
 	static const FName ItemLabelCellWidgetPurpose;
-	static const FName DefaultItemLabelCellWidgetPurpose;
 
 
 	static FSceneOutlinerTedsQueryBinder& GetInstance();
 
-	void AssignQuery(UE::Editor::DataStorage::QueryHandle Query, const TSharedPtr<ISceneOutliner>& Widget, TConstArrayView<FName> InCellWidgetPurposes);
+	void AssignQuery(UE::Editor::DataStorage::QueryHandle Query, const TSharedPtr<ISceneOutliner>& Widget,
+		const UE::Editor::DataStorage::IUiProvider::FPurposeID& InWidgetPurpose, const UE::Editor::DataStorage::IUiProvider::FPurposeID& InHeaderPurpose);
+
+	// Assign query and use default widget purposes for the scene outliner
+	void AssignQuery(UE::Editor::DataStorage::QueryHandle Query, const TSharedPtr<ISceneOutliner>& Outliner);
 
 	// Register a dealiaser for a specific TEDS-Outliner to convert a row handle to an FSceneOutlinerTreeItemID
 	void RegisterTreeItemIDDealiaser(const TSharedPtr<ISceneOutliner>& Widget, const FTreeItemIDDealiaser& InDealiaser);
diff --git a/Engine/Plugins/Experimental/EditorDataStorageFeatures/Source/TedsOutliner/Public/TedsOutlinerImpl.h b/Engine/Plugins/Experimental/EditorDataStorageFeatures/Source/TedsOutliner/Public/TedsOutlinerImpl.h
--- a/Engine/Plugins/Experimental/EditorDataStorageFeatures/Source/TedsOutliner/Public/TedsOutlinerImpl.h
+++ b/Engine/Plugins/Experimental/EditorDataStorageFeatures/Source/TedsOutliner/Public/TedsOutlinerImpl.h
@@ -6,6 +6,7 @@
 #include "ISceneOutlinerHierarchy.h"
 #include "ISceneOutlinerMode.h"
 #include "Elements/Columns/TypedElementHiearchyColumns.h"
+#include "Elements/Interfaces/TypedElementDataStorageUiInterface.h"
 
 namespace UE::Editor::DataStorage
 {
@@ -79,15 +80,7 @@ struct FTedsOutlinerHierarchyData
 
 struct FTedsOutlinerParams
 {
-	FTedsOutlinerParams(SSceneOutliner* InSceneOutliner)
-	: SceneOutliner(InSceneOutliner)
-	, QueryDescription()
-	, bUseDefaultTedsFilters(false)
-	, bShowRowHandleColumn(true)
-	, bForceShowParents(true)
-	, HierarchyData(FTedsOutlinerHierarchyData::GetDefaultHierarchyData())
-	, CellWidgetPurposes{TEXT("SceneOutliner.Cell"), TEXT("General.Cell")}
-	{}
+	TEDSOUTLINER_API FTedsOutlinerParams(SSceneOutliner* InSceneOutliner);
 
 	SSceneOutliner* SceneOutliner;
 
@@ -115,8 +108,14 @@ struct FTedsOutlinerParams
 	// The selection set to use for this Outliner, unset = don't propagate tree selection to the TEDS column
 	TOptional<FName> SelectionSetOverride;
 
-	// The purposes to use when generating widgets for the columns through TEDS UI
-	TArray<FName> CellWidgetPurposes;
+	// The purpose to use when generating widgets for row/column pairs through TEDS UI
+	DataStorage::IUiProvider::FPurposeID CellWidgetPurpose;
+
+	// The purpose to use when generating widgets for column headers through TEDS UI
+	DataStorage::IUiProvider::FPurposeID HeaderWidgetPurpose;
+
+	// The purpose to use when generating widgets for the "Item Label" column through TEDS UI
+	DataStorage::IUiProvider::FPurposeID LabelWidgetPurpose;
 };
 
 
@@ -206,8 +205,11 @@ protected:
 
 	FTedsOutlinerParams CreationParams;
 
-	// Widget purposes this table viewer supports
-	TArray<FName> CellWidgetPurposes;
+	// The purpose to use when generating widgets for row/column pairs through TEDS UI
+	DataStorage::IUiProvider::FPurposeID CellWidgetPurpose;
+
+	// The purpose to use when generating widgets for the "Item Label" column through TEDS UI
+	DataStorage::IUiProvider::FPurposeID LabelWidgetPurpose;
 	
 	// Initial query provided by user
 	TAttribute<DataStorage::FQueryDescription> InitialQueryDescription;
diff --git a/Engine/Plugins/Experimental/EditorDataStorageFeatures/Source/TedsOutliner/Public/Widgets/OutlinerVisibilityWidget.h b/Engine/Plugins/Experimental/EditorDataStorageFeatures/Source/TedsOutliner/Public/Widgets/OutlinerVisibilityWidget.h
--- a/Engine/Plugins/Experimental/EditorDataStorageFeatures/Source/TedsOutliner/Public/Widgets/OutlinerVisibilityWidget.h
+++ b/Engine/Plugins/Experimental/EditorDataStorageFeatures/Source/TedsOutliner/Public/Widgets/OutlinerVisibilityWidget.h
@@ -12,6 +12,13 @@
 
 #include "OutlinerVisibilityWidget.generated.h"
 
+namespace UE::Editor::DataStorage
+{
+	class ICoreProvider;
+}
+
+class UScriptStruct;
+
 UCLASS()
 class UOutlinerVisibilityWidgetFactory : public UEditorDataStorageFactory
 {
@@ -20,12 +27,7 @@ class UOutlinerVisibilityWidgetFactory : public UEditorDataStorageFactory
 public:
 	~UOutlinerVisibilityWidgetFactory() override = default;
 
-	TEDSOUTLINER_API void RegisterWidgetConstructors(UE::Editor::DataStorage::ICoreProvider& DataStorage,
-		UE::Editor::DataStorage::IUiProvider& DataStorageUi) const override;
-
-public:
-
-	static FName VisibilityWidgetPurpose;
+	TEDSOUTLINER_API void RegisterWidgetConstructors(UE::Editor::DataStorage::ICoreProvider& DataStorage, UE::Editor::DataStorage::IUiProvider& DataStorageUi) const override;
 };
 
 UCLASS()
@@ -36,12 +38,7 @@ class UOutlinerVisibilityHeaderFactory : public UEditorDataStorageFactory
 public:
 	~UOutlinerVisibilityHeaderFactory() override = default;
 
-	TEDSOUTLINER_API void RegisterWidgetConstructors(UE::Editor::DataStorage::ICoreProvider& DataStorage,
-		UE::Editor::DataStorage::IUiProvider& DataStorageUi) const override;
-
-public:
-
-	static FName VisibilityHeaderPurpose;
+	TEDSOUTLINER_API void RegisterWidgetConstructors(UE::Editor::DataStorage::ICoreProvider& DataStorage, UE::Editor::DataStorage::IUiProvider& DataStorageUi) const override;
 };
 
 USTRUCT()
@@ -131,17 +128,17 @@ protected:
 	/** Set the item this widget is responsible for to be hidden or shown */
 	void SetIsVisible(RowHandle InRow, const bool bVisible);
 
-	RowHandle TargetRow;
-	RowHandle WidgetRow;
+	RowHandle TargetRow = UE::Editor::DataStorage::InvalidRowHandle;
+	RowHandle WidgetRow = UE::Editor::DataStorage::InvalidRowHandle;
 
 	/** Scoped undo transaction */
 	TUniquePtr<FScopedTransaction> UndoTransaction;
 
 	/** Visibility brushes for the various states */
-	const FSlateBrush* VisibleHoveredBrush;
-	const FSlateBrush* VisibleNotHoveredBrush;
-	const FSlateBrush* NotVisibleHoveredBrush;
-	const FSlateBrush* NotVisibleNotHoveredBrush;
+	const FSlateBrush* VisibleHoveredBrush = nullptr;
+	const FSlateBrush* VisibleNotHoveredBrush = nullptr;
+	const FSlateBrush* NotVisibleHoveredBrush = nullptr;
+	const FSlateBrush* NotVisibleNotHoveredBrush = nullptr;
 
 private:
 
diff --git a/Engine/Plugins/Experimental/EditorDataStorageFeatures/Source/TedsSettings/Private/TedsSettingsFactory.cpp b/Engine/Plugins/Experimental/EditorDataStorageFeatures/Source/TedsSettings/Private/TedsSettingsFactory.cpp
--- a/Engine/Plugins/Experimental/EditorDataStorageFeatures/Source/TedsSettings/Private/TedsSettingsFactory.cpp
+++ b/Engine/Plugins/Experimental/EditorDataStorageFeatures/Source/TedsSettings/Private/TedsSettingsFactory.cpp
@@ -9,7 +9,9 @@ void UTedsSettingsFactory::RegisterWidgetConstructors(UE::Editor::DataStorage::I
 {
 	using namespace UE::Editor::DataStorage::Queries;
 
-	DataStorageUi.RegisterWidgetFactory<FSettingsContainerReferenceWidgetConstructor>(FName(TEXT("General.Cell")), TColumn<FSettingsContainerReferenceColumn>());
-	DataStorageUi.RegisterWidgetFactory<FSettingsCategoryReferenceWidgetConstructor>(FName(TEXT("General.Cell")), TColumn<FSettingsCategoryReferenceColumn>());
-	DataStorageUi.RegisterWidgetFactory<FSettingsSectionWidgetConstructor>(FName(TEXT("General.Cell")), TColumn<FSettingsSectionTag>());
+	RowHandle PurposeRow = DataStorageUi.FindPurpose(DataStorageUi.GetGeneralWidgetPurposeID());
+	
+	DataStorageUi.RegisterWidgetFactory<FSettingsContainerReferenceWidgetConstructor>(PurposeRow, TColumn<FSettingsContainerReferenceColumn>());
+	DataStorageUi.RegisterWidgetFactory<FSettingsCategoryReferenceWidgetConstructor>(PurposeRow, TColumn<FSettingsCategoryReferenceColumn>());
+	DataStorageUi.RegisterWidgetFactory<FSettingsSectionWidgetConstructor>(PurposeRow, TColumn<FSettingsSectionTag>());
 }
diff --git a/Engine/Plugins/Experimental/EditorDataStorageFeatures/Source/TedsTableViewer/Private/TedsTableViewerColumn.cpp b/Engine/Plugins/Experimental/EditorDataStorageFeatures/Source/TedsTableViewer/Private/TedsTableViewerColumn.cpp
--- a/Engine/Plugins/Experimental/EditorDataStorageFeatures/Source/TedsTableViewer/Private/TedsTableViewerColumn.cpp
+++ b/Engine/Plugins/Experimental/EditorDataStorageFeatures/Source/TedsTableViewer/Private/TedsTableViewerColumn.cpp
@@ -261,7 +261,7 @@ namespace UE::Editor::DataStorage
 			bool bMatchesQueryConditions = true;
 
 			// First we try to match against the conditions provided by the widget constructor if possible
-			const Queries::FConditions* WidgetConstructorConditions = CellWidgetConstructor->GetQueryConditions();
+			const Queries::FConditions* WidgetConstructorConditions = CellWidgetConstructor->GetQueryConditions(Storage);
 			if(WidgetConstructorConditions && WidgetConstructorConditions->IsCompiled())
 			{
 				bMatchesQueryConditions = bMatchesQueryConditions && Storage->MatchesColumns(Pair.Key, *WidgetConstructorConditions);
diff --git a/Engine/Plugins/Experimental/EditorDataStorageFeatures/Source/TedsTableViewer/Private/TedsTableViewerModel.cpp b/Engine/Plugins/Experimental/EditorDataStorageFeatures/Source/TedsTableViewer/Private/TedsTableViewerModel.cpp
--- a/Engine/Plugins/Experimental/EditorDataStorageFeatures/Source/TedsTableViewer/Private/TedsTableViewerModel.cpp
+++ b/Engine/Plugins/Experimental/EditorDataStorageFeatures/Source/TedsTableViewer/Private/TedsTableViewerModel.cpp
@@ -15,11 +15,12 @@ DEFINE_LOG_CATEGORY_STATIC(LogTableViewer, Log, All)
 namespace UE::Editor::DataStorage
 {
 	FTedsTableViewerModel::FTedsTableViewerModel(const TSharedPtr<QueryStack::IRowNode>& InRowQueryStack,
-		const TArray<TWeakObjectPtr<const UScriptStruct>>& InRequestedColumns, const TArray<FName>& InCellWidgetPurposes,
-		const FIsItemVisible& InIsItemVisibleDelegate)
+		const TArray<TWeakObjectPtr<const UScriptStruct>>& InRequestedColumns, const IUiProvider::FPurposeID& InCellWidgetPurpose,
+		const IUiProvider::FPurposeID& InHeaderWidgetPurpose, const FIsItemVisible& InIsItemVisibleDelegate)
 		: RowQueryStack(InRowQueryStack)
 		, RequestedTedsColumns(InRequestedColumns)
-		, CellWidgetPurposes(InCellWidgetPurposes)
+		, CellWidgetPurpose(InCellWidgetPurpose)
+		, HeaderWidgetPurpose(InHeaderWidgetPurpose)
 		, IsItemVisible(InIsItemVisibleDelegate)
 	{
 		Storage = GetMutableDataStorageFeature<ICoreProvider>(StorageFeatureName);
@@ -164,6 +165,9 @@ namespace UE::Editor::DataStorage
 	void FTedsTableViewerModel::GenerateColumns()
 	{
 		ValidateRequestedColumns();
+
+		RowHandle CellWidgetPurposeRow = StorageUi->FindPurpose(CellWidgetPurpose);
+		RowHandle HeaderWidgetPurposeRow = StorageUi->FindPurpose(HeaderWidgetPurpose);
 		
 		using MatchApproach = IUiProvider::EMatchApproach;
 		int32 IndexOffset = 0;
@@ -177,13 +181,13 @@ namespace UE::Editor::DataStorage
 		TArray<TWeakObjectPtr<const UScriptStruct>> ColumnsCopy = RequestedTedsColumns;
 		
 		// Lambda to create the constructor for a given list of columns
-		auto ColumnConstructor = [this, &IndexOffset, &NewColumnMap](
+		auto ColumnConstructor = [this, &IndexOffset, &NewColumnMap, HeaderWidgetPurposeRow](
 			TUniquePtr<FTypedElementWidgetConstructor> Constructor, TConstArrayView<TWeakObjectPtr<const UScriptStruct>> MatchedColumns)
 			{
 				TSharedPtr<FTypedElementWidgetConstructor> CellConstructor(Constructor.Release());
 
 				TSharedPtr<FTypedElementWidgetConstructor> HeaderConstructor =
-					TableViewerUtils::CreateHeaderWidgetConstructor(*StorageUi, FMetaDataView(), MatchedColumns, CellWidgetPurposes);
+					TableViewerUtils::CreateHeaderWidgetConstructor(*StorageUi, FMetaDataView(), MatchedColumns, HeaderWidgetPurposeRow);
 			
 				FName NameId = TableViewerUtils::FindLongestMatchingName(MatchedColumns, IndexOffset);
 
@@ -202,56 +206,9 @@ namespace UE::Editor::DataStorage
 			};
 
 		// Create the widget constructors for the columns
-		for(const FName& WidgetPurpose : CellWidgetPurposes)
-		{
-			StorageUi->CreateWidgetConstructors(WidgetPurpose, MatchApproach::LongestMatch, ColumnsCopy, 
+		StorageUi->CreateWidgetConstructors(CellWidgetPurposeRow, MatchApproach::LongestMatch, ColumnsCopy, 
 			FMetaDataView(), ColumnConstructor);
-		}
-
-		// For any remaining columns, we'll try to find and use any default widgets
-		for (TWeakObjectPtr<const UScriptStruct>& ColumnType : ColumnsCopy)
-		{
-			auto AssignWidgetToColumn = [this, ColumnType, &IndexOffset, &NewColumnMap](
-				TUniquePtr<FTypedElementWidgetConstructor> Constructor, TConstArrayView<TWeakObjectPtr<const UScriptStruct>>)
-			{
-				TSharedPtr<FTypedElementWidgetConstructor> CellConstructor(Constructor.Release());
-
-				TSharedPtr<FTypedElementWidgetConstructor> HeaderConstructor =
-					TableViewerUtils::CreateHeaderWidgetConstructor(*StorageUi, FMetaDataView(), {ColumnType}, CellWidgetPurposes);
-
-				FName NameId = FName(ColumnType->GetDisplayNameText().ToString());
-
-				TArray<TWeakObjectPtr<const UScriptStruct>> MatchedColumnsCopy({ColumnType});
-				TSharedRef<FTedsTableViewerColumn> Column = MakeShared<FTedsTableViewerColumn>(NameId, CellConstructor, MoveTemp(MatchedColumnsCopy), HeaderConstructor);
-				
-				Column->SetIsRowVisibleDelegate(FTedsTableViewerColumn::FIsRowVisible::CreateRaw(this, &FTedsTableViewerModel::IsRowVisible));
-
-				NewColumnMap.Emplace(ColumnType, Column);
-				
-				++IndexOffset;
-				return false;
-			};
-
-			const int32 BeforeIndexOffset = IndexOffset;
-			
-			for(const FName& WidgetPurpose : CellWidgetPurposes)
-			{
-				const FName DefaultWidgetPurpose(WidgetPurpose.ToString() + TEXT(".Default"));
-
-				StorageUi->CreateWidgetConstructors(DefaultWidgetPurpose, FMetaDataView(), AssignWidgetToColumn);
-				
-				if (BeforeIndexOffset != IndexOffset)
-				{
-					break;
-				}
-			}
-
-			if (BeforeIndexOffset == IndexOffset)
-			{
-				++IndexOffset;
-			}
-		}
-
+		
 		// Add the actual UI columns in the order the Teds Columns were specified
 		for (const TWeakObjectPtr<const UScriptStruct>& ColumnType : RequestedTedsColumns)
 		{
diff --git a/Engine/Plugins/Experimental/EditorDataStorageFeatures/Source/TedsTableViewer/Private/TedsTableViewerUtils.cpp b/Engine/Plugins/Experimental/EditorDataStorageFeatures/Source/TedsTableViewer/Private/TedsTableViewerUtils.cpp
--- a/Engine/Plugins/Experimental/EditorDataStorageFeatures/Source/TedsTableViewer/Private/TedsTableViewerUtils.cpp
+++ b/Engine/Plugins/Experimental/EditorDataStorageFeatures/Source/TedsTableViewer/Private/TedsTableViewerUtils.cpp
@@ -90,21 +90,15 @@ namespace UE::Editor::DataStorage::TableViewerUtils
 		return VerifiedColumnTypes;
 	}
 
-	TSharedPtr<FTypedElementWidgetConstructor> CreateHeaderWidgetConstructor(IUiProvider& StorageUi, const UE::Editor::DataStorage::FMetaDataView& InMetaData, TConstArrayView<TWeakObjectPtr<const UScriptStruct>> ColumnTypes, const TConstArrayView<FName> CellWidgetPurposes)
+	TSharedPtr<FTypedElementWidgetConstructor> CreateHeaderWidgetConstructor(IUiProvider& StorageUi,
+		const FMetaDataView& InMetaData, TConstArrayView<TWeakObjectPtr<const UScriptStruct>> ColumnTypes, RowHandle PurposeRow)
 	{
 		using MatchApproach = IUiProvider::EMatchApproach;
 
 		TArray<TWeakObjectPtr<const UScriptStruct>> VerifiedColumnTypes = CreateVerifiedColumnTypeArray(ColumnTypes);
-		TSharedPtr<FTypedElementWidgetConstructor> Constructor;
+		TSharedPtr<FTypedElementWidgetConstructor> Constructor = nullptr;
 
-		for (const FName& Purpose : CellWidgetPurposes)
-		{
-			// Extract the header purpose from the cell purpose (e.g "SceneOutliner.ItemLabel.Cell" -> "SceneOutliner.ItemLabel.Header")
-			FString HeaderPurpose, CellPurpose;
-			Purpose.ToString().Split(TEXT(".Cell"), &HeaderPurpose, &CellPurpose);
-			HeaderPurpose.Append(TEXT(".Header"));
-			
-			StorageUi.CreateWidgetConstructors(FName(HeaderPurpose), MatchApproach::ExactMatch, VerifiedColumnTypes, InMetaData,
+		StorageUi.CreateWidgetConstructors(PurposeRow, MatchApproach::ExactMatch, VerifiedColumnTypes, InMetaData,
 				[&Constructor, ColumnTypes](
 					TUniquePtr<FTypedElementWidgetConstructor> CreatedConstructor, 
 					TConstArrayView<TWeakObjectPtr<const UScriptStruct>> MatchedColumnTypes)
@@ -117,32 +111,8 @@ namespace UE::Editor::DataStorage::TableViewerUtils
 					// always be shorter in both cases just return.
 					return false;
 				});
-			if (Constructor)
-			{
-				return Constructor;
-			}
-		}
-		for (const FName& Purpose : CellWidgetPurposes)
-		{
-			// Extract the default header purpose from the cell purpose (e.g "SceneOutliner.ItemLabel.Cell" -> "SceneOutliner.ItemLabel.Header.Default")
-			FString HeaderPurpose, CellPurpose;
-			Purpose.ToString().Split(TEXT(".Cell"), &HeaderPurpose, &CellPurpose);
-			HeaderPurpose.Append(TEXT(".Header.Default"));
-
-			StorageUi.CreateWidgetConstructors(FName(HeaderPurpose), InMetaData,
-				[&Constructor, ColumnTypes](
-					TUniquePtr<FTypedElementWidgetConstructor> CreatedConstructor,
-					TConstArrayView<TWeakObjectPtr<const UScriptStruct>> MatchedColumnTypes)
-				{
-					Constructor = TSharedPtr<FTypedElementWidgetConstructor>(CreatedConstructor.Release());
-					return false;
-				});
-			if (Constructor)
-			{
-				return Constructor;
-			}
-		}
-		return nullptr;
+		
+		return Constructor;
 	}
 
 	const FSlateBrush* GetIconForRow(ICoreProvider* DataStorage, RowHandle Row)
diff --git a/Engine/Plugins/Experimental/EditorDataStorageFeatures/Source/TedsTableViewer/Private/Widgets/RowHandleWidget.cpp b/Engine/Plugins/Experimental/EditorDataStorageFeatures/Source/TedsTableViewer/Private/Widgets/RowHandleWidget.cpp
--- a/Engine/Plugins/Experimental/EditorDataStorageFeatures/Source/TedsTableViewer/Private/Widgets/RowHandleWidget.cpp
+++ b/Engine/Plugins/Experimental/EditorDataStorageFeatures/Source/TedsTableViewer/Private/Widgets/RowHandleWidget.cpp
@@ -11,22 +11,32 @@
 
 void URowHandleWidgetFactory::RegisterWidgetConstructors(UE::Editor::DataStorage::ICoreProvider& DataStorage, UE::Editor::DataStorage::IUiProvider& DataStorageUi) const
 {
-	DataStorageUi.RegisterWidgetFactory(FName(TEXT("General.Cell.RowHandle")), FRowHandleWidgetConstructor::StaticStruct());
+	DataStorageUi.RegisterWidgetFactory(DataStorageUi.FindPurpose(
+		UE::Editor::DataStorage::IUiProvider::FPurposeInfo("General", "Cell", "RowHandle").GeneratePurposeID()),
+		FRowHandleWidgetConstructor::StaticStruct());
 }
 
 void URowHandleWidgetFactory::RegisterWidgetPurposes(UE::Editor::DataStorage::IUiProvider& DataStorageUi) const
 {
-	using namespace UE::Editor::DataStorage;
-	DataStorageUi.RegisterWidgetPurpose(FName(TEXT("General.Cell.RowHandle")), IUiProvider::EPurposeType::UniqueByName,
-	LOCTEXT("GeneralRowHandlePurpose", "Specific purpose to request a widget to display row handles."));
-
-	DataStorageUi.RegisterWidgetPurpose(FName(TEXT("RowDetails.Cell")), IUiProvider::EPurposeType::UniqueByNameAndColumn,
-	LOCTEXT("DetailsRowHandlePurpose", "Specific purpose to request a widget to display the details on a row (e.g SRowDetails)."));
-
-	DataStorageUi.RegisterWidgetPurpose(FName(TEXT("RowDetails.Cell.Large")), IUiProvider::EPurposeType::UniqueByNameAndColumn,
-	LOCTEXT("LargeDetailsRowHandlePurpose", "Specific purpose to request a widget that is larger than a single cell to display the details on a row (e.g SRowDetails)"));
-
-
+	DataStorageUi.RegisterWidgetPurpose(
+		UE::Editor::DataStorage::IUiProvider::FPurposeInfo("General", "Cell", "RowHandle",
+			UE::Editor::DataStorage::IUiProvider::EPurposeType::UniqueByName,
+			LOCTEXT("GeneralRowHandlePurpose", "Specific purpose to request a widget to display row handles.")));
+
+	DataStorageUi.RegisterWidgetPurpose(
+		UE::Editor::DataStorage::IUiProvider::FPurposeInfo("RowDetails", "Cell", NAME_None,
+			UE::Editor::DataStorage::IUiProvider::EPurposeType::UniqueByName,
+			LOCTEXT("DetailsRowHandlePurpose", "Specific purpose to request a widget to display the details on a row (e.g SRowDetails)."),
+			DataStorageUi.GetDefaultWidgetPurposeID()));
+	
+	UE::Editor::DataStorage::IUiProvider::FPurposeID GeneralLargePurposeID =
+		UE::Editor::DataStorage::IUiProvider::FPurposeInfo("General", "Cell", "Large").GeneratePurposeID();
+	
+	DataStorageUi.RegisterWidgetPurpose(
+		UE::Editor::DataStorage::IUiProvider::FPurposeInfo("RowDetails", "Cell", "Large",
+			UE::Editor::DataStorage::IUiProvider::EPurposeType::UniqueByNameAndColumn,
+			LOCTEXT("GeneralRowHandlePurpose", "Specific purpose to request a widget to display row handles."),
+			GeneralLargePurposeID));
 }
 
 FRowHandleWidgetConstructor::FRowHandleWidgetConstructor()
@@ -42,9 +52,7 @@ TSharedPtr<SWidget> FRowHandleWidgetConstructor::CreateWidget(const UE::Editor::
 		.VAlign(VAlign_Center)
 		.Padding(8, 0, 0, 0);}
 
-bool FRowHandleWidgetConstructor::FinalizeWidget(UE::Editor::DataStorage::ICoreProvider* DataStorage,
-	UE::Editor::DataStorage::IUiProvider* DataStorageUi, 
-	UE::Editor::DataStorage::RowHandle Row, const TSharedPtr<SWidget>& Widget)
+bool FRowHandleWidgetConstructor::FinalizeWidget(UE::Editor::DataStorage::ICoreProvider* DataStorage, UE::Editor::DataStorage::IUiProvider* DataStorageUi, UE::Editor::DataStorage::RowHandle Row, const TSharedPtr<SWidget>& Widget)
 {
 
 	checkf(Widget->GetType() == SBox::StaticWidgetClass().GetWidgetType(),
diff --git a/Engine/Plugins/Experimental/EditorDataStorageFeatures/Source/TedsTableViewer/Private/Widgets/SRowDetails.cpp b/Engine/Plugins/Experimental/EditorDataStorageFeatures/Source/TedsTableViewer/Private/Widgets/SRowDetails.cpp
--- a/Engine/Plugins/Experimental/EditorDataStorageFeatures/Source/TedsTableViewer/Private/Widgets/SRowDetails.cpp
+++ b/Engine/Plugins/Experimental/EditorDataStorageFeatures/Source/TedsTableViewer/Private/Widgets/SRowDetails.cpp
@@ -18,8 +18,6 @@ namespace UE::Editor::DataStorage
 	{
 		static const FName NameColumn(TEXT("Name"));
 		static const FName DataColumn(TEXT("Data"));
-		static TArray<FName> DefaultWidgetPurposes = { FName(TEXT("RowDetails.Cell.Large")),
-			FName(TEXT("RowDetails.Cell")), FName(TEXT("General.Cell.Large")), FName(TEXT("General.Cell")) };
 	} // namespace Widgets::Private
 
 	//
@@ -29,11 +27,12 @@ namespace UE::Editor::DataStorage
 	void SRowDetails::Construct(const FArguments& InArgs)
 	{
 		bShowAllDetails = InArgs._ShowAllDetails;
-		WidgetPurposes = InArgs._WidgetPurposesOverride;
+		
+		WidgetPurpose = InArgs._WidgetPurposeOverride;
 
-		if(WidgetPurposes.IsEmpty())
+		if(!WidgetPurpose.IsSet())
 		{
-			WidgetPurposes = Widgets::Private::DefaultWidgetPurposes;
+			WidgetPurpose = IUiProvider::FPurposeInfo("RowDetails", "Cell", "Large").GeneratePurposeID();
 		}
 
 		checkf(AreEditorDataStorageFeaturesEnabled(), TEXT("Unable to initialize SRowDetails without the editor data storage interfaces."));
@@ -76,24 +75,23 @@ namespace UE::Editor::DataStorage
 					Columns.Emplace(&ColumnType);
 					return true;
 				});
+
+			RowHandle PurposeRow = DataStorageUi->FindPurpose(WidgetPurpose);
 			
-			for (const FName& Purpose : WidgetPurposes)
-			{
-				DataStorageUi->CreateWidgetConstructors(Purpose, IUiProvider::EMatchApproach::LongestMatch,
+			DataStorageUi->CreateWidgetConstructors(PurposeRow, IUiProvider::EMatchApproach::LongestMatch,
 					Columns, {}, [this, Row](
 						TUniquePtr<FTypedElementWidgetConstructor> Constructor, TConstArrayView<TWeakObjectPtr<const UScriptStruct>> Columns)
 					{
 						Items.Add(MakeShared<FRowDetailsItem>(nullptr, MoveTemp(Constructor), Row));
 						return true;
 					});
-			}
 			
 			if (bShowAllDetails)
 			{
 				// Create defaults for the remaining widgets.
 				for (TWeakObjectPtr<const UScriptStruct> Column : Columns)
 				{
-					DataStorageUi->CreateWidgetConstructors(FName(TEXT("General.Cell.Default")), {},
+					DataStorageUi->CreateWidgetConstructors(DataStorageUi->FindPurpose(DataStorageUi->GetDefaultWidgetPurposeID()), {},
 						[this, Column, Row](
 							TUniquePtr<FTypedElementWidgetConstructor> Constructor, TConstArrayView<TWeakObjectPtr<const UScriptStruct>> Columns)
 						{
@@ -128,7 +126,7 @@ namespace UE::Editor::DataStorage
 	//
 	// SRowDetailsRow
 	//
-	void SRowDetailsRow::Construct(const FArguments& Args, const TSharedRef<STableViewBase>& OwnerTableView,ICoreProvider* InDataStorage,
+	void SRowDetailsRow::Construct(const FArguments& Args, const TSharedRef<STableViewBase>& OwnerTableView, ICoreProvider* InDataStorage,
 			IUiProvider* InDataStorageUi)
 	{
 		Item = Args._Item;
diff --git a/Engine/Plugins/Experimental/EditorDataStorageFeatures/Source/TedsTableViewer/Private/Widgets/STedsTableViewer.cpp b/Engine/Plugins/Experimental/EditorDataStorageFeatures/Source/TedsTableViewer/Private/Widgets/STedsTableViewer.cpp
--- a/Engine/Plugins/Experimental/EditorDataStorageFeatures/Source/TedsTableViewer/Private/Widgets/STedsTableViewer.cpp
+++ b/Engine/Plugins/Experimental/EditorDataStorageFeatures/Source/TedsTableViewer/Private/Widgets/STedsTableViewer.cpp
@@ -5,6 +5,7 @@
 #include "TedsTableViewerColumn.h"
 #include "Columns/SlateDelegateColumns.h"
 #include "Elements/Columns/TypedElementUIColumns.h"
+#include "Elements/Common/EditorDataStorageFeatures.h"
 #include "Elements/Framework/TypedElementAttributeBinding.h"
 #include "Elements/Framework/TypedElementDataStorageWidget.h"
 #include "Widgets/STedsTableViewerRow.h"
@@ -20,7 +21,22 @@ namespace UE::Editor::DataStorage
 		OnSelectionChanged = InArgs._OnSelectionChanged;
 		EmptyRowsMessage = InArgs._EmptyRowsMessage;
 		
-		Model = MakeShared<FTedsTableViewerModel>(InArgs._QueryStack, InArgs._Columns, InArgs._CellWidgetPurposes,
+		IUiProvider::FPurposeID CellWigetPurpose = InArgs._CellWidgetPurpose;
+
+		if (!CellWigetPurpose.IsSet())
+		{
+			IUiProvider* StorageUi = GetMutableDataStorageFeature<IUiProvider>(UiFeatureName);
+			CellWigetPurpose = StorageUi->GetGeneralWidgetPurposeID();
+		}
+
+		IUiProvider::FPurposeID HeaderWidgetPurpose = InArgs._HeaderWidgetPurpose;
+
+		if (!HeaderWidgetPurpose.IsSet())
+		{
+			HeaderWidgetPurpose = IUiProvider::FPurposeInfo("General", "Header", NAME_None).GeneratePurposeID();
+		}
+		
+		Model = MakeShared<FTedsTableViewerModel>(InArgs._QueryStack, InArgs._Columns, CellWigetPurpose, HeaderWidgetPurpose,
 			FTedsTableViewerModel::FIsItemVisible::CreateSP(this, &STedsTableViewer::IsItemVisible));
 		
 		HeaderRowWidget = SNew( SHeaderRow )
diff --git a/Engine/Plugins/Experimental/EditorDataStorageFeatures/Source/TedsTableViewer/Private/Widgets/STedsTileViewer.cpp b/Engine/Plugins/Experimental/EditorDataStorageFeatures/Source/TedsTableViewer/Private/Widgets/STedsTileViewer.cpp
--- a/Engine/Plugins/Experimental/EditorDataStorageFeatures/Source/TedsTableViewer/Private/Widgets/STedsTileViewer.cpp
+++ b/Engine/Plugins/Experimental/EditorDataStorageFeatures/Source/TedsTableViewer/Private/Widgets/STedsTileViewer.cpp
@@ -22,12 +22,20 @@ namespace UE::Editor::DataStorage
 		Columns = InArgs._Columns;
 		TableRowStyle = InArgs._TileStyle;
 
-		TArray<FName> Purposes;
+		// Use the default purpose if the user didn't specify any
+		// TODO: Have a better "default" widget for tiles instead of the default cell widget purpose
+		if (!WidgetPurpose.IsSet())
+		{
+			IUiProvider* StorageUi = GetMutableDataStorageFeature<IUiProvider>(UiFeatureName);
+			WidgetPurpose = StorageUi->GetGeneralWidgetPurposeID();
+		}
+		
 		// While we might have columns for the widget constructor, we don't want to display any UI columns since that doesn't make sense
-		// so we give FTedsTableViewerModel empty dummy columns
+		// so we give FTedsTableViewerModel empty dummy columns and purposes
+		IUiProvider::FPurposeID DummyPurpose;
 		TArray<TWeakObjectPtr<const UScriptStruct>> DummyColumns;
 		
-		Model = MakeShared<FTedsTableViewerModel>(InArgs._QueryStack, DummyColumns, Purposes,
+		Model = MakeShared<FTedsTableViewerModel>(InArgs._QueryStack, DummyColumns, DummyPurpose, DummyPurpose,
 			FTedsTableViewerModel::FIsItemVisible::CreateSP(this, &STedsTileViewer::IsItemVisible));
 
 		IUiProvider* StorageUi = Model->GetDataStorageUiProvider();
@@ -84,14 +92,16 @@ namespace UE::Editor::DataStorage
 		};
 	
 		IUiProvider* StorageUi = GetMutableDataStorageFeature<IUiProvider>(UiFeatureName);
+		
+		RowHandle WidgetPurposeRow = StorageUi->FindPurpose(WidgetPurpose);
 
 		if (Columns.IsEmpty())
 		{
-			StorageUi->CreateWidgetConstructors(WidgetPurpose, FMetaDataView(), AssignWidgetToColumn);
+			StorageUi->CreateWidgetConstructors(WidgetPurposeRow, FMetaDataView(), AssignWidgetToColumn);
 		}
 		else
 		{
-			StorageUi->CreateWidgetConstructors(WidgetPurpose, IUiProvider::EMatchApproach::ExactMatch, Columns,
+			StorageUi->CreateWidgetConstructors(WidgetPurposeRow, IUiProvider::EMatchApproach::ExactMatch, Columns,
 				FMetaDataView(), AssignWidgetToColumn);
 		}
 
diff --git a/Engine/Plugins/Experimental/EditorDataStorageFeatures/Source/TedsTableViewer/Private/Widgets/UObjectLabelWidget.cpp b/Engine/Plugins/Experimental/EditorDataStorageFeatures/Source/TedsTableViewer/Private/Widgets/UObjectLabelWidget.cpp
--- a/Engine/Plugins/Experimental/EditorDataStorageFeatures/Source/TedsTableViewer/Private/Widgets/UObjectLabelWidget.cpp
+++ b/Engine/Plugins/Experimental/EditorDataStorageFeatures/Source/TedsTableViewer/Private/Widgets/UObjectLabelWidget.cpp
@@ -19,7 +19,7 @@ void UUObjectLabelWidgetFactory::RegisterWidgetConstructors(UE::Editor::DataStor
 {
 	using namespace UE::Editor::DataStorage::Queries;
 	DataStorageUi.RegisterWidgetFactory<FUObjectLabelWidgetConstructor>(
-		TEXT("General.RowLabel"),
+		DataStorageUi.FindPurpose(IUiProvider::FPurposeInfo("General", "RowLabel", NAME_None).GeneratePurposeID()),
 		TColumn<FTypedElementLabelColumn>() && TColumn<FTypedElementUObjectColumn>());
 }
 
diff --git a/Engine/Plugins/Experimental/EditorDataStorageFeatures/Source/TedsTableViewer/Public/TedsTableViewerModel.h b/Engine/Plugins/Experimental/EditorDataStorageFeatures/Source/TedsTableViewer/Public/TedsTableViewerModel.h
--- a/Engine/Plugins/Experimental/EditorDataStorageFeatures/Source/TedsTableViewer/Public/TedsTableViewerModel.h
+++ b/Engine/Plugins/Experimental/EditorDataStorageFeatures/Source/TedsTableViewer/Public/TedsTableViewerModel.h
@@ -8,6 +8,7 @@
 #include "Templates/SharedPointer.h"
 #include "TypedElementUITypes.h"
 #include "TedsQueryStackInterfaces.h"
+#include "Elements/Interfaces/TypedElementDataStorageUiInterface.h"
 #include "UObject/NameTypes.h"
 
 namespace UE::Editor::DataStorage
@@ -30,7 +31,7 @@ namespace UE::Editor::DataStorage
 		DECLARE_MULTICAST_DELEGATE(FOnModelChanged);
 
 		TEDSTABLEVIEWER_API FTedsTableViewerModel(const TSharedPtr<QueryStack::IRowNode>& RowQueryStack, const TArray<TWeakObjectPtr<const UScriptStruct>>& RequestedColumns,
-			const TArray<FName>& CellWidgetPurposes, const FIsItemVisible& InIsItemVisibleDelegate);
+			const IUiProvider::FPurposeID& CellWidgetPurpose, const IUiProvider::FPurposeID& HeaderWidgetPurpose, const FIsItemVisible& InIsItemVisibleDelegate);
 
 		~FTedsTableViewerModel();
 
@@ -94,7 +95,8 @@ namespace UE::Editor::DataStorage
 		TArray<TWeakObjectPtr<const UScriptStruct>> RequestedTedsColumns;
 
 		// The widget purposes used to create widgets in this table viewer
-		TArray<FName> CellWidgetPurposes;
+		IUiProvider::FPurposeID CellWidgetPurpose;
+		IUiProvider::FPurposeID HeaderWidgetPurpose;
 
 		// Cached revision ID for the query stack used to check when the table viewer needs a refresh
 		uint32 CachedRowQueryStackRevision = 0;
diff --git a/Engine/Plugins/Experimental/EditorDataStorageFeatures/Source/TedsTableViewer/Public/TedsTableViewerUtils.h b/Engine/Plugins/Experimental/EditorDataStorageFeatures/Source/TedsTableViewer/Public/TedsTableViewerUtils.h
--- a/Engine/Plugins/Experimental/EditorDataStorageFeatures/Source/TedsTableViewer/Public/TedsTableViewerUtils.h
+++ b/Engine/Plugins/Experimental/EditorDataStorageFeatures/Source/TedsTableViewer/Public/TedsTableViewerUtils.h
@@ -29,8 +29,7 @@ namespace UE::Editor::DataStorage::TableViewerUtils
 
 	// Create a header widget constructor for the given columns
 	TEDSTABLEVIEWER_API TSharedPtr<FTypedElementWidgetConstructor> CreateHeaderWidgetConstructor(IUiProvider& StorageUi, 
-	const UE::Editor::DataStorage::FMetaDataView& InMetaData, TConstArrayView<TWeakObjectPtr<const UScriptStruct>> ColumnTypes,
-	const TConstArrayView<FName> CellWidgetPurposes);
+		const FMetaDataView& InMetaData, TConstArrayView<TWeakObjectPtr<const UScriptStruct>> ColumnTypes, RowHandle PurposeRow);
 
 	// Create a copy of the provided column types array after discarding invalid entries
 	TEDSTABLEVIEWER_API TArray<TWeakObjectPtr<const UScriptStruct>> CreateVerifiedColumnTypeArray(TConstArrayView<TWeakObjectPtr<const UScriptStruct>> ColumnTypes);
diff --git a/Engine/Plugins/Experimental/EditorDataStorageFeatures/Source/TedsTableViewer/Public/Widgets/SRowDetails.h b/Engine/Plugins/Experimental/EditorDataStorageFeatures/Source/TedsTableViewer/Public/Widgets/SRowDetails.h
--- a/Engine/Plugins/Experimental/EditorDataStorageFeatures/Source/TedsTableViewer/Public/Widgets/SRowDetails.h
+++ b/Engine/Plugins/Experimental/EditorDataStorageFeatures/Source/TedsTableViewer/Public/Widgets/SRowDetails.h
@@ -2,6 +2,7 @@
 #pragma once
 
 #include "Elements/Common/TypedElementHandles.h"
+#include "Elements/Interfaces/TypedElementDataStorageUiInterface.h"
 #include "Templates/SharedPointer.h"
 #include "Widgets/SCompoundWidget.h"
 #include "Widgets/Views/SListView.h"
@@ -46,7 +47,7 @@ namespace UE::Editor::DataStorage
 			SLATE_ARGUMENT(bool, ShowAllDetails)
 
 			// Override for the default widget purposes used to create widgets for the columns
-			SLATE_ARGUMENT(TArray<FName>, WidgetPurposesOverride)
+			SLATE_ARGUMENT(IUiProvider::FPurposeID, WidgetPurposeOverride)
 
 		
 		SLATE_END_ARGS()
@@ -72,7 +73,7 @@ namespace UE::Editor::DataStorage
 
 		bool bShowAllDetails = true;
 
-		TArray<FName> WidgetPurposes;
+		IUiProvider::FPurposeID WidgetPurpose;
 
 	};
 
diff --git a/Engine/Plugins/Experimental/EditorDataStorageFeatures/Source/TedsTableViewer/Public/Widgets/STedsTableViewer.h b/Engine/Plugins/Experimental/EditorDataStorageFeatures/Source/TedsTableViewer/Public/Widgets/STedsTableViewer.h
--- a/Engine/Plugins/Experimental/EditorDataStorageFeatures/Source/TedsTableViewer/Public/Widgets/STedsTableViewer.h
+++ b/Engine/Plugins/Experimental/EditorDataStorageFeatures/Source/TedsTableViewer/Public/Widgets/STedsTableViewer.h
@@ -39,8 +39,7 @@ namespace UE::Editor::DataStorage
 		DECLARE_DELEGATE_OneParam(FOnSelectionChanged, RowHandle)
 
 		SLATE_BEGIN_ARGS(STedsTableViewer)
-			: _CellWidgetPurposes({TEXT("General.Cell")})
-			, _ListSelectionMode(ESelectionMode::Type::Single)
+			: _ListSelectionMode(ESelectionMode::Type::Single)
 		{
 			
 		}
@@ -53,8 +52,11 @@ namespace UE::Editor::DataStorage
 		// Table Viewer TODO: How do we specify column metadata (ReadOnly or ReadWrite)?
 		SLATE_ARGUMENT(TArray<TWeakObjectPtr<const UScriptStruct>>, Columns)
 
-		// The widget purposes to use to create the widgets
-		SLATE_ARGUMENT(TArray<FName>, CellWidgetPurposes)
+		// The widget purpose to use to create the cell widgets
+		SLATE_ARGUMENT(IUiProvider::FPurposeID, CellWidgetPurpose)
+			
+		// The widget purpose to use to create the header widgets
+		SLATE_ARGUMENT(IUiProvider::FPurposeID, HeaderWidgetPurpose)
 		
 		// Delegate called when the selection changes
 		SLATE_EVENT(FOnSelectionChanged, OnSelectionChanged)
diff --git a/Engine/Plugins/Experimental/EditorDataStorageFeatures/Source/TedsTableViewer/Public/Widgets/STedsTileViewer.h b/Engine/Plugins/Experimental/EditorDataStorageFeatures/Source/TedsTableViewer/Public/Widgets/STedsTileViewer.h
--- a/Engine/Plugins/Experimental/EditorDataStorageFeatures/Source/TedsTableViewer/Public/Widgets/STedsTileViewer.h
+++ b/Engine/Plugins/Experimental/EditorDataStorageFeatures/Source/TedsTableViewer/Public/Widgets/STedsTileViewer.h
@@ -32,7 +32,7 @@ namespace UE::Editor::DataStorage
 		DECLARE_DELEGATE_OneParam(FOnSelectionChanged, RowHandle)
 
 		SLATE_BEGIN_ARGS(STedsTileViewer)
-			: _WidgetPurpose(TEXT("General.Cell")) // TODO: Have a better "default" widget for tiles
+			: _WidgetPurpose()
 			, _SelectionMode(ESelectionMode::Type::Single)
 			, _ItemHeight(128)
 			, _ItemWidth(128)
@@ -47,8 +47,8 @@ namespace UE::Editor::DataStorage
 		// Optional list of columns to give TEDS UI when creating the tile widget
 		SLATE_ARGUMENT(TArray<TWeakObjectPtr<const UScriptStruct>>, Columns)
 
-		// The widget purpose to use to create the widgets
-		SLATE_ARGUMENT(FName, WidgetPurpose)
+		// The widget purpose to use to create the tile widget
+		SLATE_ARGUMENT(IUiProvider::FPurposeID, WidgetPurpose)
 		
 		// Delegate called when the selection changes
 		SLATE_EVENT(FOnSelectionChanged, OnSelectionChanged)
@@ -113,7 +113,7 @@ namespace UE::Editor::DataStorage
 		TSharedPtr<ITedsWidget> TedsWidget;
 
 		// The widget purpose used to generate tiles
-		FName WidgetPurpose;
+		IUiProvider::FPurposeID WidgetPurpose;
 
 		// Optional columns to use when matching the tile widget
 		TArray<TWeakObjectPtr<const UScriptStruct>> Columns;
diff --git a/Engine/Plugins/Fab/Source/Fab/Private/Teds/FabContentSource.cpp b/Engine/Plugins/Fab/Source/Fab/Private/Teds/FabContentSource.cpp
--- a/Engine/Plugins/Fab/Source/Fab/Private/Teds/FabContentSource.cpp
+++ b/Engine/Plugins/Fab/Source/Fab/Private/Teds/FabContentSource.cpp
@@ -65,7 +65,9 @@ void FFabMyLibraryContentSource::GetAssetViewInitParams(UE::Editor::ContentBrows
 
 	OutInitParams.Columns = { FNameColumn::StaticStruct(), FUrlColumn::StaticStruct() };
 
-	OutInitParams.WidgetPurposes = { FName("General.Cell") }; 
+	OutInitParams.CellWidgetPurpose = IUiProvider::FPurposeID(IUiProvider::FPurposeInfo(
+		"General", "Cell", NAME_None).GeneratePurposeID());
+
 }
 
 #undef LOCTEXT_NAMESPACE
\ No newline at end of file
diff --git a/Engine/Source/Editor/ContentBrowser/Private/ContentSources/Widgets/SContentSource.cpp b/Engine/Source/Editor/ContentBrowser/Private/ContentSources/Widgets/SContentSource.cpp
--- a/Engine/Source/Editor/ContentBrowser/Private/ContentSources/Widgets/SContentSource.cpp
+++ b/Engine/Source/Editor/ContentBrowser/Private/ContentSources/Widgets/SContentSource.cpp
@@ -11,7 +11,8 @@ namespace UE::Editor::ContentBrowser
 {
 	namespace Private
 	{
-		static const FName Purpose("ContentBrowser.AssetView");
+		static const DataStorage::IUiProvider::FPurposeID
+			Purpose(DataStorage::IUiProvider::FPurposeInfo("ContentBrowser", "AssetView", NAME_None).GeneratePurposeID());
 	}
 	
 	void SContentSource::Construct(const FArguments& InArgs)
@@ -68,7 +69,7 @@ namespace UE::Editor::ContentBrowser
 			return false;
 		};
 
-		DataStorageUi->CreateWidgetConstructors(Private::Purpose, FMetaDataView(), AssignWidgetToColumn);
+		DataStorageUi->CreateWidgetConstructors(DataStorageUi->FindPurpose(Private::Purpose), FMetaDataView(), AssignWidgetToColumn);
 
 		TSharedPtr<SWidget> AssetViewWidget;
 		
diff --git a/Engine/Source/Editor/ContentBrowser/Public/ContentSources/IContentSource.h b/Engine/Source/Editor/ContentBrowser/Public/ContentSources/IContentSource.h
--- a/Engine/Source/Editor/ContentBrowser/Public/ContentSources/IContentSource.h
+++ b/Engine/Source/Editor/ContentBrowser/Public/ContentSources/IContentSource.h
@@ -3,6 +3,7 @@
 #pragma once
 
 #include "Elements/Common/TypedElementQueryDescription.h"
+#include "Elements/Interfaces/TypedElementDataStorageUiInterface.h"
 #include "Framework/Views/ITypedTableView.h"
 #include "Internationalization/Text.h"
 #include "Textures/SlateIcon.h"
@@ -22,7 +23,7 @@ namespace UE::Editor::ContentBrowser
 		TArray<TWeakObjectPtr<const UScriptStruct>> Columns;
 
 		// The widget purposes used by the table viewer to display the widgets in the cells
-		TArray<FName> WidgetPurposes;
+		DataStorage::IUiProvider::FPurposeID CellWidgetPurpose;
 
 		// The table view type to use, currently only supports list and tile
 		ETableViewMode::Type TableViewMode = ETableViewMode::Type::List;
diff --git a/Engine/Source/Runtime/TypedElementFramework/Private/Elements/Interfaces/TypedElementDataStorageUiInterface.cpp b/Engine/Source/Runtime/TypedElementFramework/Private/Elements/Interfaces/TypedElementDataStorageUiInterface.cpp
--- a/Engine/Source/Runtime/TypedElementFramework/Private/Elements/Interfaces/TypedElementDataStorageUiInterface.cpp
+++ b/Engine/Source/Runtime/TypedElementFramework/Private/Elements/Interfaces/TypedElementDataStorageUiInterface.cpp
@@ -7,6 +7,8 @@
 #include "Elements/Columns/TypedElementSlateWidgetColumns.h"
 #include "Elements/Columns/TypedElementUIColumns.h"
 #include "Elements/Framework/TypedElementDataStorageWidget.h"
+#include "Elements/Columns/WidgetPurposeColumns.h"
+#include "Elements/Common/EditorDataStorageFeatures.h"
 #include "Widgets/SNullWidget.h"
 
 #define LOCTEXT_NAMESPACE "FTypedElementWidgetConstructor"
@@ -19,9 +21,23 @@ FTypedElementWidgetConstructor::FTypedElementWidgetConstructor(const UScriptStru
 bool FTypedElementWidgetConstructor::Initialize(const UE::Editor::DataStorage::FMetaDataView& InArguments,
 	TArray<TWeakObjectPtr<const UScriptStruct>> InMatchedColumnTypes, const UE::Editor::DataStorage::Queries::FConditions& InQueryConditions)
 {
-	MatchedColumnTypes = MoveTemp(InMatchedColumnTypes);
+	
 	QueryConditions = &InQueryConditions;
+	return Initialize_Internal(InArguments, MoveTemp(InMatchedColumnTypes));
+}
+
+bool FTypedElementWidgetConstructor::Initialize(const UE::Editor::DataStorage::FMetaDataView& InArguments,
+	TArray<TWeakObjectPtr<const UScriptStruct>> InMatchedColumnTypes, UE::Editor::DataStorage::RowHandle InFactoryRowHandle)
+{
+	WidgetFactoryRow = InFactoryRowHandle;
+	return Initialize_Internal(InArguments, MoveTemp(InMatchedColumnTypes));
+}
 
+bool FTypedElementWidgetConstructor::Initialize_Internal(const UE::Editor::DataStorage::FMetaDataView& InArguments,
+	TArray<TWeakObjectPtr<const UScriptStruct>> InMatchedColumnTypes)
+{
+	MatchedColumnTypes = MoveTemp(InMatchedColumnTypes);
+	
 	// If we matched with any dynamic columns, store a mapping from the base template to the identifier so the widget can look it up later
 	for (const TWeakObjectPtr<const UScriptStruct>& Column : MatchedColumnTypes)
 	{
@@ -47,7 +63,21 @@ const TArray<TWeakObjectPtr<const UScriptStruct>>& FTypedElementWidgetConstructo
 
 const UE::Editor::DataStorage::Queries::FConditions* FTypedElementWidgetConstructor::GetQueryConditions() const
 {
-	return QueryConditions;
+	const UE::Editor::DataStorage::ICoreProvider* Storage = UE::Editor::DataStorage::GetDataStorageFeature<UE::Editor::DataStorage::ICoreProvider>(UE::Editor::DataStorage::StorageFeatureName);
+
+	return GetQueryConditions(Storage);
+}
+
+const UE::Editor::DataStorage::Queries::FConditions* FTypedElementWidgetConstructor::GetQueryConditions(const UE::Editor::DataStorage::ICoreProvider* Storage) const
+{
+	if (const FWidgetFactoryConditionsColumn* ConditionsColumn = Storage->GetColumn<FWidgetFactoryConditionsColumn>(WidgetFactoryRow))
+	{
+		return &ConditionsColumn->Conditions;
+	}
+	else 
+	{
+		return QueryConditions;
+	}
 }
 
 TConstArrayView<const UScriptStruct*> FTypedElementWidgetConstructor::GetAdditionalColumnsList() const
@@ -127,7 +157,7 @@ TSharedPtr<SWidget> FTypedElementWidgetConstructor::ConstructFinalWidget(
 		bool bConstructWidget = DataStorage->IsRowAssigned(RowReference->Row);
 
 		// If the original row matches this widgets query conditions currently, create the actual internal widget
-		if (const UE::Editor::DataStorage::Queries::FConditions* MatchedQueryConditions = GetQueryConditions();
+		if (const UE::Editor::DataStorage::Queries::FConditions* MatchedQueryConditions = GetQueryConditions(DataStorage);
 			MatchedQueryConditions != nullptr && bConstructWidget)
 		{
 			bConstructWidget &= DataStorage->MatchesColumns(RowReference->Row, *MatchedQueryConditions);
@@ -314,4 +344,57 @@ TSharedPtr<SWidget> FSimpleWidgetConstructor::Construct(UE::Editor::DataStorage:
 	return Widget;
 }
 
-#undef LOCTEXT_NAMESPACE
\ No newline at end of file
+UE::Editor::DataStorage::IUiProvider::FPurposeInfo::FPurposeInfo(const FName& InNamespace, const FName& InName, const FName& InFrame,
+	EPurposeType InPurposeType, const FText& InDescription, const FPurposeID& InParent)
+	: Namespace(InNamespace)
+	, Name(InName)
+	, Frame(InFrame)
+	, Type(InPurposeType)
+	, Description(InDescription)
+	, ParentPurposeID(InParent)
+{
+}
+
+UE::Editor::DataStorage::IUiProvider::FPurposeInfo::FPurposeInfo(const FName& InLegacyPurposeName, EPurposeType InPurposeType,
+	const FText& InDescription, const FPurposeID& InParent)
+{
+	FString NamespaceSplit, NameSplit, FrameSplit;
+
+	// We want to split the legacy purpose separated by dots into 3 unique names. E.g "General.Cell.Large" -> "General", "Cell", "Large"
+	InLegacyPurposeName.ToString().Split(".", &NamespaceSplit, &NameSplit);
+	NameSplit.Split(".", &NameSplit, &FrameSplit);
+
+	Namespace = FName(NamespaceSplit);
+	Name = FName(NameSplit);
+	Frame = FName(FrameSplit);
+	
+	Type = InPurposeType;
+	Description = InDescription;
+	ParentPurposeID = InParent;
+}
+
+UE::Editor::DataStorage::IUiProvider::FPurposeID UE::Editor::DataStorage::IUiProvider::FPurposeInfo::GeneratePurposeID() const
+{
+	// Create the unique ID from the purpose by combining the namespace, name and frame
+	FStringBuilderBase PurposeIDBuilder;
+
+	if (!Namespace.IsNone())
+	{
+		Namespace.AppendString(PurposeIDBuilder);
+		PurposeIDBuilder.Append(TEXT("."));
+	}
+
+	Name.AppendString(PurposeIDBuilder);
+
+	if (!Frame.IsNone())
+	{
+		Frame.AppendString(PurposeIDBuilder);
+		PurposeIDBuilder.Append(TEXT("."));
+
+	}
+	
+	return FMapKey(FName(PurposeIDBuilder));
+}
+
+
+#undef LOCTEXT_NAMESPACE
diff --git a/Engine/Source/Runtime/TypedElementFramework/Public/Elements/Columns/WidgetPurposeColumns.h b/Engine/Source/Runtime/TypedElementFramework/Public/Elements/Columns/WidgetPurposeColumns.h
new file mode 100644
--- /dev/null
+++ b/Engine/Source/Runtime/TypedElementFramework/Public/Elements/Columns/WidgetPurposeColumns.h
@@ -0,0 +1,87 @@
+// Copyright Epic Games, Inc. All Rights Reserved.
+
+#pragma once
+
+#include "Elements/Interfaces/TypedElementDataStorageUiInterface.h"
+
+#include "WidgetPurposeColumns.generated.h"
+
+/** Column used to store the type of a widget purpose */
+USTRUCT()
+struct FWidgetPurposeColumn : public FEditorDataStorageColumn
+{
+	GENERATED_BODY()
+
+	/** The type of the purpose, used to determine how factories are registered for it */
+	UE::Editor::DataStorage::IUiProvider::EPurposeType PurposeType;
+
+	/** The unique ID of this purpose that can be used to reference it instead of the row handle */
+	UE::Editor::DataStorage::IUiProvider::FPurposeID PurposeID;
+};
+
+/** Column used to store the name of a widget purpose split into 3 parts (E.g "SceneOutliner.Cell.Large") */
+USTRUCT()
+struct FWidgetPurposeNameColumn : public FEditorDataStorageColumn
+{
+	GENERATED_BODY()
+
+	UPROPERTY(meta = (Searchable))
+	FName Namespace;
+
+	UPROPERTY(meta = (Searchable))
+	FName Name;
+
+	UPROPERTY(meta = (Searchable))
+	FName Frame;
+};
+
+
+/** Column to store info about a widget factory */
+USTRUCT()
+struct FWidgetFactoryColumn : public FEditorDataStorageColumn
+{
+	GENERATED_BODY()
+
+	/** Row handle of the purpose the widget factory belongs to */
+	UE::Editor::DataStorage::RowHandle PurposeRowHandle;
+};
+
+/** Column used to store the widget constructor used by a widget factory by its TypeInfo
+ *  A factory can either have FWidgetFactoryConstructorTypeInfoColumn or FWidgetFactoryConstructorColumn
+ */
+USTRUCT()
+struct FWidgetFactoryConstructorTypeInfoColumn : public FEditorDataStorageColumn
+{
+	GENERATED_BODY()
+
+	TWeakObjectPtr<const UScriptStruct> Constructor;
+};
+
+/** Column used to store the widget constructor used by a widget factory */
+USTRUCT()
+struct FWidgetFactoryConstructorColumn : public FEditorDataStorageColumn
+{
+	GENERATED_BODY()
+
+	TUniquePtr<FTypedElementWidgetConstructor> Constructor;
+};
+
+template<>
+struct TStructOpsTypeTraits<FWidgetFactoryConstructorColumn> : public TStructOpsTypeTraitsBase2<FWidgetFactoryConstructorColumn>
+{
+	enum
+	{
+		WithCopy = false
+	};
+};
+
+/** Column used to store the query conditions used by a widget factory
+ *  If the column is not present the factory doesn't match against any conditions (i.e is a general purpose factory)
+ */
+USTRUCT()
+struct FWidgetFactoryConditionsColumn : public FEditorDataStorageColumn
+{
+	GENERATED_BODY()
+
+	UE::Editor::DataStorage::Queries::FConditions Conditions;
+};
diff --git a/Engine/Source/Runtime/TypedElementFramework/Public/Elements/Interfaces/TypedElementDataStorageUiInterface.h b/Engine/Source/Runtime/TypedElementFramework/Public/Elements/Interfaces/TypedElementDataStorageUiInterface.h
--- a/Engine/Source/Runtime/TypedElementFramework/Public/Elements/Interfaces/TypedElementDataStorageUiInterface.h
+++ b/Engine/Source/Runtime/TypedElementFramework/Public/Elements/Interfaces/TypedElementDataStorageUiInterface.h
@@ -37,31 +37,22 @@ public:
 	explicit FTypedElementWidgetConstructor(EForceInit) {} //< For compatibility and shouldn't be directly used.
 
 	virtual ~FTypedElementWidgetConstructor() = default;
-
-	/** Initializes a new constructor based on the provided arguments.. */
+	
+	/** Initializes a new constructor based on the provided arguments. */
 	TYPEDELEMENTFRAMEWORK_API virtual bool Initialize(const UE::Editor::DataStorage::FMetaDataView& InArguments,
-		TArray<TWeakObjectPtr<const UScriptStruct>> InMatchedColumnTypes, const UE::Editor::DataStorage::Queries::FConditions& InQueryConditions);
+		TArray<TWeakObjectPtr<const UScriptStruct>> InMatchedColumnTypes, UE::Editor::DataStorage::RowHandle FactoryRowHandle);
 
 	/** Retrieves the type information for the constructor type. */
 	TYPEDELEMENTFRAMEWORK_API virtual const UScriptStruct* GetTypeInfo() const;
 	/** Retrieves the columns, if any, that were matched to this constructor when it was created. */
 	TYPEDELEMENTFRAMEWORK_API virtual const TArray<TWeakObjectPtr<const UScriptStruct>>& GetMatchedColumns() const;
+	
 	/** Retrieves the query conditions that need to match for this widget constructor to produce a widget. */
-	TYPEDELEMENTFRAMEWORK_API virtual const UE::Editor::DataStorage::Queries::FConditions* GetQueryConditions() const;
+	TYPEDELEMENTFRAMEWORK_API virtual const UE::Editor::DataStorage::Queries::FConditions* GetQueryConditions(const UE::Editor::DataStorage::ICoreProvider* Storage) const;
 
 	/** Returns a list of additional columns the widget requires to be added to its rows. */
 	TYPEDELEMENTFRAMEWORK_API virtual TConstArrayView<const UScriptStruct*> GetAdditionalColumnsList() const;
-
-	/**
-	 * Returns a friendly name for the data the created widget represents.
-	 * By default the associated column is used. If there are multiple columns associated with the constructor
-	 * the default implementation will attempt to find the longest common starting string for all the columns.
-	 * Individual widget constructors can override this function with a name specific to them.
-	 */
-	UE_DEPRECATED(5.6, "CreateWidgetDisplayName is deprecated, please use CreateWidgetDisplayNameText instead, which will return FText rather than an FString.")
-	TYPEDELEMENTFRAMEWORK_API virtual FString CreateWidgetDisplayName(
-		UE::Editor::DataStorage::ICoreProvider* DataStorage, UE::Editor::DataStorage::RowHandle Row) const final;
-
+	
 	/** 
 	 * Returns an FText of the friendly name for the data the created widget represents.
 	 * By default the associated column is used. If there are multiple columns associated with the constructor
@@ -100,6 +91,25 @@ public:
 		UE::Editor::DataStorage::ICoreProvider* DataStorage,
 		UE::Editor::DataStorage::IUiProvider* DataStorageUi,
 		const UE::Editor::DataStorage::FMetaDataView& Arguments);
+	
+	/** Initializes a new constructor based on the provided arguments.. */
+	UE_DEPRECATED(5.6, TEXT("Use Initialize that takes in the factory row handle instead"))
+	TYPEDELEMENTFRAMEWORK_API virtual bool Initialize(const UE::Editor::DataStorage::FMetaDataView& InArguments,
+		TArray<TWeakObjectPtr<const UScriptStruct>> InMatchedColumnTypes, const UE::Editor::DataStorage::Queries::FConditions& InQueryConditions);
+
+	/** Retrieves the query conditions that need to match for this widget constructor to produce a widget. */
+	UE_DEPRECATED(5.6, TEXT("Use GetQueryConditions that takes in ICoreProvider* instead"))
+	TYPEDELEMENTFRAMEWORK_API virtual const UE::Editor::DataStorage::Queries::FConditions* GetQueryConditions() const final;
+	
+	/**
+	 * Returns a friendly name for the data the created widget represents.
+	 * By default the associated column is used. If there are multiple columns associated with the constructor
+	 * the default implementation will attempt to find the longest common starting string for all the columns.
+	 * Individual widget constructors can override this function with a name specific to them.
+	 */
+	UE_DEPRECATED(5.6, "CreateWidgetDisplayName is deprecated, please use CreateWidgetDisplayNameText instead, which will return FText rather than an FString.")
+	TYPEDELEMENTFRAMEWORK_API virtual FString CreateWidgetDisplayName(
+		UE::Editor::DataStorage::ICoreProvider* DataStorage, UE::Editor::DataStorage::RowHandle Row) const final;
 
 protected:
 	/** Create a new instance of the target widget. This is a required function. */
@@ -136,13 +146,27 @@ protected:
 	 */
 	TYPEDELEMENTFRAMEWORK_API UE::Editor::DataStorage::RowHandle GetTargetRow(
 		UE::Editor::DataStorage::ICoreProvider* DataStorage, UE::Editor::DataStorage::RowHandle WidgetRow) const;
-
+	
 protected:
 
+	// The columns this widget constructor matched against
 	TArray<TWeakObjectPtr<const UScriptStruct>> MatchedColumnTypes;
+
+	// If the constructor matched against any dynamic columns, map from the base template for the column -> FName identifier
 	TMap<TWeakObjectPtr<const UScriptStruct>, FName> MatchedDynamicTemplates;
+
+	// The row containing the widget factory this constructor was factories from
+	UE::Editor::DataStorage::RowHandle WidgetFactoryRow = UE::Editor::DataStorage::InvalidRowHandle;
+
+	// The query conditions required by this constructor
 	const UE::Editor::DataStorage::Queries::FConditions* QueryConditions = nullptr;
+
+	// The constructor's typeinfo
 	const UScriptStruct* TypeInfo = nullptr;
+
+private:
+	bool Initialize_Internal(const UE::Editor::DataStorage::FMetaDataView& InArguments,
+	TArray<TWeakObjectPtr<const UScriptStruct>> InMatchedColumnTypes);
 };
 
 /**
@@ -223,169 +247,293 @@ struct TStructOpsTypeTraits<FTypedElementWidgetConstructor> : public TStructOpsT
 
 namespace UE::Editor::DataStorage
 {
-class IUiProvider : public IModularFeature
-{
-public:
-	enum class EPurposeType : uint8
+	class IUiProvider : public IModularFeature
 	{
-		/** General purpose name which allows multiple factory registrations. */
-		Generic,
+	public:
+		enum class EPurposeType : uint8
+		{
+			/** General purpose name which allows multiple factory registrations. */
+			Generic,
+			/**
+			 * Only one factory can be registered with this purpose. If multiple factories are registered only the last factory
+			 * will be stored.
+			 */
+			UniqueByName,
+			/**
+			 * Only one factory can be registered with this purpose for a specific combination of columns. If multiple factories
+			 * are registered only the last factory will be stored.
+			 */
+			UniqueByNameAndColumn
+		};
+
+		enum class EMatchApproach : uint8
+		{
+			/** 
+			 * Looks for the longest chain of columns matching widget factories. The matching columns are removed and the process
+			 * is repeated until there are no more columns or no matches are found.
+			 */
+			LongestMatch,
+			/** A single widget factory is reduced which matches the requested columns exactly. */
+			ExactMatch,
+			/** Each column is matched to widget factory. Only widget factories that use a single column are used. */
+			SingleMatch
+		};
+
+		using FPurposeID = UE::Editor::DataStorage::FMapKey;
+
+		/*
+		 * Struct describing the default init params for a widget purpose. These can also be added to the purpose row post init using the columns
+		 * in WidgetPurposeColumns.h
+		 */
+		struct FPurposeInfo
+		{
+			/** The namespace the purpose belongs to (e.g "General", "SceneOutliner") */
+			FName Namespace;
+
+			/** The name of the purpose (e.g "RowLabel", "Cell") */
+			FName Name;
+
+			/** An optional suffix for the purpose (e.g "Large", "Small", "Default") */
+			FName Frame;
+
+			/** The purpose type */
+			EPurposeType Type;
+
+			/** A user facing description for the purpose */
+			FText Description;
+
+			/** The parent purpose. If valid, widget construction methods will go up the hierarchy if no widgets were found for a purpose */
+			FPurposeID ParentPurposeID;
+
+			/** Default constructor for FPurposeInfo, the members with a default value are optional when you only need an FPurposeInfo to look up FPurposeID */
+			TYPEDELEMENTFRAMEWORK_API FPurposeInfo(const FName& InNamespace, const FName& InName, const FName& InFrame,
+				EPurposeType InPurposeType = EPurposeType::UniqueByNameAndColumn, const FText& InDescription = FText::GetEmpty(),
+				const FPurposeID& InParent = FPurposeID());
+		
+			/** Create an FPurposeInfo struct using a legacy purpose name to parse it as Namespace.Name.Frame (e.g "SceneOutliner.Cell.Large") */
+			TYPEDELEMENTFRAMEWORK_API FPurposeInfo(const FName& InLegacyPurposeName, EPurposeType InPurposeType, const FText& InDescription,
+				const FPurposeID& InParent = FPurposeID());
+
+			/** Convert this PurposeInfo to a PurposeID that can be used to look up the purpose row */
+			TYPEDELEMENTFRAMEWORK_API FPurposeID GeneratePurposeID() const;
+		};
+
+		using WidgetCreatedCallback = TFunctionRef<void(const TSharedRef<SWidget>& NewWidget, UE::Editor::DataStorage::RowHandle Row)>;
+		using WidgetConstructorCallback = TFunctionRef<bool(TUniquePtr<FTypedElementWidgetConstructor>, TConstArrayView<TWeakObjectPtr<const UScriptStruct>>)>;
+		using WidgetPurposeCallback = TFunctionRef<void(FName, EPurposeType, const FText&)>;
+
 		/**
-		 * Only one factory can be registered with this purpose. If multiple factories are registered only the last factory
-		 * will be stored.
+		 * Register a widget purpose. Widget purposes indicates how widgets can be used and categorizes/organizes the available 
+		 * widget factories. If the same purpose is registered multiple times, only the first will be recorded and later registrations
+		 * will be ignored and the return value will be InvalidRowHandle.
+		 * The purposes will be mapped to the TEDS Mapping Table using InPurposeInfo.GeneratePurposeID() which can be used to lookup the purpose row.
+		 *
+		 * @param InPurposeInfo Struct with init params for the purpose
 		 */
-		UniqueByName,
+		virtual UE::Editor::DataStorage::RowHandle RegisterWidgetPurpose(const FPurposeInfo& InPurposeInfo) = 0;
+	
 		/**
-		 * Only one factory can be registered with this purpose for a specific combination of columns. If multiple factories
-		 * are registered only the last factory will be stored.
+		 * Overload for RegisterWidgetPurpose that allows you to use a custom ID to map to the TEDS Mapping table for lookup
+		 * 
+		 * @param PurposeID A unique ID for the purpose that can be used to look up the row handle later
+		 * @param InPurposeInfo Struct with init params for the purpose
 		 */
-		UniqueByNameAndColumn
-	};
+		virtual UE::Editor::DataStorage::RowHandle RegisterWidgetPurpose(const FPurposeID& PurposeID, const FPurposeInfo& InPurposeInfo) = 0;
+	
+	
+		/** 
+		 * Registers a widget factory that will be called when the purpose it's registered under is requested.
+		 * This version registers a generic type. Construction using these are typically cheaper as they can avoid
+		 * copying the Constructor and take up less memory. The downside is that they can't store additional configuration
+		 * options. If the purpose has not been registered the factory will not be recorded and a warning will be printed.
+		 * If registration is successful true will be returned otherwise false.
+		 */
+		virtual bool RegisterWidgetFactory(UE::Editor::DataStorage::RowHandle PurposeRow, const UScriptStruct* Constructor) = 0;
 
-	enum class EMatchApproach : uint8
-	{
+		/**
+		 * Registers a widget factory that will be called when the purpose it's registered under is requested.
+		 * This version registers a generic type. Construction using these are typically cheaper as they can avoid
+		 * copying the Constructor and take up less memory. The downside is that they can't store additional configuration
+		 * options. If the purpose has not been registered the factory will not be recorded and a warning will be printed.
+		 * If registration is successful true will be returned otherwise false.
+		 */
+		template<typename ConstructorType>
+		bool RegisterWidgetFactory(UE::Editor::DataStorage::RowHandle PurposeRow);
+
+		/**
+		 * Registers a widget factory that will be called when the purpose it's registered under is requested.
+		 * This version registers a generic type. Construction using these are typically cheaper as they can avoid
+		 * copying the Constructor and take up less memory. The downside is that they can't store additional configuration
+		 * options. If the purpose has not been registered the factory will not be recorded and a warning will be printed.
+		 * The provided columns will be used when matching the factory during widget construction.
+		 * If registration is successful true will be returned otherwise false.
+		 */
+		virtual bool RegisterWidgetFactory(UE::Editor::DataStorage::RowHandle PurposeRow, const UScriptStruct* Constructor,
+			UE::Editor::DataStorage::Queries::FConditions Columns) = 0;
+	
+		/**
+		 * Registers a widget factory that will be called when the purpose it's registered under is requested.
+		 * This version registers a generic type. Construction using these are typically cheaper as they can avoid
+		 * copying the Constructor and take up less memory. The downside is that they can't store additional configuration
+		 * options. If the purpose has not been registered the factory will not be recorded and a warning will be printed.
+		 * The provided columns will be used when matching the factory during widget construction.
+		 * If registration is successful true will be returned otherwise false.
+		 */
+		template<typename ConstructorType>
+		bool RegisterWidgetFactory(UE::Editor::DataStorage::RowHandle PurposeRow, UE::Editor::DataStorage::Queries::FConditions Columns);
+
+		/**
+		 * Registers a widget factory that will be called when the purpose it's registered under is requested.
+		 * This version uses a previously created instance of the Constructor. The benefit of this is that it store
+		 * configuration options. The downside is that it takes up more memory and requires copying when it's
+		 * used. If the purpose has not been registered the factory will not be recorded and a warning will be printed.
+		 * If registration is successful true will be returned otherwise false.
+		 */
+		virtual bool RegisterWidgetFactory(UE::Editor::DataStorage::RowHandle PurposeRow, TUniquePtr<FTypedElementWidgetConstructor>&& Constructor) = 0;
+	
+		/**
+		 * Registers a widget factory that will be called when the purpose it's registered under is requested.
+		 * This version uses a previously created instance of the Constructor. The benefit of this is that it store
+		 * configuration options. The downside is that it takes up more memory and requires copying when it's
+		 * used. If the purpose has not been registered the factory will not be recorded and a warning will be printed.
+		 * The provided columns will be used when matching the factory during widget construction.
+		 * If registration is successful true will be returned otherwise false.
+		 */
+		virtual bool RegisterWidgetFactory(UE::Editor::DataStorage::RowHandle PurposeRow, TUniquePtr<FTypedElementWidgetConstructor>&& Constructor, 
+			UE::Editor::DataStorage::Queries::FConditions Columns) = 0;
+	
 		/** 
-		 * Looks for the longest chain of columns matching widget factories. The matching columns are removed and the process
-		 * is repeated until there are no more columns or no matches are found.
+		 * Creates widget constructors for the requested purpose.
+		 * The provided arguments will be used to configure the constructor. Settings made this way will be applied to all
+		 * widgets created from the constructor, if applicable.
+		 * 
+		 * If no factories were found for the requested purpose, the purpose's parent chain will be traversed to try and find an ancestor purpose with factories
 		 */
-		LongestMatch,
-		/** A single widget factory is reduced which matches the requested columns exactly. */
-		ExactMatch,
-		/** Each column is matched to widget factory. Only widget factories that use a single column are used. */
-		SingleMatch
-	};
+		virtual void CreateWidgetConstructors(UE::Editor::DataStorage::RowHandle PurposeRow, 
+			const UE::Editor::DataStorage::FMetaDataView& Arguments, const WidgetConstructorCallback& Callback) = 0;
+	
+		/** 
+		 * Finds matching widget constructors for provided columns, preferring longer matches over shorter matches.
+		 * The provided list of columns will be updated to contain all columns that couldn't be matched.
+		 * The provided arguments will be used to configure the constructor. Settings made this way will be applied to all
+		 * widgets created from the constructor, if applicable.
+		 *
+		 * If there are still columns remaining after matching against the requested purpose, the purpose's parent chain will be traversed to try and match
+		 * the remaining columns
+		 */
+		virtual void CreateWidgetConstructors(UE::Editor::DataStorage::RowHandle PurposeRow, EMatchApproach MatchApproach, TArray<TWeakObjectPtr<const UScriptStruct>>& Columns,
+			const UE::Editor::DataStorage::FMetaDataView& Arguments, const WidgetConstructorCallback& Callback) = 0;
+	
+		/**
+		 * Creates all the widgets registered under the provided purpose. This may be a large number of widgets for a wide name
+		 * or exactly one when the exact name of the widget is registered. Arguments can be provided, but widgets are free
+		 * to ignore them.
+		 * 
+		 * If no factories were found for the requested purpose, the purpose's parent chain will be traversed to try and find an ancestor purpose with factories
+		 */
+		virtual void ConstructWidgets(UE::Editor::DataStorage::RowHandle PurposeRow, const UE::Editor::DataStorage::FMetaDataView& Arguments,
+			const WidgetCreatedCallback& ConstructionCallback) = 0;
+	
+		/** 
+		 * Creates a single widget using the provided constructor. 
+		 * The provided row will be used to store the widget information on. If columns have already been added to the row, the 
+		 * constructor is free to use that to configure the widget. Arguments are used by the constructor to configure the widget.
+		 */
+		virtual TSharedPtr<SWidget> ConstructWidget(UE::Editor::DataStorage::RowHandle Row, FTypedElementWidgetConstructor& Constructor,
+			const UE::Editor::DataStorage::FMetaDataView& Arguments) = 0;
 
-	using WidgetCreatedCallback = TFunctionRef<void(const TSharedRef<SWidget>& NewWidget, RowHandle Row)>;
-	using WidgetConstructorCallback = TFunctionRef<bool(TUniquePtr<FTypedElementWidgetConstructor>, TConstArrayView<TWeakObjectPtr<const UScriptStruct>>)>;
-	using WidgetPurposeCallback = TFunctionRef<void(FName, EPurposeType, const FText&)>;
+		/** Calls the provided callback for all known registered widget purposes. */
+		virtual void ListWidgetPurposes(const WidgetPurposeCallback& Callback) const = 0;
 
-	/**
-	 * Register a widget purpose. Widget purposes indicates how widgets can be used and categorizes/organizes the available 
-	 * widget factories. If the same purpose is registered multiple times, only the first will be recorded and later registrations
-	 * will be silently ignored.
-	 * Widget purposes follow a specific naming convention currently:
-	 * "PurposeName.Cell" for widgets created for (row, columns) pairs
-	 * "PurposeName.Header" for widgets created for column headers
-	 * "PurposeName.Cell.Default" or "PurposeName.Header.Default" for generic widgets not registered against column(s)
-	 */
-	virtual void RegisterWidgetPurpose(FName Purpose, EPurposeType Type, FText Description) = 0;
-	/** 
-	 * Registers a widget factory that will be called when the purpose it's registered under is requested.
-	 * This version registers a generic type. Construction using these are typically cheaper as they can avoid
-	 * copying the Constructor and take up less memory. The downside is that they can't store additional configuration
-	 * options. If the purpose has not been registered the factory will not be recorded and a warning will be printed.
-	 * If registration is successful true will be returned otherwise false.
-	 */
-	virtual bool RegisterWidgetFactory(FName Purpose, const UScriptStruct* Constructor) = 0;
-	/**
-	 * Registers a widget factory that will be called when the purpose it's registered under is requested.
-	 * This version registers a generic type. Construction using these are typically cheaper as they can avoid
-	 * copying the Constructor and take up less memory. The downside is that they can't store additional configuration
-	 * options. If the purpose has not been registered the factory will not be recorded and a warning will be printed.
-	 * If registration is successful true will be returned otherwise false.
-	 */
-	template<typename ConstructorType>
-	bool RegisterWidgetFactory(FName Purpose);
-	/**
-	 * Registers a widget factory that will be called when the purpose it's registered under is requested.
-	 * This version registers a generic type. Construction using these are typically cheaper as they can avoid
-	 * copying the Constructor and take up less memory. The downside is that they can't store additional configuration
-	 * options. If the purpose has not been registered the factory will not be recorded and a warning will be printed.
-	 * The provided columns will be used when matching the factory during widget construction.
-	 * If registration is successful true will be returned otherwise false.
-	 */
-	virtual bool RegisterWidgetFactory(FName Purpose, const UScriptStruct* Constructor,
-		Queries::FConditions Columns) = 0;
-	/**
-	 * Registers a widget factory that will be called when the purpose it's registered under is requested.
-	 * This version registers a generic type. Construction using these are typically cheaper as they can avoid
-	 * copying the Constructor and take up less memory. The downside is that they can't store additional configuration
-	 * options. If the purpose has not been registered the factory will not be recorded and a warning will be printed.
-	 * The provided columns will be used when matching the factory during widget construction.
-	 * If registration is successful true will be returned otherwise false.
-	 */
-	template<typename ConstructorType>
-	bool RegisterWidgetFactory(FName Purpose, Queries::FConditions Columns);
-	/**
-	 * Registers a widget factory that will be called when the purpose it's registered under is requested.
-	 * This version uses a previously created instance of the Constructor. The benefit of this is that it store
-	 * configuration options. The downside is that it takes up more memory and requires copying when it's
-	 * used. If the purpose has not been registered the factory will not be recorded and a warning will be printed.
-	 * If registration is successful true will be returned otherwise false.
-	 */
-	virtual bool RegisterWidgetFactory(FName Purpose, TUniquePtr<FTypedElementWidgetConstructor>&& Constructor) = 0;
-	/**
-	 * Registers a widget factory that will be called when the purpose it's registered under is requested.
-	 * This version uses a previously created instance of the Constructor. The benefit of this is that it store
-	 * configuration options. The downside is that it takes up more memory and requires copying when it's
-	 * used. If the purpose has not been registered the factory will not be recorded and a warning will be printed.
-	 * The provided columns will be used when matching the factory during widget construction.
-	 * If registration is successful true will be returned otherwise false.
-	 */
-	virtual bool RegisterWidgetFactory(FName Purpose, TUniquePtr<FTypedElementWidgetConstructor>&& Constructor, 
-		Queries::FConditions Columns) = 0;
+		/** Check if a custom extension is supported. This can be used to check for in-development features, custom extensions, etc. */
+		virtual bool SupportsExtension(FName Extension) const = 0;
+		/** Provides a list of all extensions that are enabled. */
+		virtual void ListExtensions(TFunctionRef<void(FName)> Callback) const = 0;
+
+		/** Create the container widget that every TEDS UI widget is stored in */
+		virtual TSharedPtr<UE::Editor::DataStorage::ITedsWidget> CreateContainerTedsWidget(UE::Editor::DataStorage::RowHandle UiRowHandle) const = 0;
+
+		/** Get the table where TEDS UI widgets are stored */
+		virtual UE::Editor::DataStorage::TableHandle GetWidgetTable() const = 0;
+
+		/** Find the row handle for a purpose by looking it up using the purpose ID */
+		virtual UE::Editor::DataStorage::RowHandle FindPurpose(const FPurposeID& PurposeID) const = 0;
 	
-	/** 
-	 * Creates widget constructors for the requested purpose.
-	 * The provided arguments will be used to configure the constructor. Settings made this way will be applied to all
-	 * widgets created from the constructor, if applicable.
-	 */
-	virtual void CreateWidgetConstructors(FName Purpose, 
-		const FMetaDataView& Arguments, const WidgetConstructorCallback& Callback) = 0;
-	/** 
-	 * Finds matching widget constructors for provided columns, preferring longer matches over shorter matches.
-	 * The provided list of columns will be updated to contain all columns that couldn't be matched.
-	 * The provided arguments will be used to configure the constructor. Settings made this way will be applied to all
-	 * widgets created from the constructor, if applicable.
-	 */
-	virtual void CreateWidgetConstructors(FName Purpose, EMatchApproach MatchApproach, TArray<TWeakObjectPtr<const UScriptStruct>>& Columns,
-		const FMetaDataView& Arguments, const WidgetConstructorCallback& Callback) = 0;
+		/** Get the name of the default TEDS UI widget purpose used to register default widgets for different types of data (e.g FText -> STextBlock) */
+		virtual FPurposeID GetDefaultWidgetPurposeID() const = 0;
+	
+		/** Get the name of the general TEDS UI purpose used to register general purpose widgets for columns */
+		virtual FPurposeID GetGeneralWidgetPurposeID() const = 0;
+	
+		UE_DEPRECATED(5.6, "Use the version that takes in an FPurposeInfo struct instead")
+		virtual void RegisterWidgetPurpose(FName Purpose, EPurposeType Type, FText Description) = 0;
+	
+		UE_DEPRECATED(5.6, "Use the version that takes in PurposeRowHandle instead. FindPurpose() to lookup row handle by PurposeID")
+		virtual bool RegisterWidgetFactory(FName Purpose, const UScriptStruct* Constructor) = 0;
+	
+		template<typename ConstructorType>
+		UE_DEPRECATED(5.6, "Use the version that takes in PurposeRowHandle instead. FindPurpose() to lookup row handle by PurposeID")
+		bool RegisterWidgetFactory(FName Purpose);
 
-	/**
-	 * Creates all the widgets registered under the provided name. This may be a large number of widgets for a wide name
-	 * or exactly one when the exact name of the widget is registered. Arguments can be provided, but widgets are free
-	 * to ignore them.
-	 */
-	virtual void ConstructWidgets(FName Purpose, const FMetaDataView& Arguments,
-		const WidgetCreatedCallback& ConstructionCallback) = 0;
+		UE_DEPRECATED(5.6, "Use the version that takes in PurposeRowHandle instead. FindPurpose() to lookup row handle by PurposeID")
+		virtual bool RegisterWidgetFactory(FName Purpose, const UScriptStruct* Constructor,
+			UE::Editor::DataStorage::Queries::FConditions Columns) = 0;
 
-	/** 
-	 * Creates a single widget using the provided constructor. 
-	 * The provided row will be used to store the widget information on. If columns have already been added to the row, the 
-	 * constructor is free to use that to configure the widget. Arguments are used by the constructor to configure the widget.
-	 */
-	virtual TSharedPtr<SWidget> ConstructWidget(RowHandle Row, FTypedElementWidgetConstructor& Constructor,
-		const FMetaDataView& Arguments) = 0;
+		template<typename ConstructorType>
+		UE_DEPRECATED(5.6, "Use the version that takes in PurposeRowHandle instead. FindPurpose() to lookup row handle by PurposeID")
+		bool RegisterWidgetFactory(FName Purpose, UE::Editor::DataStorage::Queries::FConditions Columns);
 
-	/** Calls the provided callback for all known registered widget purposes. */
-	virtual void ListWidgetPurposes(const WidgetPurposeCallback& Callback) const = 0;
+		UE_DEPRECATED(5.6, "Use the version that takes in PurposeRowHandle instead. FindPurpose() to lookup row handle by PurposeID")
+		virtual bool RegisterWidgetFactory(FName Purpose, TUniquePtr<FTypedElementWidgetConstructor>&& Constructor) = 0;
 
-	/** Check if a custom extension is supported. This can be used to check for in-development features, custom extensions, etc. */
-	virtual bool SupportsExtension(FName Extension) const = 0;
-	/** Provides a list of all extensions that are enabled. */
-	virtual void ListExtensions(TFunctionRef<void(FName)> Callback) const = 0;
+		UE_DEPRECATED(5.6, "Use the version that takes in PurposeRowHandle instead. FindPurpose() to lookup row handle by PurposeID")
+		virtual bool RegisterWidgetFactory(FName Purpose, TUniquePtr<FTypedElementWidgetConstructor>&& Constructor, 
+			UE::Editor::DataStorage::Queries::FConditions Columns) = 0;
 
-	/** Create the container widget that every TEDS UI widget is stored in */
-	virtual TSharedPtr<ITedsWidget> CreateContainerTedsWidget(RowHandle UiRowHandle) const = 0;
+		UE_DEPRECATED(5.6, "Use the version that takes in PurposeRowHandle instead. FindPurpose() to lookup row handle by PurposeID")
+		virtual void CreateWidgetConstructors(FName Purpose, 
+			const UE::Editor::DataStorage::FMetaDataView& Arguments, const WidgetConstructorCallback& Callback) = 0;
 
-	/** Get the table where TEDS UI widgets are stored */
-	virtual TableHandle GetWidgetTable() const = 0;
-};
+		UE_DEPRECATED(5.6, "Use the version that takes in PurposeRowHandle instead. FindPurpose() to lookup row handle by PurposeID")
+		virtual void CreateWidgetConstructors(FName Purpose, EMatchApproach MatchApproach, TArray<TWeakObjectPtr<const UScriptStruct>>& Columns,
+			const UE::Editor::DataStorage::FMetaDataView& Arguments, const WidgetConstructorCallback& Callback) = 0;
 
+		UE_DEPRECATED(5.6, "Use the version that takes in PurposeRowHandle instead. FindPurpose() to lookup row handle by PurposeID")
+		virtual void ConstructWidgets(FName Purpose, const UE::Editor::DataStorage::FMetaDataView& Arguments,
+			const WidgetCreatedCallback& ConstructionCallback) = 0;
 
-//
-// Implementations
-//
+	};
 
-template<typename ConstructorType>
-bool IUiProvider::RegisterWidgetFactory(FName Purpose)
-{
-	return this->RegisterWidgetFactory(Purpose, ConstructorType::StaticStruct());
-}
 
-template<typename ConstructorType>
-bool IUiProvider::RegisterWidgetFactory(FName Purpose, UE::Editor::DataStorage::Queries::FConditions Columns)
-{
-	return this->RegisterWidgetFactory(Purpose, ConstructorType::StaticStruct(), Columns);
-}
-} // namespace UE::Editor::DataStorage
+	//
+	// Implementations
+	//
+
+	template<typename ConstructorType>
+	bool IUiProvider::RegisterWidgetFactory(FName Purpose)
+	{
+		PRAGMA_DISABLE_DEPRECATION_WARNINGS
+		return this->RegisterWidgetFactory(Purpose, ConstructorType::StaticStruct());
+		PRAGMA_ENABLE_DEPRECATION_WARNINGS
+	}
+
+	template<typename ConstructorType>
+	bool IUiProvider::RegisterWidgetFactory(UE::Editor::DataStorage::RowHandle PurposeRow)
+	{
+		return this->RegisterWidgetFactory(PurposeRow, ConstructorType::StaticStruct());
+	}
+
+	template<typename ConstructorType>
+	bool IUiProvider::RegisterWidgetFactory(FName Purpose, UE::Editor::DataStorage::Queries::FConditions Columns)
+	{
+		PRAGMA_DISABLE_DEPRECATION_WARNINGS
+		return this->RegisterWidgetFactory(Purpose, ConstructorType::StaticStruct(), Columns);
+		PRAGMA_ENABLE_DEPRECATION_WARNINGS
+	}
+
+	template<typename ConstructorType>
+	bool IUiProvider::RegisterWidgetFactory(UE::Editor::DataStorage::RowHandle PurposeRow, UE::Editor::DataStorage::Queries::FConditions Columns)
+	{
+		return this->RegisterWidgetFactory(PurposeRow, ConstructorType::StaticStruct(), Columns);
+	}
+}
\ No newline at end of file
