TEDS UI: Change Malloc+Free to FMemory_Alloca_Aligned when creating widget constructors by moving the code to a helper so the stack memory is cleaned up on function scope end

diff --git a/Engine/Plugins/Experimental/EditorDataStorage/Source/TedsCore/Private/TypedElementDatabaseUI.cpp b/Engine/Plugins/Experimental/EditorDataStorage/Source/TedsCore/Private/TypedElementDatabaseUI.cpp
--- a/Engine/Plugins/Experimental/EditorDataStorage/Source/TedsCore/Private/TypedElementDatabaseUI.cpp
+++ b/Engine/Plugins/Experimental/EditorDataStorage/Source/TedsCore/Private/TypedElementDatabaseUI.cpp
@@ -485,20 +485,7 @@ void UEditorDataStorageUi::ConstructWidgets(UE::Editor::DataStorage::RowHandle P
 		{
 			if (const UScriptStruct* ConstructorType = ConstructorTypeInfoColumn->Constructor.Get())
 			{
-				FTypedElementWidgetConstructor* Constructor = reinterpret_cast<FTypedElementWidgetConstructor*>(
-						FMemory::Malloc(ConstructorType->GetStructureSize(), ConstructorType->GetMinAlignment()));
-				if (Constructor)
-				{
-					ConstructorType->InitializeStruct(Constructor);
-					CreateWidgetInstance(*Constructor, Arguments, ConstructionCallback);
-					ConstructorType->DestroyStruct(&Constructor);
-				}
-				else
-				{
-					checkf(false, TEXT("Remaining memory is too small to create a widget constructor from a description."));
-				}
-
-				FMemory::Free(Constructor);
+				CreateWidgetInstance(ConstructorType, Arguments, ConstructionCallback);
 			}
 		}
 	}
@@ -648,6 +635,23 @@ void UEditorDataStorageUi::CreateWidgetInstance(
 	}
 }
 
+void UEditorDataStorageUi::CreateWidgetInstance(const UScriptStruct* ConstructorType, const UE::Editor::DataStorage::FMetaDataView& Arguments,
+	const WidgetCreatedCallback& ConstructionCallback)
+{
+	FTypedElementWidgetConstructor* Constructor = reinterpret_cast<FTypedElementWidgetConstructor*>(
+						FMemory_Alloca_Aligned(ConstructorType->GetStructureSize(), ConstructorType->GetMinAlignment()));
+	if (Constructor)
+	{
+		ConstructorType->InitializeStruct(Constructor);
+		CreateWidgetInstance(*Constructor, Arguments, ConstructionCallback);
+		ConstructorType->DestroyStruct(&Constructor);
+	}
+	else
+	{
+		checkf(false, TEXT("Remaining memory is too small to create a widget constructor from a description."));
+	}
+}
+
 TSharedPtr<SWidget> UEditorDataStorageUi::ConstructWidget(UE::Editor::DataStorage::RowHandle Row, FTypedElementWidgetConstructor& Constructor,
 	const UE::Editor::DataStorage::FMetaDataView& Arguments)
 {
diff --git a/Engine/Plugins/Experimental/EditorDataStorage/Source/TedsCore/Private/TypedElementDatabaseUI.h b/Engine/Plugins/Experimental/EditorDataStorage/Source/TedsCore/Private/TypedElementDatabaseUI.h
--- a/Engine/Plugins/Experimental/EditorDataStorage/Source/TedsCore/Private/TypedElementDatabaseUI.h
+++ b/Engine/Plugins/Experimental/EditorDataStorage/Source/TedsCore/Private/TypedElementDatabaseUI.h
@@ -106,6 +106,11 @@ private:
 		const UE::Editor::DataStorage::FMetaDataView& Arguments,
 		const WidgetCreatedCallback& ConstructionCallback);
 
+	void CreateWidgetInstance(
+		const UScriptStruct* ConstructorType,
+		const UE::Editor::DataStorage::FMetaDataView& Arguments,
+		const WidgetCreatedCallback& ConstructionCallback);
+
 	bool CreateWidgetConstructors_LongestMatch(
 		const TArray<UE::Editor::DataStorage::RowHandle>& WidgetFactories, 
 		TArray<TWeakObjectPtr<const UScriptStruct>>& Columns, 
