TEDS UI Purposes - follow up CL after the refactor to fix misc issues

diff --git a/Engine/Plugins/Experimental/EditorDataStorage/Source/TedsUI/Private/Widgets/ExportedTextWidget.cpp b/Engine/Plugins/Experimental/EditorDataStorage/Source/TedsUI/Private/Widgets/ExportedTextWidget.cpp
--- a/Engine/Plugins/Experimental/EditorDataStorage/Source/TedsUI/Private/Widgets/ExportedTextWidget.cpp
+++ b/Engine/Plugins/Experimental/EditorDataStorage/Source/TedsUI/Private/Widgets/ExportedTextWidget.cpp
@@ -44,7 +44,7 @@ TConstArrayView<const UScriptStruct*> FExportedTextWidgetConstructor::GetAdditio
 const UE::Editor::DataStorage::Queries::FConditions* FExportedTextWidgetConstructor::GetQueryConditions(const UE::Editor::DataStorage::ICoreProvider* Storage) const
 {
 	// For the exported text widget, the query condition we are matched against is the column we are exporting text for
-	if(MatchedColumn.IsCompiled() && !MatchedColumn.IsEmpty())
+	if (MatchedColumn.IsCompiled() && !MatchedColumn.IsEmpty())
 	{
 		return &MatchedColumn;
 	}
@@ -59,10 +59,15 @@ FText FExportedTextWidgetConstructor::CreateWidgetDisplayNameText(UE::Editor::Da
 	{
 		return DescribeColumnType(TypeInfoColumn->TypeInfo.Get());
 	}
-	else
+	// The default behavior is to display the name of the column this widget matched against
+	// NOTE: This currently only works if CreateWidget has been called at least once because we rely on the TypeInfoColumn on the widget row
+	// to know which column we matched against
+	else if (MatchedColumn.IsCompiled() && !MatchedColumn.IsEmpty())
 	{
-		return DescribeColumnType(nullptr);
+		return DescribeColumnType(MatchedColumn.GetColumns()[0].Get());
 	}
+
+	return FText::GetEmpty();
 }
 
 TSharedPtr<SWidget> FExportedTextWidgetConstructor::CreateWidget(UE::Editor::DataStorage::ICoreProvider* DataStorage,
diff --git a/Engine/Plugins/Experimental/EditorDataStorageFeatures/Source/TedsOutliner/Private/TedsOutlinerModule.cpp b/Engine/Plugins/Experimental/EditorDataStorageFeatures/Source/TedsOutliner/Private/TedsOutlinerModule.cpp
--- a/Engine/Plugins/Experimental/EditorDataStorageFeatures/Source/TedsOutliner/Private/TedsOutlinerModule.cpp
+++ b/Engine/Plugins/Experimental/EditorDataStorageFeatures/Source/TedsOutliner/Private/TedsOutlinerModule.cpp
@@ -96,7 +96,7 @@ TSharedRef<ISceneOutliner> FTedsOutlinerModule::CreateTedsOutliner(const FSceneO
 	
 	TSharedRef<ISceneOutliner> TedsOutlinerShared = SNew(SSceneOutliner, InitOptions);
 	
-	FSceneOutlinerTedsQueryBinder::GetInstance().AssignQuery(ColumnQuery, TedsOutlinerShared, InInitTedsOptions.CellWidgetPurpose, InInitTedsOptions.LabelWidgetPurpose);
+	FSceneOutlinerTedsQueryBinder::GetInstance().AssignQuery(ColumnQuery, TedsOutlinerShared, InInitTedsOptions.CellWidgetPurpose, InInitTedsOptions.HeaderWidgetPurpose);
 	
 	return TedsOutlinerShared;
 }
diff --git a/Engine/Plugins/Experimental/EditorDataStorageFeatures/Source/TedsTableViewer/Private/TedsTableViewerColumn.cpp b/Engine/Plugins/Experimental/EditorDataStorageFeatures/Source/TedsTableViewer/Private/TedsTableViewerColumn.cpp
--- a/Engine/Plugins/Experimental/EditorDataStorageFeatures/Source/TedsTableViewer/Private/TedsTableViewerColumn.cpp
+++ b/Engine/Plugins/Experimental/EditorDataStorageFeatures/Source/TedsTableViewer/Private/TedsTableViewerColumn.cpp
@@ -115,8 +115,15 @@ namespace UE::Editor::DataStorage
 		}
 		if (!Widget.IsValid())
 		{
+			FText DefaultHeaderText = CellWidgetConstructor ? CellWidgetConstructor->CreateWidgetDisplayNameText(Storage) : FText::GetEmpty();
+
+			if (DefaultHeaderText.IsEmpty())
+			{
+				DefaultHeaderText = FText::FromString(ColumnName.ToString());
+			}
+			
 			Widget = SNew(STextBlock)
-				.Text(CellWidgetConstructor ? CellWidgetConstructor->CreateWidgetDisplayNameText(Storage) : FText::FromString(ColumnName.ToString()));
+				.Text(DefaultHeaderText);
 		}
 		
 		SHeaderRow::FColumn::FArguments Column = SHeaderRow::Column(ColumnName)
