[TEDS] Tentative fix for an ASan trigger when deleting rows.

The following changes were made:
1. Reserving arrays on the scratch buffer now returns a TArrayView. This reduces the risk of the size not matching the reserved size.
2. Added a new way to reserve memory in the scratch buffer. The new approach doesn't initialize the memory and only reserves the requested memory. This should be a slight performance gain for large lists of objects if they're immediately assigned afterwards.
3. Fixed a push command with an empty command object still trying to access the instance even though the optimization avoids storing an instance.
4. Allowed the scratch buffer to return block sizes instead of accurate sizes from the underlying allocator. TEDS doesn't currently allow the allocated size to be retrieved for a memory address so this change has no impact on TEDS' usability.
5. Allowed the memory allocator used by the scratch buffer to use inlined block allocation which should improve memory allocation speeds.
The first change improves detection of out-of-bounds reads, while the last 2 changes use simpler code paths in the allocator that will either remove or reduce the places the root problem causing the ASan trigger to be hidden.

#rnx

diff --git a/Engine/Plugins/Experimental/EditorDataStorage/Source/TedsCore/Private/Commands/EditorDataStorageCompatibilityCommands.cpp b/Engine/Plugins/Experimental/EditorDataStorage/Source/TedsCore/Private/Commands/EditorDataStorageCompatibilityCommands.cpp
--- a/Engine/Plugins/Experimental/EditorDataStorage/Source/TedsCore/Private/Commands/EditorDataStorageCompatibilityCommands.cpp
+++ b/Engine/Plugins/Experimental/EditorDataStorage/Source/TedsCore/Private/Commands/EditorDataStorageCompatibilityCommands.cpp
@@ -156,11 +156,10 @@ namespace UE::Editor::DataStorage
 		FScratchBuffer& ScratchBuffer)
 	{
 		uint32 ReinstanceCount = Commands.GetCommandCount<FTypeInfoReinstanced>();
-		FTypeInfoReinstanced* ReinstanceArray = ScratchBuffer.EmplaceArray<FTypeInfoReinstanced>(ReinstanceCount);
-		TArrayView<FTypeInfoReinstanced> ReinstanceView(ReinstanceArray, ReinstanceCount);
-
+		TArrayView<FTypeInfoReinstanced> ReinstanceArray = ScratchBuffer.EmplaceArray<FTypeInfoReinstanced>(ReinstanceCount);
+		
 		// Populate the list of re-instance data and disable the commands.
-		FTypeInfoReinstanced* ReinstanceArrayIt = ReinstanceArray;
+		FTypeInfoReinstanced* ReinstanceArrayIt = ReinstanceArray.GetData();
 		Commands.ForEach(
 			[&Commands, &ReinstanceArrayIt]
 			(int32 Index, CompatibilityCommandBuffer::TCommandVariant& Command)
@@ -174,10 +173,10 @@ namespace UE::Editor::DataStorage
 			});
 
 		// Add a command to process type information stored in TEDS itself in a later step.
-		Commands.AddCommand<FTypeBatchInfoReinstanced>(FTypeBatchInfoReinstanced{ .Batch = ReinstanceView });
+		Commands.AddCommand<FTypeBatchInfoReinstanced>(FTypeBatchInfoReinstanced{ .Batch = ReinstanceArray });
 
 		// Sort the extracted array by the index of the original value for faster indexing in the future.
-		ReinstanceView.Sort([](const FTypeInfoReinstanced& Lhs, const FTypeInfoReinstanced& Rhs)
+		ReinstanceArray.Sort([](const FTypeInfoReinstanced& Lhs, const FTypeInfoReinstanced& Rhs)
 			{
 				return Lhs.Original.GetWeakPtrTypeHash() < Rhs.Original.GetWeakPtrTypeHash();
 			});
@@ -186,9 +185,8 @@ namespace UE::Editor::DataStorage
 		// There can be a large number of re-instanced type information but typically a limited number of type-to-table mappings. Instead
 		// of searching through all the type information updates, copy the type-to-table locally and rebuild it with update info.
 		using TypeTablePair = TPair<TWeakObjectPtr<UStruct>, TableHandle>;
-		TypeTablePair* TempTypeToTable = ScratchBuffer.EmplaceArray<TypeTablePair>(StorageCompat.TypeToTableMap.Num());
-		TArrayView<TypeTablePair> TempTypeToTableView(TempTypeToTable, StorageCompat.TypeToTableMap.Num());
-		TypeTablePair* TempTypeToTableIt = TempTypeToTable;
+		TArrayView<TypeTablePair> TempTypeToTable = ScratchBuffer.EmplaceArray<TypeTablePair>(StorageCompat.TypeToTableMap.Num());
+		TypeTablePair* TempTypeToTableIt = TempTypeToTable.GetData();
 		for (decltype(StorageCompat.TypeToTableMap)::TIterator It = StorageCompat.TypeToTableMap.CreateIterator(); It; ++It)
 		{
 			TempTypeToTableIt->Key = It.Key();
@@ -196,9 +194,9 @@ namespace UE::Editor::DataStorage
 			++TempTypeToTableIt;
 		}
 		StorageCompat.TypeToTableMap.Reset();
-		for (TypeTablePair It : TempTypeToTableView)
+		for (TypeTablePair It : TempTypeToTable)
 		{
-			StorageCompat.TypeToTableMap.Add(FTypeBatchInfoReinstanced::FindObjectRecursively(ReinstanceView, It.Key), It.Value);
+			StorageCompat.TypeToTableMap.Add(FTypeBatchInfoReinstanced::FindObjectRecursively(ReinstanceArray, It.Key), It.Value);
 		}
 
 		// Patch locally cached type information
@@ -211,11 +209,11 @@ namespace UE::Editor::DataStorage
 		for (UScriptStruct** TypeInfo : LocalTypeInfo)
 		{
 			*TypeInfo = Cast<UScriptStruct>(
-				FTypeBatchInfoReinstanced::FindObjectRecursively(ReinstanceView, TWeakObjectPtr<UStruct>(*TypeInfo)));
+				FTypeBatchInfoReinstanced::FindObjectRecursively(ReinstanceArray, TWeakObjectPtr<UStruct>(*TypeInfo)));
 		}
 
 		// Patch existing commands with the new type information.
-		Commands.Process(FPatchCommand(ReinstanceView));
+		Commands.Process(FPatchCommand(ReinstanceArray));
 	}
 	
 
@@ -826,8 +824,8 @@ namespace UE::Editor::DataStorage
 		// If there are more than 1 adds, batch them together into a batch call and nop the additional adds out.
 		if (Count > 1)
 		{
-			RowHandle* Rows = ScratchBuffer.EmplaceArray<RowHandle>(Count);
-			TWeakObjectPtr<UObject>* Objects = ScratchBuffer.EmplaceArray<TWeakObjectPtr<UObject>>(Count);
+			TArrayView<RowHandle> Rows = ScratchBuffer.AllocateUninitializedArray<RowHandle>(Count);
+			TArrayView<TWeakObjectPtr<UObject>> Objects = ScratchBuffer.EmplaceArray<TWeakObjectPtr<UObject>>(Count);
 
 			Rows[0] = Command.Row;
 			Objects[0] = Command.Object;
@@ -846,8 +844,8 @@ namespace UE::Editor::DataStorage
 			Optimizer.ReplaceLeft<FBatchAddCompatibleUObject>(
 				FBatchAddCompatibleUObject
 				{
-					.ObjectArray = Objects,
-					.RowArray = Rows,
+					.ObjectArray = Objects.GetData(),
+					.RowArray = Rows.GetData(),
 					.Table = Command.Table,
 					.Count = static_cast<uint32>(Count)
 				});
@@ -863,9 +861,9 @@ namespace UE::Editor::DataStorage
 		// If there are more than 1 adds, batch them together into a batch call and nop them out.
 		if (Count > 1)
 		{
-			RowHandle* Rows = ScratchBuffer.EmplaceArray<RowHandle>(Count);
-			void** Objects = ScratchBuffer.EmplaceArray<void*>(Count);
-			TWeakObjectPtr<UScriptStruct>* TypeInfo = ScratchBuffer.EmplaceArray<TWeakObjectPtr<UScriptStruct>>(Count);
+			TArrayView<RowHandle> Rows = ScratchBuffer.AllocateUninitializedArray<RowHandle>(Count);
+			TArrayView<void*> Objects = ScratchBuffer.AllocateUninitializedArray<void*>(Count);
+			TArrayView<TWeakObjectPtr<UScriptStruct>> TypeInfo = ScratchBuffer.EmplaceArray<TWeakObjectPtr<UScriptStruct>>(Count);
 
 			Rows[0] = Command.Row;
 			Objects[0] = Command.Object;
@@ -883,9 +881,9 @@ namespace UE::Editor::DataStorage
 			Optimizer.GetLeft().Emplace<FBatchAddCompatibleExternalObject>(
 				FBatchAddCompatibleExternalObject
 				{
-					.ObjectArray = Objects,
-					.TypeInfoArray = TypeInfo,
-					.RowArray = Rows,
+					.ObjectArray = Objects.GetData(),
+					.TypeInfoArray = TypeInfo.GetData(),
+					.RowArray = Rows.GetData(),
 					.Table = Command.Table,
 					.Count = static_cast<uint32>(Count)
 				});
diff --git a/Engine/Plugins/Experimental/EditorDataStorage/Source/TedsCore/Private/Processors/TypedElementProcessorAdaptors.cpp b/Engine/Plugins/Experimental/EditorDataStorage/Source/TedsCore/Private/Processors/TypedElementProcessorAdaptors.cpp
--- a/Engine/Plugins/Experimental/EditorDataStorage/Source/TedsCore/Private/Processors/TypedElementProcessorAdaptors.cpp
+++ b/Engine/Plugins/Experimental/EditorDataStorage/Source/TedsCore/Private/Processors/TypedElementProcessorAdaptors.cpp
@@ -237,7 +237,7 @@ namespace UE::Editor::DataStorage
 						int32 Count = 0;
 						for (FAddedColumnsBatch* const AddColumns : AddedColumnsBatchQueue)
 						{
-							Count += AddColumns->EntityCount;
+							Count += AddColumns->Entities.Num();
 						}
 						const int32 AddedColumnsBathRowsCount = Count;
 
@@ -251,7 +251,7 @@ namespace UE::Editor::DataStorage
 						Count = 0;
 						for (FRemovedColumnsBatch* const RemovedColumns : RemovedColumnsBatchQueue)
 						{
-							Count += RemovedColumns->EntityCount;
+							Count += RemovedColumns->Entities.Num();
 						}
 						const int32 RemovedColumnsBatchRowCount = Count;
 
@@ -275,8 +275,7 @@ namespace UE::Editor::DataStorage
 
 						for (FAddedColumnsBatch* const AddColumns : AddedColumnsBatchQueue)
 						{
-							TArrayView<FMassEntityHandle> Rows(AddColumns->Entities, AddColumns->EntityCount);
-							for (const FMassEntityHandle& Row : Rows)
+							for (const FMassEntityHandle& Row : AddColumns->Entities)
 							{
 								FRowOperations& Operations = RowOperations.FindOrAdd(Row);
 								Operations.ColumnsChange.AddedTags += AddColumns->AddDescriptor.Tags;
@@ -330,8 +329,7 @@ namespace UE::Editor::DataStorage
 
 						for (FRemovedColumnsBatch* const RemovedColumns : RemovedColumnsBatchQueue)
 						{
-							TArrayView<FMassEntityHandle> Rows(RemovedColumns->Entities, RemovedColumns->EntityCount);
-							for (const FMassEntityHandle& Row : Rows)
+							for (const FMassEntityHandle& Row : RemovedColumns->Entities)
 							{
 								FRowOperations& Operations = RowOperations.FindOrAdd(Row);
 								Operations.ColumnsChange.RemovedTags += RemovedColumns->RemoveDescriptor.Tags;
@@ -522,16 +520,12 @@ namespace UE::Editor::DataStorage
 				FScratchBuffer& ScratchBuffer = Environment.GetScratchBuffer();
 				
 				FAddDynamicColumnsBatch* CommandData = ScratchBuffer.Emplace<FAddDynamicColumnsBatch>();
-				TArrayView<RowHandle> ScratchRows = MakeArrayView(
-					ScratchBuffer.EmplaceArray<RowHandle>(Rows.Num()),
-					Rows.Num());
-				TArrayView<FDynamicColumnDescription> ScratchDescriptions = MakeArrayView(
-					ScratchBuffer.EmplaceArray<FDynamicColumnDescription>(DynamicColumnDescriptions.Num()),
-					DynamicColumnDescriptions.Num());
-				
-				TArrayView<const UScriptStruct*> ScratchTypes = MakeArrayView(
-					ScratchBuffer.EmplaceArray<const UScriptStruct*>(DynamicColumnDescriptions.Num()),
-					DynamicColumnDescriptions.Num());
+				TArrayView<RowHandle> ScratchRows = 
+					ScratchBuffer.AllocateUninitializedArray<RowHandle>(Rows.Num());
+				TArrayView<FDynamicColumnDescription> ScratchDescriptions = 
+					ScratchBuffer.AllocateZeroInitializedArray<FDynamicColumnDescription>(DynamicColumnDescriptions.Num());
+				TArrayView<const UScriptStruct*> ScratchTypes = 
+					ScratchBuffer.AllocateUninitializedArray<const UScriptStruct*>(DynamicColumnDescriptions.Num());
 				
 				CopyArrayViews(Rows, ScratchRows);
 				CopyArrayViews(DynamicColumnDescriptions, ScratchDescriptions);
@@ -588,7 +582,7 @@ namespace UE::Editor::DataStorage
 				checkf(UE::Mass::IsA<FMassFragment>(ObjectType), TEXT("Column [%s] can not be a tag"), *ObjectType->GetName());
 	
 				FScratchBuffer& ScratchBuffer = Environment.GetScratchBuffer();
-				void* ColumnData = ScratchBuffer.Allocate(ObjectType->GetStructureSize(), ObjectType->GetMinAlignment());
+				void* ColumnData = ScratchBuffer.AllocateUninitialized(ObjectType->GetStructureSize(), ObjectType->GetMinAlignment());
 				FAddValueColumn* AddedColumn =
 					ScratchBuffer.Emplace<FAddValueColumn>(Relocator, ObjectType, FMassEntityHandle::FromNumber(Row), ColumnData);
 
@@ -638,7 +632,7 @@ namespace UE::Editor::DataStorage
 			{
 				FScratchBuffer& ScratchBuffer = Environment.GetScratchBuffer();
 				// DynamicColumn types are derivations from their template that add no new members.  The size and alignment will be the same
-				void* ColumnData = ScratchBuffer.Allocate(Description.TemplateType->GetStructureSize(), Description.TemplateType->GetMinAlignment());
+				void* ColumnData = ScratchBuffer.AllocateUninitialized(Description.TemplateType->GetStructureSize(), Description.TemplateType->GetMinAlignment());
 				FAddDynamicColumn* AddedColumn =
 					ScratchBuffer.Emplace<FAddDynamicColumn>(Relocator, Description, FMassEntityHandle::FromNumber(Row), ColumnData);
 		
@@ -707,15 +701,14 @@ namespace UE::Editor::DataStorage
 				FAddedColumnsBatch* AddedColumns = ScratchBuffer.Emplace<FAddedColumnsBatch>();
 				TedsColumnsToMassDescriptor(AddedColumns->AddDescriptor, ColumnTypes);
 		
-				FMassEntityHandle* Entities = ScratchBuffer.EmplaceArray<FMassEntityHandle>(Rows.Num());
-				AddedColumns->Entities = Entities;
+				AddedColumns->Entities = ScratchBuffer.AllocateZeroInitializedArray<FMassEntityHandle>(Rows.Num());
+				FMassEntityHandle* Entities = AddedColumns->Entities.GetData();
 				for (RowHandle Row : Rows)
 				{
 					*Entities = FMassEntityHandle::FromNumber(Row);
 					Entities++;
 				}
-				AddedColumns->EntityCount = Rows.Num();
-
+				
 				if (bBatchDeferredCommands)
 				{
 					AddedColumnsBatchQueue.AddElement(AddedColumns);
@@ -725,13 +718,11 @@ namespace UE::Editor::DataStorage
 				this->Context.Defer().template PushCommand<FMassDeferredAddCommand>(
 					[AddedColumns](FMassEntityManager& System)
 					{
-						FMassEntityHandle* Entities = AddedColumns->Entities;
-						int32 Count = AddedColumns->EntityCount;
-						for (int32 Counter = 0; Counter < Count; ++Counter)
+						for (FMassEntityHandle& Entity : AddedColumns->Entities)
 						{
-							if (System.IsEntityValid(*Entities))
+							if (System.IsEntityValid(Entity))
 							{
-								System.AddCompositionToEntity_GetDelta(*Entities++, AddedColumns->AddDescriptor);
+								System.AddCompositionToEntity_GetDelta(Entity, AddedColumns->AddDescriptor);
 							}
 						}
 					});
@@ -771,14 +762,13 @@ namespace UE::Editor::DataStorage
 				FRemovedColumnsBatch* RemovedColumns = ScratchBuffer.Emplace<FRemovedColumnsBatch>();
 				TedsColumnsToMassDescriptorIfActiveTable(RemovedColumns->RemoveDescriptor, ColumnTypes);
 
-				FMassEntityHandle* Entities = ScratchBuffer.EmplaceArray<FMassEntityHandle>(Rows.Num());
-				RemovedColumns->Entities = Entities;
+				RemovedColumns->Entities = ScratchBuffer.EmplaceArray<FMassEntityHandle>(Rows.Num());
+				FMassEntityHandle* Entities = RemovedColumns->Entities.GetData();
 				for (RowHandle Row : Rows)
 				{
 					*Entities = FMassEntityHandle::FromNumber(Row);
 					Entities++;
 				}
-				RemovedColumns->EntityCount = Rows.Num();
 
 				if (bBatchDeferredCommands)
 				{
@@ -789,9 +779,8 @@ namespace UE::Editor::DataStorage
 					this->Context.Defer().template PushCommand<FMassDeferredRemoveCommand>(
 						[RemovedColumns](FMassEntityManager& System)
 						{
-							FMassEntityHandle* Entities = RemovedColumns->Entities;
-							int32 Count = RemovedColumns->EntityCount;
-
+							TArrayView<FMassEntityHandle> Entities = RemovedColumns->Entities;
+							
 							using EntityHandleArray = TArray<FMassEntityHandle, TInlineAllocator<32>>;
 							using EntityArchetypeLookup = TMap<FMassArchetypeHandle, EntityHandleArray, TInlineSetAllocator<32>>;
 							using ArchetypeEntityArray = TArray<FMassArchetypeEntityCollection, TInlineAllocator<32>>;
@@ -804,15 +793,14 @@ namespace UE::Editor::DataStorage
 							// This could be improved by adding an operation that would both combine the Fragments and Tags change in one bath operation.
 							if (!RemovedColumns->RemoveDescriptor.Fragments.IsEmpty())
 							{
-								for (int32 Counter = 0; Counter < Count; ++Counter)
+								for (FMassEntityHandle Entity : Entities)
 								{
-									if (System.IsEntityValid(*Entities))
+									if (System.IsEntityValid(Entity))
 									{
-										FMassArchetypeHandle Archetype = System.GetArchetypeForEntity(*Entities);
+										FMassArchetypeHandle Archetype = System.GetArchetypeForEntity(Entity);
 										EntityHandleArray& EntityCollection = LookupTable.FindOrAdd(Archetype);
-										EntityCollection.Add(*Entities);
+										EntityCollection.Add(Entity);
 									}
-									Entities++;
 								}
 
 								// Construct table (archetype) specific row (entity) collections.
@@ -830,15 +818,14 @@ namespace UE::Editor::DataStorage
 								LookupTable.Reset();
 								EntityCollections.Reset();
 
-								for (int32 Counter = 0; Counter < Count; ++Counter)
+								for (FMassEntityHandle Entity : Entities)
 								{
-									if (System.IsEntityValid(*Entities))
+									if (System.IsEntityValid(Entity))
 									{
-										FMassArchetypeHandle Archetype = System.GetArchetypeForEntity(*Entities);
+										FMassArchetypeHandle Archetype = System.GetArchetypeForEntity(Entity);
 										EntityHandleArray& EntityCollection = LookupTable.FindOrAdd(Archetype);
-										EntityCollection.Add(*Entities);
+										EntityCollection.Add(Entity);
 									}
-									Entities++;
 								}
 
 								// Construct table (archetype) specific row (entity) collections.
@@ -931,31 +918,34 @@ namespace UE::Editor::DataStorage
 
 			void* EmplaceObjectInScratch(size_t ObjectSize, size_t Alignment, void(* Construct)(void*, void*), void(* Destroy)(void*), void* SourceCommandContext)
 			{
+				struct FDestructor final
+				{
+					using DestroyFnType = decltype(Destroy);
+					FDestructor(DestroyFnType InDestroyFn, void* InObjectPtr)
+						: DestroyFn(InDestroyFn)
+						, ObjectPtr(InObjectPtr)
+					{
+					}
+					~FDestructor()
+					{
+						DestroyFn(ObjectPtr);
+					}
+					DestroyFnType DestroyFn;
+					void* ObjectPtr;
+				};
+
 				FScratchBuffer& ScratchBuffer = Environment.GetScratchBuffer();
-				void* ObjectMemory = ScratchBuffer.Allocate(ObjectSize, Alignment);
+				
+				void* ObjectMemory = ScratchBuffer.AllocateUninitialized(ObjectSize, Alignment);
 				Construct(ObjectMemory, SourceCommandContext);
+
 				// The presence of a Destroy function implies that the objects that was just added to the scratch buffer
-				// is not trivially destructable, hence needs its destructor called.
-				// The API for the scratch buffer's internal memory allocator needs us to emplace a non-trivially destructable object
-				// of some type.  FDestructor is used to fulfil that role to destroy the object that was just constructed.
+				// is not trivially destructible, hence needs its destructor called.
+				// The API for the scratch buffer's internal memory allocator needs us to emplace a non-trivially destructible object
+				// of some type. FDestructor is used to fulfill that role to destroy the object that was just constructed.
 				if (Destroy)
 				{
-					struct FDestructor
-					{
-						using DestroyFnType = decltype(Destroy);
-						using ObjectPtrType = decltype(ObjectMemory);
-						FDestructor(DestroyFnType InDestroyFn, ObjectPtrType InObjectPtr)
-							: DestroyFn(InDestroyFn)
-							, ObjectPtr(InObjectPtr)
-						{}
-						~FDestructor()
-						{
-							DestroyFn(ObjectPtr);
-						}
-						DestroyFnType DestroyFn;
-						ObjectPtrType ObjectPtr;
-					};
-					ScratchBuffer.Emplace<FDestructor>(Destroy, ObjectMemory);
+					ScratchBuffer.Emplace<FDestructor>(Destroy, ObjectMemory);	
 				}
 				return ObjectMemory;
 			}
@@ -969,8 +959,7 @@ namespace UE::Editor::DataStorage
 			struct FAddedColumnsBatch
 			{
 				FMassArchetypeCompositionDescriptor AddDescriptor;
-				FMassEntityHandle* Entities;
-				int32 EntityCount;
+				TArrayView<FMassEntityHandle> Entities;
 			};
 
 			struct FAddValueColumn
@@ -1042,8 +1031,7 @@ namespace UE::Editor::DataStorage
 			struct FRemovedColumnsBatch
 			{
 				FMassArchetypeCompositionDescriptor RemoveDescriptor;
-				FMassEntityHandle* Entities;
-				int32 EntityCount;
+				TArrayView<FMassEntityHandle> Entities;
 			};
 
 			struct FAddRow
diff --git a/Engine/Plugins/Experimental/EditorDataStorage/Source/TedsCore/Private/TypedElementDatabase.cpp b/Engine/Plugins/Experimental/EditorDataStorage/Source/TedsCore/Private/TypedElementDatabase.cpp
--- a/Engine/Plugins/Experimental/EditorDataStorage/Source/TedsCore/Private/TypedElementDatabase.cpp
+++ b/Engine/Plugins/Experimental/EditorDataStorage/Source/TedsCore/Private/TypedElementDatabase.cpp
@@ -331,7 +331,7 @@ void UEditorDataStorage::BatchReserveRows(int32 Count, TFunctionRef<void(UE::Edi
 {
 	if (ActiveEditorEntityManager)
 	{
-		TArrayView<FMassEntityHandle> ReservedEntities(Environment->GetScratchBuffer().EmplaceArray<FMassEntityHandle>(Count), Count);
+		TArrayView<FMassEntityHandle> ReservedEntities = Environment->GetScratchBuffer().AllocateZeroInitializedArray<FMassEntityHandle>(Count);
 		ActiveEditorEntityManager->BatchReserveEntities(ReservedEntities);
 
 		for (FMassEntityHandle ReservedEntity : ReservedEntities)
diff --git a/Engine/Plugins/Experimental/EditorDataStorage/Source/TedsCore/Private/TypedElementDatabaseCommandBuffer.cpp b/Engine/Plugins/Experimental/EditorDataStorage/Source/TedsCore/Private/TypedElementDatabaseCommandBuffer.cpp
--- a/Engine/Plugins/Experimental/EditorDataStorage/Source/TedsCore/Private/TypedElementDatabaseCommandBuffer.cpp
+++ b/Engine/Plugins/Experimental/EditorDataStorage/Source/TedsCore/Private/TypedElementDatabaseCommandBuffer.cpp
@@ -132,9 +132,8 @@ namespace UE::Editor::DataStorage::Legacy
 		PendingColumnMappingKey Key(Row, ColumnType);
 		if (void** StoredData = PendingColumns.Find(Key); StoredData == nullptr || *StoredData == nullptr)
 		{
-			void* Data = Environment.GetScratchBuffer().Allocate(ColumnType->GetStructureSize(), ColumnType->GetMinAlignment());
 			// Initialize to zero to replicate the default from Mass.
-			FMemory::Memzero(Data, ColumnType->GetStructureSize());
+			void* Data = Environment.GetScratchBuffer().AllocateZeroInitialized(ColumnType->GetStructureSize(), ColumnType->GetMinAlignment());
 			PendingColumns.Add(PendingColumnMappingKey(Row, ColumnType), Data);
 			AddCommand(Row, FAddDataColumnCommand
 				{
diff --git a/Engine/Plugins/Experimental/EditorDataStorage/Source/TedsCore/Private/TypedElementDatabaseScratchBuffer.cpp b/Engine/Plugins/Experimental/EditorDataStorage/Source/TedsCore/Private/TypedElementDatabaseScratchBuffer.cpp
--- a/Engine/Plugins/Experimental/EditorDataStorage/Source/TedsCore/Private/TypedElementDatabaseScratchBuffer.cpp
+++ b/Engine/Plugins/Experimental/EditorDataStorage/Source/TedsCore/Private/TypedElementDatabaseScratchBuffer.cpp
@@ -11,11 +11,16 @@ namespace UE::Editor::DataStorage
 	{
 	}
 
-	void* FScratchBuffer::Allocate(SIZE_T Size, uint32 Alignment)
+	void* FScratchBuffer::AllocateUninitialized(SIZE_T Size, uint32 Alignment)
 	{
 		return CurrentAllocator.load()->Malloc(Size, Alignment);
 	}
 
+	void* FScratchBuffer::AllocateZeroInitialized(SIZE_T Size, uint32 Alignment)
+	{
+		return CurrentAllocator.load()->MallocAndMemset(Size, Alignment, 0);
+	}
+
 	void FScratchBuffer::BatchDelete()
 	{
 		MemoryAllocator* Current = CurrentAllocator.exchange(LeastRecentAllocator);
diff --git a/Engine/Plugins/Experimental/EditorDataStorage/Source/TedsCore/Private/TypedElementDatabaseScratchBuffer.h b/Engine/Plugins/Experimental/EditorDataStorage/Source/TedsCore/Private/TypedElementDatabaseScratchBuffer.h
--- a/Engine/Plugins/Experimental/EditorDataStorage/Source/TedsCore/Private/TypedElementDatabaseScratchBuffer.h
+++ b/Engine/Plugins/Experimental/EditorDataStorage/Source/TedsCore/Private/TypedElementDatabaseScratchBuffer.h
@@ -3,6 +3,7 @@
 #pragma once
 
 #include <atomic>
+#include <type_traits>
 #include "Experimental/ConcurrentLinearAllocator.h"
 
 // A thread-safe memory allocator that uses linear allocation. This provides a fast and lightweight way to allocate temporary memory 
@@ -16,8 +17,8 @@ namespace UE::Editor::DataStorage
 		{
 			static constexpr uint32 BlockSize = 64 * 1024;		// Blocksize used to allocate from
 			static constexpr bool AllowOversizedBlocks = true;  // The allocator supports oversized Blocks and will store them in a separate Block with counter 1
-			static constexpr bool RequiresAccurateSize = true;  // GetAllocationSize returning the accurate size of the allocation otherwise it could be relaxed to return the size to the end of the Block
-			static constexpr bool InlineBlockAllocation = false;  // Inline or Noinline the BlockAllocation which can have an impact on Performance
+			static constexpr bool RequiresAccurateSize = false;  // GetAllocationSize returning the accurate size of the allocation otherwise it could be relaxed to return the size to the end of the Block
+			static constexpr bool InlineBlockAllocation = true;  // Inline or Noinline the BlockAllocation which can have an impact on Performance
 			static constexpr const char* TagName = "TedsScratchBuffer";
 	
 			using Allocator = TBlockAllocationLockFreeCache<BlockSize, FOsAllocator>;
@@ -25,13 +26,63 @@ namespace UE::Editor::DataStorage
 		using MemoryAllocator = TConcurrentLinearBulkObjectAllocator<FAllocationTag>;
 	
 	public:
+		/** Whether or not functions that do not handle destruction of typed objects allow classes with non-trivial destructor. */
+		enum class ERequiresTrivialDestructor : bool
+		{
+			Yes,
+			No
+		};
+
 		FScratchBuffer();
 	
-		void* Allocate(SIZE_T Size, uint32 Alignment);
+		/** Allocate uninitialized memory. */
+		void* AllocateUninitialized(SIZE_T Size, uint32 Alignment);
+
+		/** 
+		 * Allocates memory for a single object but does not initialize the memory. Memory is only reserved and the destructor 
+		 * will not be called.
+		 */
+		template<typename T, ERequiresTrivialDestructor RequiresTrivialDestructors = ERequiresTrivialDestructor::Yes>
+		T* AllocateUninitialized();
+		
+		/** 
+		 * Allocates memory for an array of object but does not initialize the memory. Memory is only reserved and the 
+		 * destructor will not be called.
+		 */
+		template<typename T, ERequiresTrivialDestructor RequiresTrivialDestructors = ERequiresTrivialDestructor::Yes>
+		TArrayView<T> AllocateUninitializedArray(int32 Count);
+
+		/** Allocate memory and sets it to zero. */
+		void* AllocateZeroInitialized(SIZE_T Size, uint32 Alignment);
+
+		/** 
+		 * Allocates memory for a single object and sets the memory to zero. This is not the same as calling a constructor, which can be 
+		 * done using Emplace. Memory is only reserved and the destructor will not be called.
+		 */
+		template<typename T, ERequiresTrivialDestructor RequiresTrivialDestructors = ERequiresTrivialDestructor::Yes>
+		T* AllocateZeroInitialized();
+
+		/** 
+		 * Allocates memory for an array of object and sets the memory to zero. This is not the same as calling a constructor, which can be 
+		 * done using Emplace. Memory is only reserved and the destructor will not be called.
+		 */
+		template<typename T, ERequiresTrivialDestructor RequiresTrivialDestructors = ERequiresTrivialDestructor::Yes>
+		TArrayView<T> AllocateZeroInitializedArray(int32 Count);
+
+		/** 
+		 * Allocate memory for one object and construct it with the provided arguments. The object created by this call will
+		 * has its constructor and destructor called.
+		 */
 		template<typename T, typename... ArgTypes>
 		T* Emplace(ArgTypes... Args);
+
+		/** 
+		 * Allocate memory for an array of objects and initialize each object with the provided arguments. Objects created by this 
+		 * call will have their constructor and destructor called.
+		 */
 		template<typename T, typename... ArgTypes>
-		T* EmplaceArray(int32 Count, const ArgTypes&... Args);
+		TArrayView<T> EmplaceArray(int32 Count, const ArgTypes&... Args);
+
 	
 		// Activates a new allocator and deletes all commands and objects in the least recently touched scratch buffer.
 		void BatchDelete();
@@ -51,6 +102,52 @@ namespace UE::Editor::DataStorage
 		MemoryAllocator* LeastRecentAllocator;
 	};
 	
+	template<typename T, FScratchBuffer::ERequiresTrivialDestructor RequiresTrivialDestructors>
+	T* FScratchBuffer::AllocateUninitialized()
+	{
+		if constexpr (RequiresTrivialDestructors == ERequiresTrivialDestructor::Yes)
+		{
+			static_assert(std::is_trivially_destructible_v<T>, 
+				"Scratch buffer allocator requires a trivially destructible class type or to be explicitly told it's safe to construct.");
+		}
+		return CurrentAllocator.load()->Malloc<T>();
+	}
+
+	template<typename T, FScratchBuffer::ERequiresTrivialDestructor RequiresTrivialDestructors>
+	TArrayView<T> FScratchBuffer::AllocateUninitializedArray(int32 Count)
+	{
+		if constexpr (RequiresTrivialDestructors == ERequiresTrivialDestructor::Yes)
+		{
+			static_assert(std::is_trivially_destructible_v<T>,
+				"Scratch buffer allocator requires a trivially destructible class type or to be explicitly told it's safe to construct.");
+		}
+		T* Result = CurrentAllocator.load()->MallocArray<T>(Count);
+		return TArrayView<T>(Result, Result ? Count : 0);
+	}
+
+	template<typename T, FScratchBuffer::ERequiresTrivialDestructor RequiresTrivialDestructors>
+	T* FScratchBuffer::AllocateZeroInitialized()
+	{
+		if constexpr (RequiresTrivialDestructors == ERequiresTrivialDestructor::Yes)
+		{
+			static_assert(std::is_trivially_destructible_v<T>,
+				"Scratch buffer allocator requires a trivially destructible class type or to be explicitly told it's safe to construct.");
+		}
+		return CurrentAllocator.load()->MallocAndMemset<T>(0);
+	}
+
+	template<typename T, FScratchBuffer::ERequiresTrivialDestructor RequiresTrivialDestructors>
+	TArrayView<T> FScratchBuffer::AllocateZeroInitializedArray(int32 Count)
+	{
+		if constexpr (RequiresTrivialDestructors == ERequiresTrivialDestructor::Yes)
+		{
+			static_assert(std::is_trivially_destructible_v<T>,
+				"Scratch buffer allocator requires a trivially destructible class type or to be explicitly told it's safe to construct.");
+		}
+		T* Result = CurrentAllocator.load()->MallocAndMemsetArray<T>(Count, 0);
+		return TArrayView<T>(Result, Result ? Count : 0);
+	}
+
 	template<typename T, typename... ArgTypes>
 	T* FScratchBuffer::Emplace(ArgTypes... Args)
 	{
@@ -58,11 +155,12 @@ namespace UE::Editor::DataStorage
 	}
 	
 	template<typename T, typename... ArgTypes>
-	T* FScratchBuffer::EmplaceArray(int32 Count, const ArgTypes&... Args)
+	TArrayView<T> FScratchBuffer::EmplaceArray(int32 Count, const ArgTypes&... Args)
 	{
-		return CurrentAllocator.load()->CreateArray<T>(Count, Args...);
+		T* Result = CurrentAllocator.load()->CreateArray<T>(Count, Args...);
+		return TArrayView<T>(Result, Result ? Count : 0);
 	}
-	
+
 	constexpr int32 FScratchBuffer::MaxAllocationSize()
 	{
 		return FAllocationTag::BlockSize;
diff --git a/Engine/Plugins/Experimental/EditorDataStorage/Source/TedsCore/Tests/TypedElementDatabaseScratchBufferTests.cpp b/Engine/Plugins/Experimental/EditorDataStorage/Source/TedsCore/Tests/TypedElementDatabaseScratchBufferTests.cpp
--- a/Engine/Plugins/Experimental/EditorDataStorage/Source/TedsCore/Tests/TypedElementDatabaseScratchBufferTests.cpp
+++ b/Engine/Plugins/Experimental/EditorDataStorage/Source/TedsCore/Tests/TypedElementDatabaseScratchBufferTests.cpp
@@ -47,7 +47,7 @@ TEST_CASE_NAMED(Editor_DataStorage_ScratchBuffer_MultiThreaded_Tests, "Editor::D
 				{
 					while (bKeepRunning)
 					{
-						void* Data = Buffer.Allocate(128, 4);
+						void* Data = Buffer.AllocateUninitialized(128, 4);
 						if (Data == nullptr)
 						{
 							FailedAllocations++;
@@ -106,7 +106,7 @@ TEST_CASE_NAMED(Editor_DataStorage_ScratchBuffer_MultiThreaded_Tests, "Editor::D
 							MemoryAllocationSizeIndex = 0;
 						}
 
-						void* Data = Buffer.Allocate(static_cast<size_t>(MemoryAllocationSizes[MemoryAllocationSizeIndex]), 4);
+						void* Data = Buffer.AllocateUninitialized(static_cast<size_t>(MemoryAllocationSizes[MemoryAllocationSizeIndex]), 4);
 						if (Data == nullptr)
 						{
 							FailedAllocations++;
@@ -149,22 +149,22 @@ TEST_CASE_NAMED(Editor_DataStorage_ScratchBuffer_Tests, "Editor::DataStorage::Sc
 	SECTION("Allocate small block")
 	{
 		FScratchBuffer Buffer;
-		void* Data = Buffer.Allocate(4, 4);
+		void* Data = Buffer.AllocateUninitialized(4, 4);
 		CHECK(Data != nullptr);
 	}
 
 	SECTION("Allocate over-sized block")
 	{
 		FScratchBuffer Buffer;
-		void* Data = Buffer.Allocate(Buffer.MaxAllocationSize() * 4, 4);
+		void* Data = Buffer.AllocateUninitialized(Buffer.MaxAllocationSize() * 4, 4);
 		CHECK(Data != nullptr);
 	}
 
 	SECTION("Alignment respected")
 	{
 		FScratchBuffer Buffer;
-		void* Data1 = Buffer.Allocate(1, 1);
-		void* Data2 = Buffer.Allocate(4, 4);
+		void* Data1 = Buffer.AllocateUninitialized(1, 1);
+		void* Data2 = Buffer.AllocateUninitialized(4, 4);
 
 		REQUIRE(Data1 != nullptr);
 		REQUIRE(Data2 != nullptr);
@@ -179,7 +179,7 @@ TEST_CASE_NAMED(Editor_DataStorage_ScratchBuffer_Tests, "Editor::DataStorage::Sc
 		int32 IncrementCount = (Buffer.MaxAllocationSize() * 4 /* Fill 4 blocks*/) / 64 /* With 64 byte allocations */;
 		for (int32 Counter = 0; Counter < IncrementCount; ++Counter)
 		{
-			void* Data = Buffer.Allocate(64, 4);
+			void* Data = Buffer.AllocateUninitialized(64, 4);
 			REQUIRE(Data != nullptr);
 		}
 	}
@@ -193,11 +193,39 @@ TEST_CASE_NAMED(Editor_DataStorage_ScratchBuffer_Tests, "Editor::DataStorage::Sc
 			int32 IncrementCount = Buffer.MaxAllocationSize() / 64;
 			for (int32 Counter = 0; Counter <= IncrementCount; ++Counter) // Increment one more than needed to fill the block and make sure it's recycled.
 			{
-				void* Data = Buffer.Allocate(64, 4);
+				void* Data = Buffer.AllocateUninitialized(64, 4);
 				REQUIRE(Data != nullptr);
 			}
 			Buffer.BatchDelete();
 		}
 	}
+
+	SECTION("Order of destruction.")
+	{
+		int Counter = 5;
+		
+		{
+			FScratchBuffer Buffer;
+			struct TestObject
+			{
+				int Id;
+				int* Counter;
+				TestObject(int Id, int* Counter) : Id(Id), Counter(Counter) {}
+				~TestObject()
+				{
+					REQUIRE(Id == *Counter);
+					(*Counter)--;
+				}
+			};
+
+			for (int Index = 0; Index <= Counter; ++Index)
+			{
+				Buffer.Emplace<TestObject>(Index, &Counter);
+			}
+
+			// The objects are expected to be deleted in the reverse order they are declared.
+		}
+		REQUIRE(Counter == -1);
+	}
 }
 #endif // WITH_TESTS
diff --git a/Engine/Source/Runtime/TypedElementFramework/Public/Elements/Interfaces/TypedElementQueryStorageInterfaces.h b/Engine/Source/Runtime/TypedElementFramework/Public/Elements/Interfaces/TypedElementQueryStorageInterfaces.h
--- a/Engine/Source/Runtime/TypedElementFramework/Public/Elements/Interfaces/TypedElementQueryStorageInterfaces.h
+++ b/Engine/Source/Runtime/TypedElementFramework/Public/Elements/Interfaces/TypedElementQueryStorageInterfaces.h
@@ -309,6 +309,7 @@ namespace UE::Editor::DataStorage
 			void(*Destroy)(void*);
 			void* SourceObject;
 		};
+		// Required as the interface doesn't not have direct access to any temporary handling of memory in the backend.
 		virtual void* EmplaceObjectInScratch(const FEmplaceObjectParams& Params) = 0;
 	};
 
@@ -640,21 +641,25 @@ namespace UE::Editor::DataStorage
 	template <typename T>
 	void ICommonQueryWithEnvironmentContext::PushCommand(T CommandContext)
 	{
-		void (*CommandFunction)(void* /*Context*/);
-
-		// If a member operator() is defined
-		CommandFunction = [](void* InInstanceOfT)
-		{
-			T* Instance = static_cast<T*>(InInstanceOfT);
-			Instance->operator()();
-		};
-			
-		if (std::is_empty_v<T>)
+		if constexpr(std::is_empty_v<T>)
 		{
-			PushCommand(CommandFunction, nullptr);
+			// The command object doesn't hold any arguments so use a cheaper path that only stores the callback.
+			PushCommand(
+				[](void*)
+				{
+					T Instance;
+					Instance->operator()();
+				}, nullptr);
 		}
 		else
 		{
+			// Wrapper lambda to store the type and call operator().
+			auto CommandFunction = [](void* InInstanceOfT)
+				{
+					T* Instance = static_cast<T*>(InInstanceOfT);
+					Instance->operator()();
+				};
+
 			FEmplaceObjectParams Params;
 
 			Params.ObjectSize = sizeof(T);
