[TEDS Core API] Refactoring TEDS Core API UScriptStructs to use const and const_cast

diff --git a/Engine/Plugins/Experimental/EditorDataStorage/Source/TedsCore/Private/Commands/EditorDataStorageCompatibilityCommands.cpp b/Engine/Plugins/Experimental/EditorDataStorage/Source/TedsCore/Private/Commands/EditorDataStorageCompatibilityCommands.cpp
--- a/Engine/Plugins/Experimental/EditorDataStorage/Source/TedsCore/Private/Commands/EditorDataStorageCompatibilityCommands.cpp
+++ b/Engine/Plugins/Experimental/EditorDataStorage/Source/TedsCore/Private/Commands/EditorDataStorageCompatibilityCommands.cpp
@@ -136,10 +136,10 @@ namespace UE::Editor::DataStorage
 
 	void FPatchData::FPatchCommand::operator()(FBatchAddCompatibleExternalObject& Command)
 	{
-		TArrayView<TWeakObjectPtr<UScriptStruct>> TypeInfoArray = 
-			TArrayView<TWeakObjectPtr<UScriptStruct>>(Command.TypeInfoArray, Command.Count);
+		TArrayView<TWeakObjectPtr<const UScriptStruct>> TypeInfoArray = 
+			TArrayView<TWeakObjectPtr<const UScriptStruct>>(Command.TypeInfoArray, Command.Count);
 
-		for (TWeakObjectPtr<UScriptStruct>& TypeInfo : TypeInfoArray)
+		for (TWeakObjectPtr<const UScriptStruct>& TypeInfo : TypeInfoArray)
 		{
 			TypeInfo = FTypeBatchInfoReinstanced::FindObjectRecursively(Reinstances, TypeInfo);
 			checkf(TypeInfo.IsValid(), TEXT("A script struct has be re-instanced to an object that's not a script struct."))
@@ -511,7 +511,7 @@ namespace UE::Editor::DataStorage
 		StorageCompatibility.TriggerOnObjectAdded(Object, Object->GetClass(), Row);
 	}
 
-	void FCommandProcessor::SetupRow(RowHandle Row, void* Object, TWeakObjectPtr<UScriptStruct> TypeInfo)
+	void FCommandProcessor::SetupRow(RowHandle Row, void* Object, TWeakObjectPtr<const UScriptStruct> TypeInfo)
 	{
 		Storage.AddColumn(Row, FTypedElementExternalObjectColumn{ .Object = Object });
 		Storage.AddColumn(Row, FTypedElementScriptStructTypeInfoColumn{ .TypeInfo = TypeInfo });
@@ -602,7 +602,7 @@ namespace UE::Editor::DataStorage
 		TConstArrayView<RowHandle> Rows(Batch.RowArray, Batch.Count);
 		RowHandle* RowIterator = Batch.RowArray;
 		void** ObjectIterator = Batch.ObjectArray;
-		TWeakObjectPtr<UScriptStruct>* TypeInfoIterator = Batch.TypeInfoArray;
+		TWeakObjectPtr<const UScriptStruct>* TypeInfoIterator = Batch.TypeInfoArray;
 
 		Storage.BatchAddRow(Batch.Table, Rows, [&RowIterator, &ObjectIterator, &TypeInfoIterator, this](RowHandle Row)
 			{
@@ -863,7 +863,7 @@ namespace UE::Editor::DataStorage
 		{
 			TArrayView<RowHandle> Rows = ScratchBuffer.AllocateUninitializedArray<RowHandle>(Count);
 			TArrayView<void*> Objects = ScratchBuffer.AllocateUninitializedArray<void*>(Count);
-			TArrayView<TWeakObjectPtr<UScriptStruct>> TypeInfo = ScratchBuffer.EmplaceArray<TWeakObjectPtr<UScriptStruct>>(Count);
+			TArrayView<TWeakObjectPtr<const UScriptStruct>> TypeInfo = ScratchBuffer.EmplaceArray<TWeakObjectPtr<const UScriptStruct>>(Count);
 
 			Rows[0] = Command.Row;
 			Objects[0] = Command.Object;
diff --git a/Engine/Plugins/Experimental/EditorDataStorage/Source/TedsCore/Private/Commands/EditorDataStorageCompatibilityCommands.h b/Engine/Plugins/Experimental/EditorDataStorage/Source/TedsCore/Private/Commands/EditorDataStorageCompatibilityCommands.h
--- a/Engine/Plugins/Experimental/EditorDataStorage/Source/TedsCore/Private/Commands/EditorDataStorageCompatibilityCommands.h
+++ b/Engine/Plugins/Experimental/EditorDataStorage/Source/TedsCore/Private/Commands/EditorDataStorageCompatibilityCommands.h
@@ -108,7 +108,7 @@ namespace UE::Editor::DataStorage
 	struct FAddCompatibleExternalObject final
 	{
 		void* Object;
-		TWeakObjectPtr<UScriptStruct> TypeInfo;
+		TWeakObjectPtr<const UScriptStruct> TypeInfo;
 		RowHandle Row;
 		TableHandle Table;
 	};
@@ -145,7 +145,7 @@ namespace UE::Editor::DataStorage
 	struct FBatchAddCompatibleExternalObject final
 	{
 		void** ObjectArray;
-		TWeakObjectPtr<UScriptStruct>* TypeInfoArray;
+		TWeakObjectPtr<const UScriptStruct>* TypeInfoArray;
 		RowHandle* RowArray;
 		TableHandle Table;
 		uint32 Count;
@@ -302,7 +302,7 @@ namespace UE::Editor::DataStorage
 		FCommandProcessor(ICoreProvider& InStorage, UEditorDataStorageCompatibility& InStorageCompatibility);
 
 		void SetupRow(RowHandle Row, UObject* Object);
-		void SetupRow(RowHandle Row, void* Object, TWeakObjectPtr<UScriptStruct> TypeInfo);
+		void SetupRow(RowHandle Row, void* Object, TWeakObjectPtr<const UScriptStruct> TypeInfo);
 		
 		void operator()(FNopCommand&) {}
 		void operator()(FTypeInfoReinstanced&) {}
diff --git a/Engine/Plugins/Experimental/EditorDataStorage/Source/TedsCore/Private/TypedElementDatabase.cpp b/Engine/Plugins/Experimental/EditorDataStorage/Source/TedsCore/Private/TypedElementDatabase.cpp
--- a/Engine/Plugins/Experimental/EditorDataStorage/Source/TedsCore/Private/TypedElementDatabase.cpp
+++ b/Engine/Plugins/Experimental/EditorDataStorage/Source/TedsCore/Private/TypedElementDatabase.cpp
@@ -48,7 +48,7 @@ FAutoConsoleCommandWithOutputDevice PrintSupportedColumnsConsoleCommand(
 		{
 			Output.Log(TEXT("The Editor Data Storage supports the following columns:"));
 
-			for (TObjectIterator<UScriptStruct> It; It; ++It)
+			for (TObjectIterator<const UScriptStruct> It; It; ++It)
 			{
 				if (UE::Mass::IsA<FMassFragment>(*It) || UE::Mass::IsA<FMassTag>(*It))
 				{
@@ -505,7 +505,7 @@ void UEditorDataStorage::BatchRemoveRows(TConstArrayView<UE::Editor::DataStorage
 	}
 }
 
-void UEditorDataStorage::RemoveAllRowsWithColumns(TConstArrayView<UScriptStruct*> Columns)
+void UEditorDataStorage::RemoveAllRowsWithColumns(TConstArrayView<const UScriptStruct*> Columns)
 {
 	if (ActiveEditorEntityManager)
 	{
diff --git a/Engine/Plugins/Experimental/EditorDataStorage/Source/TedsCore/Private/TypedElementDatabase.h b/Engine/Plugins/Experimental/EditorDataStorage/Source/TedsCore/Private/TypedElementDatabase.h
--- a/Engine/Plugins/Experimental/EditorDataStorage/Source/TedsCore/Private/TypedElementDatabase.h
+++ b/Engine/Plugins/Experimental/EditorDataStorage/Source/TedsCore/Private/TypedElementDatabase.h
@@ -9,7 +9,6 @@
 #include "Templates/SharedPointer.h"
 #include "TypedElementDatabaseCommandBuffer.h"
 #include "UObject/ObjectMacros.h"
-#include "UObject/StrongObjectPtr.h"
 
 #include "TypedElementDatabase.generated.h"
 
@@ -104,7 +103,7 @@ public:
 		UE::Editor::DataStorage::RowCreationCallbackRef OnCreated) override;
 	virtual void RemoveRow(UE::Editor::DataStorage::RowHandle Row) override;
 	virtual void BatchRemoveRows(TConstArrayView<UE::Editor::DataStorage::RowHandle> Rows) override;
-	virtual void RemoveAllRowsWithColumns(TConstArrayView<UScriptStruct*> Columns) override;
+	virtual void RemoveAllRowsWithColumns(TConstArrayView<const UScriptStruct*> Columns) override;
 	virtual bool IsRowAvailable(UE::Editor::DataStorage::RowHandle Row) const override;
 	virtual bool IsRowAssigned(UE::Editor::DataStorage::RowHandle Row) const override;
 	/** Same as IsRowAvailable, but doesn't check if the data storage has been initialized. */
diff --git a/Engine/Plugins/Experimental/EditorDataStorage/Source/TedsCore/Private/TypedElementDatabaseCompatibility.cpp b/Engine/Plugins/Experimental/EditorDataStorage/Source/TedsCore/Private/TypedElementDatabaseCompatibility.cpp
--- a/Engine/Plugins/Experimental/EditorDataStorage/Source/TedsCore/Private/TypedElementDatabaseCompatibility.cpp
+++ b/Engine/Plugins/Experimental/EditorDataStorage/Source/TedsCore/Private/TypedElementDatabaseCompatibility.cpp
@@ -130,7 +130,7 @@ void UEditorDataStorageCompatibility::RegisterDealiaserCallback(ObjectToRowDeali
 }
 
 void UEditorDataStorageCompatibility::RegisterTypeTableAssociation(
-	TObjectPtr<UStruct> TypeInfo, UE::Editor::DataStorage::TableHandle Table)
+	TWeakObjectPtr<UStruct> TypeInfo, UE::Editor::DataStorage::TableHandle Table)
 {
 	using namespace UE::Editor::DataStorage;
 	
@@ -216,7 +216,8 @@ UE::Editor::DataStorage::RowHandle UEditorDataStorageCompatibility::AddCompatibl
 	return bCanAddObject ? AddCompatibleObjectExplicitTransactionable<true>(Object) : UE::Editor::DataStorage::InvalidRowHandle;
 }
 
-UE::Editor::DataStorage::RowHandle UEditorDataStorageCompatibility::AddCompatibleObjectExplicit(void* Object, TWeakObjectPtr<UScriptStruct> TypeInfo)
+
+UE::Editor::DataStorage::RowHandle UEditorDataStorageCompatibility::AddCompatibleObjectExplicit(void* Object, TWeakObjectPtr<const UScriptStruct> TypeInfo)
 {
 	using namespace UE::Editor::DataStorage;
 	
@@ -645,9 +646,9 @@ void UEditorDataStorageCompatibility::FPendingTypeInformationUpdate::Process(UEd
 			[this](IDirectQueryContext& Context, FTypedElementScriptStructTypeInfoColumn& Type)
 			{
 				if (TOptional<TWeakObjectPtr<UObject>> NewObject = 
-					ProcessResolveTypeRecursively(reinterpret_cast<TWeakObjectPtr<UScriptStruct>&>(Type.TypeInfo)); NewObject.IsSet())
+					ProcessResolveTypeRecursively(reinterpret_cast<TWeakObjectPtr<const UScriptStruct>&>(Type.TypeInfo)); NewObject.IsSet())
 				{
-					Type.TypeInfo = Cast<UScriptStruct>(*NewObject);
+					Type.TypeInfo = Cast<const UScriptStruct>(*NewObject);
 					checkf(Type.TypeInfo.IsValid(),
 						TEXT("Type info column in data storage has been re-instanced to an object without struct type information"));
 				}
@@ -658,7 +659,7 @@ void UEditorDataStorageCompatibility::FPendingTypeInformationUpdate::Process(UEd
 			{
 				if (TOptional<TWeakObjectPtr<UObject>> NewObject = ProcessResolveTypeRecursively(Entry.TypeInfo); NewObject.IsSet())
 				{
-					Entry.TypeInfo = Cast<UScriptStruct>(*NewObject);
+					Entry.TypeInfo = Cast<const UScriptStruct>(*NewObject);
 					checkf(Entry.TypeInfo.Get(),
 						TEXT("Type info pending processing in data storage has been re-instanced to an object without struct type information"));
 				}
@@ -669,7 +670,7 @@ void UEditorDataStorageCompatibility::FPendingTypeInformationUpdate::Process(UEd
 }
 
 TOptional<TWeakObjectPtr<UObject>> UEditorDataStorageCompatibility::FPendingTypeInformationUpdate::ProcessResolveTypeRecursively(
-	const TWeakObjectPtr<UObject>& Target)
+	const TWeakObjectPtr<const UObject>& Target)
 {
 	// Thread-safety guaranteed because this is a private function that only gets called from functions that called inside a mutex.
 
diff --git a/Engine/Plugins/Experimental/EditorDataStorage/Source/TedsCore/Private/TypedElementDatabaseCompatibility.h b/Engine/Plugins/Experimental/EditorDataStorage/Source/TedsCore/Private/TypedElementDatabaseCompatibility.h
--- a/Engine/Plugins/Experimental/EditorDataStorage/Source/TedsCore/Private/TypedElementDatabaseCompatibility.h
+++ b/Engine/Plugins/Experimental/EditorDataStorage/Source/TedsCore/Private/TypedElementDatabaseCompatibility.h
@@ -46,14 +46,14 @@ public:
 
 	void RegisterRegistrationFilter(ObjectRegistrationFilter Filter) override;
 	void RegisterDealiaserCallback(ObjectToRowDealiaser Dealiaser) override;
-	void RegisterTypeTableAssociation(TObjectPtr<UStruct> TypeInfo, UE::Editor::DataStorage::TableHandle Table) override;
+	void RegisterTypeTableAssociation(TWeakObjectPtr<UStruct> TypeInfo, UE::Editor::DataStorage::TableHandle Table) override;
 	FDelegateHandle RegisterObjectAddedCallback(UE::Editor::DataStorage::ObjectAddedCallback&& OnObjectAdded);
 	void UnregisterObjectAddedCallback(FDelegateHandle Handle);
 	FDelegateHandle RegisterObjectRemovedCallback(UE::Editor::DataStorage::ObjectRemovedCallback&& OnObjectRemoved);
 	void UnregisterObjectRemovedCallback(FDelegateHandle Handle);
 	
 	UE::Editor::DataStorage::RowHandle AddCompatibleObjectExplicit(UObject* Object) override;
-	UE::Editor::DataStorage::RowHandle AddCompatibleObjectExplicit(void* Object, TWeakObjectPtr<UScriptStruct> TypeInfo) override;
+	UE::Editor::DataStorage::RowHandle AddCompatibleObjectExplicit(void* Object, TWeakObjectPtr<const UScriptStruct> TypeInfo) override;
 	
 	void RemoveCompatibleObjectExplicit(UObject* Object) override;
 	void RemoveCompatibleObjectExplicit(void* Object) override;
@@ -143,7 +143,7 @@ private:
 		void Process(UEditorDataStorageCompatibility& Compatibility);
 
 	private:
-		TOptional<TWeakObjectPtr<UObject>> ProcessResolveTypeRecursively(const TWeakObjectPtr<UObject>& Target);
+		TOptional<TWeakObjectPtr<UObject>> ProcessResolveTypeRecursively(const TWeakObjectPtr<const UObject>& Target);
 
 		struct FTypeInfoEntryKeyFuncs : TDefaultMapHashableKeyFuncs<TWeakObjectPtr<UObject>, TWeakObjectPtr<UObject>, false>
 		{
@@ -163,7 +163,7 @@ private:
 	struct ExternalObjectRegistration
 	{
 		void* Object;
-		TWeakObjectPtr<UScriptStruct> TypeInfo;
+		TWeakObjectPtr<const UScriptStruct> TypeInfo;
 	};
 	
 	template<typename AddressType>
@@ -214,7 +214,7 @@ private:
 	 */
 	struct FSyncTagInfo
 	{
-		TWeakObjectPtr<UScriptStruct> ColumnType;
+		TWeakObjectPtr<const UScriptStruct> ColumnType;
 		bool bAddColumn;
 
 		bool operator==(const FSyncTagInfo& Rhs) const = default;
diff --git a/Engine/Source/Runtime/TypedElementFramework/Public/Elements/Interfaces/TypedElementDataStorageCompatibilityInterface.h b/Engine/Source/Runtime/TypedElementFramework/Public/Elements/Interfaces/TypedElementDataStorageCompatibilityInterface.h
--- a/Engine/Source/Runtime/TypedElementFramework/Public/Elements/Interfaces/TypedElementDataStorageCompatibilityInterface.h
+++ b/Engine/Source/Runtime/TypedElementFramework/Public/Elements/Interfaces/TypedElementDataStorageCompatibilityInterface.h
@@ -69,7 +69,7 @@ public:
 	 * will be used to find the closest match in the registered types and use the associated table. E.g. actors derive from uobjects so
 	 * if the type information of an actor is registered the actor table will be used instead of the uobject table.
 	 */
-	virtual void RegisterTypeTableAssociation(TObjectPtr<UStruct> TypeInfo, TableHandle Table) = 0;
+	virtual void RegisterTypeTableAssociation(TWeakObjectPtr<UStruct> TypeInfo, TableHandle Table) = 0;
 
 	/**
 	 * @section Explicit functions
@@ -79,7 +79,7 @@ public:
 	/** Adds a UObject to the data storage. */
 	virtual RowHandle AddCompatibleObjectExplicit(UObject* Object) = 0;
 	/** Adds an FStruct to the data storage. */
-	virtual RowHandle AddCompatibleObjectExplicit(void* Object, TWeakObjectPtr<UScriptStruct> TypeInfo) = 0;
+	virtual RowHandle AddCompatibleObjectExplicit(void* Object, TWeakObjectPtr<const UScriptStruct> TypeInfo) = 0;
 	
 	/** Removes a UObject from the data storage. */
 	virtual void RemoveCompatibleObjectExplicit(UObject* Object) = 0;
diff --git a/Engine/Source/Runtime/TypedElementFramework/Public/Elements/Interfaces/TypedElementDataStorageInterface.h b/Engine/Source/Runtime/TypedElementFramework/Public/Elements/Interfaces/TypedElementDataStorageInterface.h
--- a/Engine/Source/Runtime/TypedElementFramework/Public/Elements/Interfaces/TypedElementDataStorageInterface.h
+++ b/Engine/Source/Runtime/TypedElementFramework/Public/Elements/Interfaces/TypedElementDataStorageInterface.h
@@ -155,7 +155,7 @@ public:
 	 * This can be used to for instance remove all rows with a specific type tag, such as removing all rows with
 	 * an entity tag to remove all entities from the data storage.
 	 */
-	virtual void RemoveAllRowsWithColumns(TConstArrayView<UScriptStruct*> Columns) = 0;
+	virtual void RemoveAllRowsWithColumns(TConstArrayView<const UScriptStruct*> Columns) = 0;
 
 	/**
 	 * Removes all rows that have at least the provided columns as template arguments.
