TEDS Compat: Dynamically register/unregister actors when their outer changes. We only want actors with a valid level in TEDS - but that information can change if the outer of the actor is modified.

diff --git a/Engine/Plugins/Experimental/EditorDataStorage/Source/TedsCore/Private/TypedElementDatabaseCompatibility.cpp b/Engine/Plugins/Experimental/EditorDataStorage/Source/TedsCore/Private/TypedElementDatabaseCompatibility.cpp
--- a/Engine/Plugins/Experimental/EditorDataStorage/Source/TedsCore/Private/TypedElementDatabaseCompatibility.cpp
+++ b/Engine/Plugins/Experimental/EditorDataStorage/Source/TedsCore/Private/TypedElementDatabaseCompatibility.cpp
@@ -88,6 +88,8 @@ void UEditorDataStorageCompatibility::Initialize(UEditorDataStorage* InStorage)
 	// Used to get all the worlds and register the actor create/destroy handles on them.
 	PostWorldInitializationDelegateHandle = FWorldDelegates::OnPostWorldInitialization.AddUObject(this, &UEditorDataStorageCompatibility::OnPostWorldInitialization);
 	PreWorldFinishDestroyDelegateHandle = FWorldDelegates::OnPreWorldFinishDestroy.AddUObject(this, &UEditorDataStorageCompatibility::OnPreWorldFinishDestroy);
+
+	ActorOuterChangedDelegateHandle = GEngine->OnLevelActorOuterChanged().AddUObject(this, &UEditorDataStorageCompatibility::OnActorOuterChanged);
 }
 
 void UEditorDataStorageCompatibility::Deinitialize()
@@ -101,6 +103,8 @@ void UEditorDataStorageCompatibility::Deinitialize()
 		It.Key->RemoveOnActorDestroyedHandler(It.Value);
 	}
 
+	GEngine->OnLevelActorOuterChanged().RemoveAll(this);
+
 	FWorldDelegates::OnPreWorldFinishDestroy.Remove(PreWorldFinishDestroyDelegateHandle);
 	FWorldDelegates::OnPostWorldInitialization.Remove(PostWorldInitializationDelegateHandle);
 	FCoreUObjectDelegates::PostReachabilityAnalysis.Remove(PostGcUnreachableAnalysisHandle);
@@ -1190,6 +1194,19 @@ void UEditorDataStorageCompatibility::OnActorDestroyed(AActor* Actor)
 	RemoveCompatibleObjectExplicit(Actor);
 }
 
+void UEditorDataStorageCompatibility::OnActorOuterChanged(AActor* Actor, UObject* Outer)
+{
+	// We only want to register actors outer'd to a level in TEDS - so if the outer changes we add/remove the object based on that status
+	if (Actor->GetLevel())
+	{
+		AddCompatibleObjectExplicit(Actor);
+	}
+	else
+	{
+		RemoveCompatibleObjectExplicit(Actor);
+	}
+}
+
 SIZE_T GetTypeHash(const UEditorDataStorageCompatibility::FSyncTagInfo& Column)
 {
 	// Thread safe as it only uses local data.
diff --git a/Engine/Plugins/Experimental/EditorDataStorage/Source/TedsCore/Private/TypedElementDatabaseCompatibility.h b/Engine/Plugins/Experimental/EditorDataStorage/Source/TedsCore/Private/TypedElementDatabaseCompatibility.h
--- a/Engine/Plugins/Experimental/EditorDataStorage/Source/TedsCore/Private/TypedElementDatabaseCompatibility.h
+++ b/Engine/Plugins/Experimental/EditorDataStorage/Source/TedsCore/Private/TypedElementDatabaseCompatibility.h
@@ -133,6 +133,7 @@ private:
 	void OnPreWorldFinishDestroy(UWorld* World);
 
 	void OnActorDestroyed(AActor* Actor);
+	void OnActorOuterChanged(AActor* Actor, UObject* Outer);
 
 	struct FPendingTypeInformationUpdate
 	{
@@ -235,6 +236,7 @@ private:
 	FDelegateHandle PreWorldFinishDestroyDelegateHandle;
 	FDelegateHandle ObjectReinstancedDelegateHandle;
 	FDelegateHandle PostGcUnreachableAnalysisHandle;
+	FDelegateHandle ActorOuterChangedDelegateHandle;
 	
 	TSharedPtr<UE::Editor::DataStorage::FEnvironment> Environment;
 	UE::Editor::DataStorage::QueryHandle ClassTypeInfoQuery;
