[TEDS Compatibility API] Refactoring TEDS compatibility API UScriptStructs to use const and const_cast

diff --git a/Engine/Plugins/Experimental/Dataflow/Source/DataflowEditor/Private/Dataflow/DataflowEditorPreviewSceneBase.cpp b/Engine/Plugins/Experimental/Dataflow/Source/DataflowEditor/Private/Dataflow/DataflowEditorPreviewSceneBase.cpp
--- a/Engine/Plugins/Experimental/Dataflow/Source/DataflowEditor/Private/Dataflow/DataflowEditorPreviewSceneBase.cpp
+++ b/Engine/Plugins/Experimental/Dataflow/Source/DataflowEditor/Private/Dataflow/DataflowEditorPreviewSceneBase.cpp
@@ -207,7 +207,7 @@ void FDataflowPreviewSceneBase::AddSceneObject(UObject* SceneObject, const bool
 	}
 }
 
-void FDataflowPreviewSceneBase::AddSceneStruct(void* SceneStruct, const TWeakObjectPtr<UScriptStruct> TypeInfo,  bool bIsConstruction) const
+void FDataflowPreviewSceneBase::AddSceneStruct(void* SceneStruct, const TWeakObjectPtr<const UScriptStruct> TypeInfo,  bool bIsConstruction) const
 {
 	using namespace UE::Editor::DataStorage;
 	if (ICompatibilityProvider* Compatibility = GetMutableDataStorageFeature<ICompatibilityProvider>(CompatibilityFeatureName))
diff --git a/Engine/Plugins/Experimental/Dataflow/Source/DataflowEditor/Public/Dataflow/DataflowEditorPreviewSceneBase.h b/Engine/Plugins/Experimental/Dataflow/Source/DataflowEditor/Public/Dataflow/DataflowEditorPreviewSceneBase.h
--- a/Engine/Plugins/Experimental/Dataflow/Source/DataflowEditor/Public/Dataflow/DataflowEditorPreviewSceneBase.h
+++ b/Engine/Plugins/Experimental/Dataflow/Source/DataflowEditor/Public/Dataflow/DataflowEditorPreviewSceneBase.h
@@ -157,7 +157,7 @@ protected:
 	void AddSceneObject(UObject* SceneObject, const bool bIsConstruction) const;
 
 	/** Store Scene struct into the TEDS database */
-	void AddSceneStruct(void* SceneStruct, TWeakObjectPtr<UScriptStruct> TypeInfo, const bool bIsConstruction) const;
+	void AddSceneStruct(void* SceneStruct, const TWeakObjectPtr<const UScriptStruct> TypeInfo, const bool bIsConstruction) const;
 
 	/** Remove Scene object from the TEDS database */
 	void RemoveSceneObject(UObject* SceneObject) const;
diff --git a/Engine/Plugins/Experimental/EditorDataStorage/Source/TedsCore/Private/Commands/EditorDataStorageCompatibilityCommands.cpp b/Engine/Plugins/Experimental/EditorDataStorage/Source/TedsCore/Private/Commands/EditorDataStorageCompatibilityCommands.cpp
--- a/Engine/Plugins/Experimental/EditorDataStorage/Source/TedsCore/Private/Commands/EditorDataStorageCompatibilityCommands.cpp
+++ b/Engine/Plugins/Experimental/EditorDataStorage/Source/TedsCore/Private/Commands/EditorDataStorageCompatibilityCommands.cpp
@@ -15,25 +15,25 @@
 
 namespace UE::Editor::DataStorage
 {
-	UScriptStruct* FAddSyncFromWorldTag::GetType()
+	const UScriptStruct* FAddSyncFromWorldTag::GetType()
 	{
 		return *GetTypeAddress();
 	}
 
-	UScriptStruct** FAddSyncFromWorldTag::GetTypeAddress()
+	const UScriptStruct** FAddSyncFromWorldTag::GetTypeAddress()
 	{
-		static UScriptStruct* Type = FTypedElementSyncFromWorldTag::StaticStruct();
+		static const UScriptStruct* Type = FTypedElementSyncFromWorldTag::StaticStruct();
 		return &Type;
 	}
 
-	UScriptStruct* FAddInteractiveSyncFromWorldTag::GetType()
+	const UScriptStruct* FAddInteractiveSyncFromWorldTag::GetType()
 	{
 		return *GetTypeAddress();
 	}
 
-	UScriptStruct** FAddInteractiveSyncFromWorldTag::GetTypeAddress()
+	const UScriptStruct** FAddInteractiveSyncFromWorldTag::GetTypeAddress()
 	{
-		static UScriptStruct* Type = FTypedElementSyncFromWorldInteractiveTag::StaticStruct();
+		static const UScriptStruct* Type = FTypedElementSyncFromWorldInteractiveTag::StaticStruct();
 		return &Type;
 	}
 
@@ -200,16 +200,16 @@ namespace UE::Editor::DataStorage
 		}
 
 		// Patch locally cached type information
-		UScriptStruct** LocalTypeInfoStorage[] =
+		const UScriptStruct** LocalTypeInfoStorage[] =
 		{ 
 			FAddSyncFromWorldTag::GetTypeAddress(),
 			FAddInteractiveSyncFromWorldTag::GetTypeAddress()
 		};
-		TArrayView<UScriptStruct**> LocalTypeInfo(LocalTypeInfoStorage, sizeof(LocalTypeInfoStorage) / sizeof(UScriptStruct**));
-		for (UScriptStruct** TypeInfo : LocalTypeInfo)
+		TArrayView<const UScriptStruct**> LocalTypeInfo(LocalTypeInfoStorage, sizeof(LocalTypeInfoStorage) / sizeof(UScriptStruct**));
+		for (const UScriptStruct** TypeInfo : LocalTypeInfo)
 		{
-			*TypeInfo = Cast<UScriptStruct>(
-				FTypeBatchInfoReinstanced::FindObjectRecursively(ReinstanceArray, TWeakObjectPtr<UStruct>(*TypeInfo)));
+			*TypeInfo = Cast<const UScriptStruct>(
+				FTypeBatchInfoReinstanced::FindObjectRecursively(ReinstanceArray, TWeakObjectPtr<const UStruct>(*TypeInfo)));
 		}
 
 		// Patch existing commands with the new type information.
diff --git a/Engine/Plugins/Experimental/EditorDataStorage/Source/TedsCore/Private/Commands/EditorDataStorageCompatibilityCommands.h b/Engine/Plugins/Experimental/EditorDataStorage/Source/TedsCore/Private/Commands/EditorDataStorageCompatibilityCommands.h
--- a/Engine/Plugins/Experimental/EditorDataStorage/Source/TedsCore/Private/Commands/EditorDataStorageCompatibilityCommands.h
+++ b/Engine/Plugins/Experimental/EditorDataStorage/Source/TedsCore/Private/Commands/EditorDataStorageCompatibilityCommands.h
@@ -156,8 +156,8 @@ namespace UE::Editor::DataStorage
 		TObjectKey<const UObject> Target;
 		RowHandle Row;
 
-		static UScriptStruct* GetType();
-		static UScriptStruct** GetTypeAddress();
+		static const UScriptStruct* GetType();
+		static const UScriptStruct** GetTypeAddress();
 	};
 
 	struct FAddInteractiveSyncFromWorldTag final
@@ -165,8 +165,8 @@ namespace UE::Editor::DataStorage
 		TObjectKey<const UObject> Target;
 		RowHandle Row;
 
-		static UScriptStruct* GetType();
-		static UScriptStruct** GetTypeAddress();
+		static const UScriptStruct* GetType();
+		static const UScriptStruct** GetTypeAddress();
 	};
 
 	struct FRemoveInteractiveSyncFromWorldTag final
diff --git a/Engine/Plugins/Experimental/EditorDataStorage/Source/TedsCore/Private/Debug/TypedElementDatabaseDebugCommands.cpp b/Engine/Plugins/Experimental/EditorDataStorage/Source/TedsCore/Private/Debug/TypedElementDatabaseDebugCommands.cpp
--- a/Engine/Plugins/Experimental/EditorDataStorage/Source/TedsCore/Private/Debug/TypedElementDatabaseDebugCommands.cpp
+++ b/Engine/Plugins/Experimental/EditorDataStorage/Source/TedsCore/Private/Debug/TypedElementDatabaseDebugCommands.cpp
@@ -21,6 +21,8 @@ namespace UE::Editor::DataStorage::Private
 	static const UStruct* GetTypeInfo(const TWeakObjectPtr<const UScriptStruct>& TypeInfo) { return TypeInfo.Get(); }
 	static const UStruct* GetTypeInfo(const TObjectPtr<const UClass>& TypeInfo) { return TypeInfo; }
 	static const UStruct* GetTypeInfo(const TObjectPtr<const UScriptStruct>& TypeInfo) { return TypeInfo; }
+	static const UStruct* GetTypeInfo(const UClass* TypeInfo) { return TypeInfo; }
+	static const UStruct* GetTypeInfo(const UScriptStruct* TypeInfo) { return TypeInfo; }
 
 	template<typename TypeInfoType>
 	void PrintObjectTypeInformation(ICoreProvider* DataStorage, FString Message, FOutputDevice& Output)
diff --git a/Engine/Plugins/Experimental/EditorDataStorage/Source/TedsCore/Private/Queries/TypedElementExtendedQueryStore.h b/Engine/Plugins/Experimental/EditorDataStorage/Source/TedsCore/Private/Queries/TypedElementExtendedQueryStore.h
--- a/Engine/Plugins/Experimental/EditorDataStorage/Source/TedsCore/Private/Queries/TypedElementExtendedQueryStore.h
+++ b/Engine/Plugins/Experimental/EditorDataStorage/Source/TedsCore/Private/Queries/TypedElementExtendedQueryStore.h
@@ -6,7 +6,6 @@
 #include "MassEntityQuery.h"
 #include "MassRequirements.h"
 #include "TypedElementHandleStore.h"
-#include "UObject/StrongObjectPtr.h"
 
 struct FMassEntityManager;
 struct FMassProcessingPhaseManager;
diff --git a/Engine/Source/Runtime/TypedElementFramework/Public/Elements/Interfaces/TypedElementDataStorageInterface.h b/Engine/Source/Runtime/TypedElementFramework/Public/Elements/Interfaces/TypedElementDataStorageInterface.h
--- a/Engine/Source/Runtime/TypedElementFramework/Public/Elements/Interfaces/TypedElementDataStorageInterface.h
+++ b/Engine/Source/Runtime/TypedElementFramework/Public/Elements/Interfaces/TypedElementDataStorageInterface.h
@@ -681,7 +681,7 @@ const DynamicColumnTemplate* ICoreProvider::GetColumn(RowHandle Row, const FName
 	const UScriptStruct* StructInfo = FindDynamicColumn(Description);
 	if (StructInfo)
 	{
-		return static_cast<DynamicColumnTemplate*>(GetColumnData(Row, StructInfo));
+		return static_cast<const DynamicColumnTemplate*>(GetColumnData(Row, StructInfo));
 	}
 	return nullptr;
 }
