TEDS UI: Temporary fix for a bug where the default export text widget would not work if you had multiple columns in a table viewer that were using it

diff --git a/Engine/Plugins/Experimental/EditorDataStorage/Source/TedsCore/Private/TypedElementDatabaseUI.cpp b/Engine/Plugins/Experimental/EditorDataStorage/Source/TedsCore/Private/TypedElementDatabaseUI.cpp
--- a/Engine/Plugins/Experimental/EditorDataStorage/Source/TedsCore/Private/TypedElementDatabaseUI.cpp
+++ b/Engine/Plugins/Experimental/EditorDataStorage/Source/TedsCore/Private/TypedElementDatabaseUI.cpp
@@ -770,48 +770,63 @@ bool UEditorDataStorageUi::CreateWidgetConstructors_LongestMatch(const TArray<UE
 		
 		if (Conditions.Verify(MatchedColumns, Columns))
 		{
-			// Empty conditions match against everything - so we update the matched columns list to reflect that
+			// If we find a factory with no conditions, for now we just individually match each column to it since it is probably the default widget
+			// Long term we need a better mechanism if there are multiple generic widgets for a single purpose since this just picks the first one
+			// for all the columns
 			if (Conditions.IsEmpty())
 			{
 				MatchedColumns = Columns;
-			}
-			// Remove the found columns from the requested list.
-			MatchedColumns.Sort(
-				[](const TWeakObjectPtr<const UScriptStruct>& Lhs, const TWeakObjectPtr<const UScriptStruct>& Rhs)
+				for (int32 Index = 0; Index < MatchedColumns.Num(); ++Index)
 				{
-					return Lhs->GetName() < Rhs->GetName();
-				});
+					// Remove the matched column from the list
+					Columns.RemoveAt(0, EAllowShrinking::No);
+					if (!CreateSingleWidgetConstructor(*FactoryIt, Arguments, {MatchedColumns[Index]}, Callback))
+					{
+						return false;
+					}
+				}
+				return true;
+			}
+			else
+			{
+				// Remove the found columns from the requested list.
+				MatchedColumns.Sort(
+					[](const TWeakObjectPtr<const UScriptStruct>& Lhs, const TWeakObjectPtr<const UScriptStruct>& Rhs)
+					{
+						return Lhs->GetName() < Rhs->GetName();
+					});
 			
-			MatchedColumns.SetNum(Algo::Unique(MatchedColumns), EAllowShrinking::No);
+				MatchedColumns.SetNum(Algo::Unique(MatchedColumns), EAllowShrinking::No);
 
-			// We need to keep track of the columns the user requested that ended up matching separately because MatchedColumns could contain the
-			// base template for a dynamic column in the requested columns that actually matched - and the widget constructor wants the latter
-			TArray<TWeakObjectPtr<const UScriptStruct>> RequestedColumnsThatMatched;
+				// We need to keep track of the columns the user requested that ended up matching separately because MatchedColumns could contain the
+				// base template for a dynamic column in the requested columns that actually matched - and the widget constructor wants the latter
+				TArray<TWeakObjectPtr<const UScriptStruct>> RequestedColumnsThatMatched;
 			
-			TWeakObjectPtr<const UScriptStruct>* ColumnsIt = Columns.GetData();
-			TWeakObjectPtr<const UScriptStruct>* ColumnsEnd = ColumnsIt + Columns.Num();
-			int32 ColumnIndex = 0;
-			for (const TWeakObjectPtr<const UScriptStruct>& MatchedColumn : MatchedColumns)
-			{
-				// Remove all the columns that were matched from the provided column list.
-				while (!Internal::CheckSingleColumnMatch(MatchedColumn.Get(), ColumnsIt->Get()))
+				TWeakObjectPtr<const UScriptStruct>* ColumnsIt = Columns.GetData();
+				TWeakObjectPtr<const UScriptStruct>* ColumnsEnd = ColumnsIt + Columns.Num();
+				int32 ColumnIndex = 0;
+				for (const TWeakObjectPtr<const UScriptStruct>& MatchedColumn : MatchedColumns)
 				{
-					++ColumnIndex;
-					++ColumnsIt;
-					if (ColumnsIt == ColumnsEnd)
+					// Remove all the columns that were matched from the provided column list.
+					while (!Internal::CheckSingleColumnMatch(MatchedColumn.Get(), ColumnsIt->Get()))
 					{
-						ensureMsgf(false, TEXT("A previously found matching column can't be found in the original array."));
-						return false;
+						++ColumnIndex;
+						++ColumnsIt;
+						if (ColumnsIt == ColumnsEnd)
+						{
+							ensureMsgf(false, TEXT("A previously found matching column can't be found in the original array."));
+							return false;
+						}
 					}
+					RequestedColumnsThatMatched.Add(*ColumnsIt);
+					Columns.RemoveAt(ColumnIndex, EAllowShrinking::No);
+					--ColumnsEnd;
 				}
-				RequestedColumnsThatMatched.Add(*ColumnsIt);
-				Columns.RemoveAt(ColumnIndex, EAllowShrinking::No);
-				--ColumnsEnd;
-			}
 			
-			if (!CreateSingleWidgetConstructor(*FactoryIt, Arguments, MoveTemp(RequestedColumnsThatMatched), Callback))
-			{
-				return false;
+				if (!CreateSingleWidgetConstructor(*FactoryIt, Arguments, MoveTemp(RequestedColumnsThatMatched), Callback))
+				{
+					return false;
+				}
 			}
 		}
 	}
