TEDS UI: Fix a bug where a single widget factory could not match against multiple subsets of the column list with longest match
As a side effect, if you have a generic purpose with multiple factories in longest match the first factory would match against all columns one by one, which is fine until we have a better solution for generic factories

diff --git a/Engine/Plugins/Experimental/EditorDataStorage/Source/TedsCore/Private/TypedElementDatabaseUI.cpp b/Engine/Plugins/Experimental/EditorDataStorage/Source/TedsCore/Private/TypedElementDatabaseUI.cpp
--- a/Engine/Plugins/Experimental/EditorDataStorage/Source/TedsCore/Private/TypedElementDatabaseUI.cpp
+++ b/Engine/Plugins/Experimental/EditorDataStorage/Source/TedsCore/Private/TypedElementDatabaseUI.cpp
@@ -798,7 +798,7 @@ bool UEditorDataStorageUi::CreateWidgetConstructors_LongestMatch(const TArray<UE
 	const WidgetConstructorCallback& Callback)
 {
 	TArray<TWeakObjectPtr<const UScriptStruct>> MatchedColumns;
-	for (auto FactoryIt = WidgetFactories.CreateConstIterator(); FactoryIt && !Columns.IsEmpty(); ++FactoryIt)
+	for (auto FactoryIt = WidgetFactories.CreateConstIterator(); FactoryIt && !Columns.IsEmpty();)
 	{
 		const UE::Editor::DataStorage::Queries::FConditions& Conditions = GetFactoryConditions(*FactoryIt);
 		
@@ -813,65 +813,58 @@ bool UEditorDataStorageUi::CreateWidgetConstructors_LongestMatch(const TArray<UE
 		
 		if (Conditions.Verify(MatchedColumns, Columns))
 		{
-			// If we find a factory with no conditions, for now we just individually match each column to it since it is probably the default widget
-			// Long term we need a better mechanism if there are multiple generic widgets for a single purpose since this just picks the first one
-			// for all the columns
-			if (Conditions.IsEmpty())
+			// Empty conditions match against everything - so we force a match against the first column in the list
+			if (Conditions.IsEmpty() || MatchedColumns.IsEmpty())
 			{
-				MatchedColumns = Columns;
-				for (int32 Index = 0; Index < MatchedColumns.Num(); ++Index)
-				{
-					// Remove the matched column from the list
-					Columns.RemoveAt(0, EAllowShrinking::No);
-					if (!CreateSingleWidgetConstructor(*FactoryIt, Arguments, {MatchedColumns[Index]}, Callback))
-					{
-						return false;
-					}
-				}
-				return true;
+				MatchedColumns = {Columns[0]};
 			}
-			else
-			{
-				// Remove the found columns from the requested list.
-				MatchedColumns.Sort(
-					[](const TWeakObjectPtr<const UScriptStruct>& Lhs, const TWeakObjectPtr<const UScriptStruct>& Rhs)
-					{
-						return Lhs->GetName() < Rhs->GetName();
-					});
 			
-				MatchedColumns.SetNum(Algo::Unique(MatchedColumns), EAllowShrinking::No);
+			// Remove the found columns from the requested list.
+			MatchedColumns.Sort(
+				[](const TWeakObjectPtr<const UScriptStruct>& Lhs, const TWeakObjectPtr<const UScriptStruct>& Rhs)
+				{
+					return Lhs->GetName() < Rhs->GetName();
+				});
+			
+			MatchedColumns.SetNum(Algo::Unique(MatchedColumns), EAllowShrinking::No);
 
-				// We need to keep track of the columns the user requested that ended up matching separately because MatchedColumns could contain the
-				// base template for a dynamic column in the requested columns that actually matched - and the widget constructor wants the latter
-				TArray<TWeakObjectPtr<const UScriptStruct>> RequestedColumnsThatMatched;
+			// We need to keep track of the columns the user requested that ended up matching separately because MatchedColumns could contain the
+			// base template for a dynamic column in the requested columns that actually matched - and the widget constructor wants the latter
+			TArray<TWeakObjectPtr<const UScriptStruct>> RequestedColumnsThatMatched;
 			
-				TWeakObjectPtr<const UScriptStruct>* ColumnsIt = Columns.GetData();
-				TWeakObjectPtr<const UScriptStruct>* ColumnsEnd = ColumnsIt + Columns.Num();
-				int32 ColumnIndex = 0;
-				for (const TWeakObjectPtr<const UScriptStruct>& MatchedColumn : MatchedColumns)
+			TWeakObjectPtr<const UScriptStruct>* ColumnsIt = Columns.GetData();
+			TWeakObjectPtr<const UScriptStruct>* ColumnsEnd = ColumnsIt + Columns.Num();
+			int32 ColumnIndex = 0;
+			for (const TWeakObjectPtr<const UScriptStruct>& MatchedColumn : MatchedColumns)
+			{
+				// Remove all the columns that were matched from the provided column list.
+				while (!Internal::CheckSingleColumnMatch(MatchedColumn.Get(), ColumnsIt->Get()))
 				{
-					// Remove all the columns that were matched from the provided column list.
-					while (!Internal::CheckSingleColumnMatch(MatchedColumn.Get(), ColumnsIt->Get()))
+					++ColumnIndex;
+					++ColumnsIt;
+					if (ColumnsIt == ColumnsEnd)
 					{
-						++ColumnIndex;
-						++ColumnsIt;
-						if (ColumnsIt == ColumnsEnd)
-						{
-							ensureMsgf(false, TEXT("A previously found matching column can't be found in the original array."));
-							return false;
-						}
+						ensureMsgf(false, TEXT("A previously found matching column can't be found in the original array."));
+						return false;
 					}
-					RequestedColumnsThatMatched.Add(*ColumnsIt);
-					Columns.RemoveAt(ColumnIndex, EAllowShrinking::No);
-					--ColumnsEnd;
 				}
+				RequestedColumnsThatMatched.Add(*ColumnsIt);
+				Columns.RemoveAt(ColumnIndex, EAllowShrinking::No);
+				--ColumnsEnd;
+			}
 			
-				if (!CreateSingleWidgetConstructor(*FactoryIt, Arguments, MoveTemp(RequestedColumnsThatMatched), Callback))
-				{
-					return false;
-				}
+			if (!CreateSingleWidgetConstructor(*FactoryIt, Arguments, MoveTemp(RequestedColumnsThatMatched), Callback))
+			{
+				return false;
 			}
 		}
+		else
+		{
+			// If this factory matched against a subset of the columns, we don't want to increment the iterator so we can check the same factory
+			// against the columns again to see if it matches another subset of the remaining columns. Otherwise a single factory only gets checked
+			// against the columns once even if it would match multiple subsets of the column list.
+			++FactoryIt;
+		}
 	}
 
 	return true;
