[Mass] FMassFragmentRequirements and FMassEntityQuery are no longer UStructs

They were never supposed to be. The change contains the sole engine file that used FMassEntityQuery as UPROPERTIES.

#rn FMassEntityQuery can no longer be used as UProperties. It doesn't affect qyery's functionality, they can still be member variables.

diff --git a/Engine/Plugins/Experimental/EditorDataStorage/Source/TedsCore/Private/Processors/TypedElementProcessorAdaptors.h b/Engine/Plugins/Experimental/EditorDataStorage/Source/TedsCore/Private/Processors/TypedElementProcessorAdaptors.h
--- a/Engine/Plugins/Experimental/EditorDataStorage/Source/TedsCore/Private/Processors/TypedElementProcessorAdaptors.h
+++ b/Engine/Plugins/Experimental/EditorDataStorage/Source/TedsCore/Private/Processors/TypedElementProcessorAdaptors.h
@@ -152,7 +152,6 @@ public:
 		UE::Editor::DataStorage::FEnvironment& Environment) override;
 
 private:
-	UPROPERTY(transient)
 	FMassEntityQuery NativeSubqueries[1];
 };
 
@@ -169,7 +168,6 @@ public:
 		UE::Editor::DataStorage::FEnvironment& Environment) override;
 
 private:
-	UPROPERTY(transient)
 	FMassEntityQuery NativeSubqueries[2];
 };
 
@@ -186,7 +184,6 @@ public:
 		UE::Editor::DataStorage::FEnvironment& Environment) override;
 
 private:
-	UPROPERTY(transient)
 	FMassEntityQuery NativeSubqueries[3];
 };
 
@@ -203,7 +200,6 @@ public:
 		UE::Editor::DataStorage::FEnvironment& Environment) override;
 
 private:
-	UPROPERTY(transient)
 	FMassEntityQuery NativeSubqueries[4];
 };
 
@@ -220,8 +216,7 @@ public:
 		UE::Editor::DataStorage::FEnvironment& Environment) override;
 
 private:
-	UPROPERTY(transient)
-		FMassEntityQuery NativeSubqueries[5];
+	FMassEntityQuery NativeSubqueries[5];
 };
 
 UCLASS()
@@ -237,8 +232,7 @@ public:
 		UE::Editor::DataStorage::FEnvironment& Environment) override;
 
 private:
-	UPROPERTY(transient)
-		FMassEntityQuery NativeSubqueries[6];
+	FMassEntityQuery NativeSubqueries[6];
 };
 
 UCLASS()
@@ -254,8 +248,7 @@ public:
 		UE::Editor::DataStorage::FEnvironment& Environment) override;
 
 private:
-	UPROPERTY(transient)
-		FMassEntityQuery NativeSubqueries[7];
+	FMassEntityQuery NativeSubqueries[7];
 };
 
 UCLASS()
@@ -271,8 +264,7 @@ public:
 		UE::Editor::DataStorage::FEnvironment& Environment) override;
 
 private:
-	UPROPERTY(transient)
-		FMassEntityQuery NativeSubqueries[8];
+	FMassEntityQuery NativeSubqueries[8];
 };
 
 /**
@@ -311,7 +303,6 @@ protected:
 	virtual void DebugOutputDescription(FOutputDevice& Ar, int32 Indent) const override;
 
 private:
-	UPROPERTY(transient)
 	FTypedElementQueryProcessorData Data;
 };
 
@@ -342,7 +333,6 @@ public:
 		UE::Editor::DataStorage::FEnvironment& Environment) override;
 
 private:
-	UPROPERTY(transient)
 	FMassEntityQuery NativeSubqueries[1];
 };
 
@@ -359,7 +349,6 @@ public:
 		UE::Editor::DataStorage::FEnvironment& Environment) override;
 
 private:
-	UPROPERTY(transient)
 	FMassEntityQuery NativeSubqueries[2];
 };
 
@@ -376,7 +365,6 @@ public:
 		UE::Editor::DataStorage::FEnvironment& Environment) override;
 
 private:
-	UPROPERTY(transient)
 	FMassEntityQuery NativeSubqueries[3];
 };
 
@@ -393,7 +381,6 @@ public:
 		UE::Editor::DataStorage::FEnvironment& Environment) override;
 
 private:
-	UPROPERTY(transient)
 	FMassEntityQuery NativeSubqueries[4];
 };
 
@@ -410,7 +397,6 @@ public:
 		UE::Editor::DataStorage::FEnvironment& Environment) override;
 
 private:
-	UPROPERTY(transient)
 	FMassEntityQuery NativeSubqueries[5];
 };
 
@@ -427,7 +413,6 @@ public:
 		UE::Editor::DataStorage::FEnvironment& Environment) override;
 
 private:
-	UPROPERTY(transient)
 	FMassEntityQuery NativeSubqueries[6];
 };
 
@@ -444,7 +429,6 @@ public:
 		UE::Editor::DataStorage::FEnvironment& Environment) override;
 
 private:
-	UPROPERTY(transient)
 	FMassEntityQuery NativeSubqueries[7];
 };
 
@@ -461,6 +445,5 @@ public:
 		UE::Editor::DataStorage::FEnvironment& Environment) override;
 
 private:
-	UPROPERTY(transient)
 	FMassEntityQuery NativeSubqueries[8];
 };
diff --git a/Engine/Source/Runtime/MassEntity/Private/MassEntityQuery.cpp b/Engine/Source/Runtime/MassEntity/Private/MassEntityQuery.cpp
--- a/Engine/Source/Runtime/MassEntity/Private/MassEntityQuery.cpp
+++ b/Engine/Source/Runtime/MassEntity/Private/MassEntityQuery.cpp
@@ -12,8 +12,6 @@
 #include "MassProcessor.h"
 #include "MassEntityTrace.h"
 
-#include UE_INLINE_GENERATED_CPP_BY_NAME(MassEntityQuery)
-
 #if WITH_MASSENTITY_DEBUG
 #include "MassRequirementAccessDetector.h"
 #include "MassDebugLogging.h"
diff --git a/Engine/Source/Runtime/MassEntity/Public/MassEntityQuery.h b/Engine/Source/Runtime/MassEntity/Public/MassEntityQuery.h
--- a/Engine/Source/Runtime/MassEntity/Public/MassEntityQuery.h
+++ b/Engine/Source/Runtime/MassEntity/Public/MassEntityQuery.h
@@ -14,7 +14,6 @@
 #include "MassExternalSubsystemTraits.h"
 #endif // UE_ENABLE_INCLUDE_ORDER_DEPRECATED_IN_5_6
 #include "MassRequirements.h"
-#include "MassEntityQuery.generated.h"
 
 #define UE_API MASSENTITY_API
 
@@ -48,11 +47,8 @@ struct FMassEntityManager;
  *  A query to be considered valid needs declared at least one EMassFragmentPresence::All, EMassFragmentPresence::Any 
  *  EMassFragmentPresence::Optional fragment requirement.
  */
-USTRUCT()
 struct FMassEntityQuery : public FMassFragmentRequirements, public FMassSubsystemRequirements
 {
-	GENERATED_BODY()
-	
 	static constexpr int32 ArchetypeMatchOverrideSize = MASS_ACHETYPE_OVERRIDE_MAX_SIZE;
 	static constexpr uint32 ArchetypeMatchOverrideAlignment = MASS_ACHETYPE_OVERRIDE_MAX_ALIGNMENT;
 	
diff --git a/Engine/Source/Runtime/MassEntity/Public/MassRequirements.h b/Engine/Source/Runtime/MassEntity/Public/MassRequirements.h
--- a/Engine/Source/Runtime/MassEntity/Public/MassRequirements.h
+++ b/Engine/Source/Runtime/MassEntity/Public/MassRequirements.h
@@ -156,15 +156,11 @@ private:
 /** 
  *  FMassFragmentRequirements is a structure that describes properties required of an archetype that's a subject of calculations.
  */
-USTRUCT()
 struct FMassFragmentRequirements
 {
-	GENERATED_BODY()
-
 	friend FMassDebugger;
 	friend FMassRequirementAccessDetector;
 
-public:
 	FMassFragmentRequirements() = default;
 	MASSENTITY_API explicit FMassFragmentRequirements(const TSharedPtr<FMassEntityManager>& EntityManager);
 	MASSENTITY_API explicit FMassFragmentRequirements(const TSharedRef<FMassEntityManager>& EntityManager);
