Fix infinite loop in TEDs UI

diff --git a/Engine/Plugins/Experimental/EditorDataStorage/Source/TedsCore/Private/TypedElementDatabaseUI.cpp b/Engine/Plugins/Experimental/EditorDataStorage/Source/TedsCore/Private/TypedElementDatabaseUI.cpp
--- a/Engine/Plugins/Experimental/EditorDataStorage/Source/TedsCore/Private/TypedElementDatabaseUI.cpp
+++ b/Engine/Plugins/Experimental/EditorDataStorage/Source/TedsCore/Private/TypedElementDatabaseUI.cpp
@@ -806,6 +806,7 @@ bool UEditorDataStorageUi::CreateWidgetConstructors_LongestMatch(const TArray<UE
 		{
 			// There are more columns required for this factory than there are in the requested columns list so skip this
 			// factory.
+			++FactoryIt;
 			continue;
 		}
 
