TEDS: Add a PostInitialize function for the three main implementations
TEDS UI: Resolve widget purpose parenting information in PostInitialize

diff --git a/Engine/Plugins/Experimental/EditorDataStorage/Source/TedsCore/Private/TypedElementDatabase.cpp b/Engine/Plugins/Experimental/EditorDataStorage/Source/TedsCore/Private/TypedElementDatabase.cpp
--- a/Engine/Plugins/Experimental/EditorDataStorage/Source/TedsCore/Private/TypedElementDatabase.cpp
+++ b/Engine/Plugins/Experimental/EditorDataStorage/Source/TedsCore/Private/TypedElementDatabase.cpp
@@ -167,6 +167,10 @@ void UEditorDataStorage::Initialize()
 	}
 }
 
+void UEditorDataStorage::PostInitialize()
+{
+}
+
 void UEditorDataStorage::SetFactories(TConstArrayView<UClass*> FactoryClasses)
 {
 	Factories.Reserve(FactoryClasses.Num());
diff --git a/Engine/Plugins/Experimental/EditorDataStorage/Source/TedsCore/Private/TypedElementDatabase.h b/Engine/Plugins/Experimental/EditorDataStorage/Source/TedsCore/Private/TypedElementDatabase.h
--- a/Engine/Plugins/Experimental/EditorDataStorage/Source/TedsCore/Private/TypedElementDatabase.h
+++ b/Engine/Plugins/Experimental/EditorDataStorage/Source/TedsCore/Private/TypedElementDatabase.h
@@ -58,6 +58,7 @@ public:
 	~UEditorDataStorage() override = default;
 	
 	void Initialize();
+	void PostInitialize();
 	
 	void SetFactories(TConstArrayView<UClass*> InFactories);
 	void ResetFactories();
diff --git a/Engine/Plugins/Experimental/EditorDataStorage/Source/TedsCore/Private/TypedElementDatabaseCompatibility.cpp b/Engine/Plugins/Experimental/EditorDataStorage/Source/TedsCore/Private/TypedElementDatabaseCompatibility.cpp
--- a/Engine/Plugins/Experimental/EditorDataStorage/Source/TedsCore/Private/TypedElementDatabaseCompatibility.cpp
+++ b/Engine/Plugins/Experimental/EditorDataStorage/Source/TedsCore/Private/TypedElementDatabaseCompatibility.cpp
@@ -92,6 +92,10 @@ void UEditorDataStorageCompatibility::Initialize(UEditorDataStorage* InStorage)
 	ActorOuterChangedDelegateHandle = GEngine->OnLevelActorOuterChanged().AddUObject(this, &UEditorDataStorageCompatibility::OnActorOuterChanged);
 }
 
+void UEditorDataStorageCompatibility::PostInitialize(UEditorDataStorage* InStorage)
+{
+}
+
 void UEditorDataStorageCompatibility::Deinitialize()
 {
 	using namespace UE::Editor::DataStorage;
diff --git a/Engine/Plugins/Experimental/EditorDataStorage/Source/TedsCore/Private/TypedElementDatabaseCompatibility.h b/Engine/Plugins/Experimental/EditorDataStorage/Source/TedsCore/Private/TypedElementDatabaseCompatibility.h
--- a/Engine/Plugins/Experimental/EditorDataStorage/Source/TedsCore/Private/TypedElementDatabaseCompatibility.h
+++ b/Engine/Plugins/Experimental/EditorDataStorage/Source/TedsCore/Private/TypedElementDatabaseCompatibility.h
@@ -42,6 +42,7 @@ public:
 	~UEditorDataStorageCompatibility() override = default;
 
 	void Initialize(UEditorDataStorage* InStorage);
+	void PostInitialize(UEditorDataStorage* InStorage);
 	void Deinitialize();
 
 	void RegisterRegistrationFilter(ObjectRegistrationFilter Filter) override;
diff --git a/Engine/Plugins/Experimental/EditorDataStorage/Source/TedsCore/Private/TypedElementDatabaseUI.cpp b/Engine/Plugins/Experimental/EditorDataStorage/Source/TedsCore/Private/TypedElementDatabaseUI.cpp
--- a/Engine/Plugins/Experimental/EditorDataStorage/Source/TedsCore/Private/TypedElementDatabaseUI.cpp
+++ b/Engine/Plugins/Experimental/EditorDataStorage/Source/TedsCore/Private/TypedElementDatabaseUI.cpp
@@ -10,6 +10,7 @@
 #include "Elements/Columns/TypedElementMiscColumns.h"
 #include "Elements/Columns/TypedElementSlateWidgetColumns.h"
 #include "Elements/Framework/TypedElementColumnUtils.h"
+#include "Elements/Framework/TypedElementRowHandleArray.h"
 #include "Elements/Interfaces/TypedElementDataStorageInterface.h"
 #include "Elements/Interfaces/TypedElementDataStorageCompatibilityInterface.h"
 #include "GenericPlatform/GenericPlatformMemory.h"
@@ -68,6 +69,48 @@ void UEditorDataStorageUi::Initialize(
 	RegisterQueries();
 }
 
+void UEditorDataStorageUi::PostInitialize(UE::Editor::DataStorage::ICoreProvider* StorageInterface,
+	UE::Editor::DataStorage::ICompatibilityProvider* StorageCompatibilityInterface)
+{
+	// Now that all the default widget purposes have been fully initialized, we can try to resolve the parenting information manually
+	// to ensure it is available for any tools that restored from the layout, which happens before TEDS has a chance to tick and resolve the
+	// information through a processor
+	using namespace UE::Editor::DataStorage::Queries;
+
+	FQueryDescription PurposesQueryDescription =
+		Select()
+		.Where(TColumn<FWidgetPurposeColumn>() && TColumn<FUnresolvedTableRowParentColumn>())
+		.Compile();
+	
+	QueryHandle PurposesQuery = StorageInterface->RegisterQuery(MoveTemp(PurposesQueryDescription));
+
+	FRowHandleArray Rows;
+	StorageInterface->RunQuery(PurposesQuery, CreateDirectQueryCallbackBinding(
+		[&Rows](IDirectQueryContext& Context, RowHandle Row)
+		{
+			Rows.Add(Row);
+		}));
+
+	for (RowHandle Row : Rows.GetRows())
+	{
+		FUnresolvedTableRowParentColumn* UnresolvedParent = Storage->GetColumn<FUnresolvedTableRowParentColumn>(Row);
+
+		if (ensureMsgf(UnresolvedParent, TEXT("Expected valid parent for row in query with TColumn<FUnresolvedTableRowParentColumn>()")))
+		{
+			RowHandle ParentRow = StorageInterface->LookupMappedRow(UnresolvedParent->ParentIdKey);
+			
+			if (Storage->IsRowAvailable(ParentRow))
+			{
+				StorageInterface->RemoveColumns<FUnresolvedTableRowParentColumn>(Row);
+				StorageInterface->AddColumn(Row, FTableRowParentColumn{ .Parent = ParentRow });
+			}
+		}
+	}
+
+	// Post Init, there is a processer run to resolve parent rows automatically so we don't need this query anymore
+	StorageInterface->UnregisterQuery(PurposesQuery);
+}
+
 void UEditorDataStorageUi::Deinitialize()
 {
 }
diff --git a/Engine/Plugins/Experimental/EditorDataStorage/Source/TedsCore/Private/TypedElementDatabaseUI.h b/Engine/Plugins/Experimental/EditorDataStorage/Source/TedsCore/Private/TypedElementDatabaseUI.h
--- a/Engine/Plugins/Experimental/EditorDataStorage/Source/TedsCore/Private/TypedElementDatabaseUI.h
+++ b/Engine/Plugins/Experimental/EditorDataStorage/Source/TedsCore/Private/TypedElementDatabaseUI.h
@@ -33,6 +33,9 @@ public:
 	void Initialize(
 		UE::Editor::DataStorage::ICoreProvider* StorageInterface, 
 		UE::Editor::DataStorage::ICompatibilityProvider* StorageCompatibilityInterface);
+	void PostInitialize(
+		UE::Editor::DataStorage::ICoreProvider* StorageInterface, 
+		UE::Editor::DataStorage::ICompatibilityProvider* StorageCompatibilityInterface);
 	void Deinitialize();
 
 	virtual UE::Editor::DataStorage::RowHandle RegisterWidgetPurpose(const FPurposeID& PurposeID, const FPurposeInfo& InPurposeInfo) override;
diff --git a/Engine/Plugins/Experimental/EditorDataStorage/Source/TedsCore/Private/TypedElementsDataStorage.cpp b/Engine/Plugins/Experimental/EditorDataStorage/Source/TedsCore/Private/TypedElementsDataStorage.cpp
--- a/Engine/Plugins/Experimental/EditorDataStorage/Source/TedsCore/Private/TypedElementsDataStorage.cpp
+++ b/Engine/Plugins/Experimental/EditorDataStorage/Source/TedsCore/Private/TypedElementsDataStorage.cpp
@@ -140,6 +140,10 @@ void FEditorDataStorageModule::StartupModule()
 				UE_LOG(LogEditorDataStorage, Log, TEXT("Initialized"));
 
 				bInitialized = true;
+				
+				DataStorage->PostInitialize();
+				DataStorageCompatibility->PostInitialize(DataStorage);
+				DataStorageUi->PostInitialize(DataStorage, DataStorageCompatibility);
 			}
 		});
 	FCoreDelegates::OnExit.AddRaw(this, &FEditorDataStorageModule::ShutdownModule);
diff --git a/Engine/Plugins/Experimental/EditorDataStorageFeatures/Source/TedsAlerts/Private/TedsAlertWidget.cpp b/Engine/Plugins/Experimental/EditorDataStorageFeatures/Source/TedsAlerts/Private/TedsAlertWidget.cpp
--- a/Engine/Plugins/Experimental/EditorDataStorageFeatures/Source/TedsAlerts/Private/TedsAlertWidget.cpp
+++ b/Engine/Plugins/Experimental/EditorDataStorageFeatures/Source/TedsAlerts/Private/TedsAlertWidget.cpp
@@ -158,13 +158,6 @@ void UAlertWidgetFactory::RegisterWidgetConstructors(
 	DataStorageUi.RegisterWidgetFactory<FAlertWidgetConstructor>(DataStorageUi.FindPurpose(DataStorageUi.GetGeneralWidgetPurposeID()),
 		TColumn<FAlertColumn>() || TColumn<FChildAlertColumn>());
 
-	// We are also explicitly registering the widget for the outliner purpose, because the outliner can be registered before TEDS has a chance
-	// to resolve the Outliner.Cell -> General.Cell parenting info - which means the alert widget can sometimes not get created on startup depending
-	// on init order
-	DataStorageUi.RegisterWidgetFactory<FAlertWidgetConstructor>(
-		DataStorageUi.FindPurpose(IUiProvider::FPurposeInfo("SceneOutliner", "Cell", NAME_None).GeneratePurposeID()),
-		TColumn<FAlertColumn>() || TColumn<FChildAlertColumn>());
-	
 	DataStorageUi.RegisterWidgetFactory<FAlertHeaderWidgetConstructor>(
 		DataStorageUi.FindPurpose(IUiProvider::FPurposeInfo("General", "Header", NAME_None).GeneratePurposeID()),
 		TColumn<FAlertColumn>() || TColumn<FChildAlertColumn>());
