[TEDS] Fixed actor visibility not syncing properly to IsTemporarilyHiddenInEditor property.

diff --git a/Engine/Plugins/Experimental/EditorDataStorage/Source/TedsCore/Private/Compatibility/TypedElementActorVisibilityQueries.cpp b/Engine/Plugins/Experimental/EditorDataStorage/Source/TedsCore/Private/Compatibility/TypedElementActorVisibilityQueries.cpp
--- a/Engine/Plugins/Experimental/EditorDataStorage/Source/TedsCore/Private/Compatibility/TypedElementActorVisibilityQueries.cpp
+++ b/Engine/Plugins/Experimental/EditorDataStorage/Source/TedsCore/Private/Compatibility/TypedElementActorVisibilityQueries.cpp
@@ -27,7 +27,7 @@ void UActorVisibilityDataStorageFactory::RegisterActorAddVisibilityColumn(UE::Ed
 			{
 				if (const AActor* ActorInstance = Cast<AActor>(Object.Object); ActorInstance != nullptr)
 				{
-					Context.AddColumn(Row, FVisibleInEditorColumn{ .bIsVisibleInEditor = !ActorInstance->IsHiddenEd() });
+					Context.AddColumn(Row, FVisibleInEditorColumn{ .bIsVisibleInEditor = !ActorInstance->IsTemporarilyHiddenInEditor() });
 				}
 			})
 		.Where()
@@ -49,7 +49,7 @@ void UActorVisibilityDataStorageFactory::RegisterActorVisibilityToColumnQuery(UE
 			{
 				if (const AActor* ActorInstance = Cast<AActor>(Object.Object); ActorInstance != nullptr)
 				{
-					VisibilityColumn.bIsVisibleInEditor = !ActorInstance->IsHiddenEd();
+					VisibilityColumn.bIsVisibleInEditor = !ActorInstance->IsTemporarilyHiddenInEditor();
 				}
 			})
 		.Where()
