[TEDS] Added support for domains in the mapping.

The mapping feature in TEDS now uses domains. This prevents the same indices from different domains from colliding. As an example, it's possible to use a folder path as a index key for a row pointing to Outliner data and Content Browser data. The same input would need to produce different rows, but currently it would collide. With domains it's now possible to create two separate domains and allow the same key to be used.

Other updates included are:
- Removed some places that incorrectly referred to the old "Index Table" name.
- Added a 5 frame delay before starting the garbage collection. It seems to be somewhat common for the editor to delete something followed by a delete the next frame. This change adds a bit of a delay to avoid repeatedly doing two garbage collection cycles in a row.

Original CRs: 42911689, 42911693, 42911695, 42911698, 42911718, 42911743, 42911747

diff --git a/Engine/Plugins/ChaosVD/Source/ChaosVD/Private/TEDS/ChaosVDParentDataStorageFactory.cpp b/Engine/Plugins/ChaosVD/Source/ChaosVD/Private/TEDS/ChaosVDParentDataStorageFactory.cpp
--- a/Engine/Plugins/ChaosVD/Source/ChaosVD/Private/TEDS/ChaosVDParentDataStorageFactory.cpp
+++ b/Engine/Plugins/ChaosVD/Source/ChaosVD/Private/TEDS/ChaosVDParentDataStorageFactory.cpp
@@ -85,7 +85,7 @@ void UChaosVDParentDataStorageFactory::RegisterAddParentColumn(UE::Editor::DataS
 							ParentColumn.ParentObject = DefaultRoot->GetTedsRowHandle();
 
 							FMapKeyView IdKey = FMapKeyView(CVDActor);
-							RowHandle ActorRow = Context.LookupMappedRow(IdKey);
+							RowHandle ActorRow = Context.LookupMappedRow(ICompatibilityProvider::ObjectMappingDomain, IdKey);
 							ParentData->ChildrenSet.Emplace(ActorRow);
 							Context.AddColumn(Row, MoveTemp(ParentColumn));
 						}
@@ -122,7 +122,7 @@ void UChaosVDParentDataStorageFactory::RegisterUpdateOrRemoveParentColumn(UE::Ed
 					else if (AChaosVDSolverInfoActor* NewParentActor = Cast<AChaosVDSolverInfoActor>(Particle->GetParentActor()))
 					{
 						FMapKeyView IdKey = FMapKeyView(NewParentActor);
-						NewParentRowHandle = Context.LookupMappedRow(IdKey);
+						NewParentRowHandle = Context.LookupMappedRow(ICompatibilityProvider::ObjectMappingDomain, IdKey);
 					}
 
 					UpdateParentData(DataStorage, Context, ParentColumn, Row, NewParentRowHandle);
@@ -146,7 +146,7 @@ void UChaosVDParentDataStorageFactory::RegisterUpdateOrRemoveParentColumn(UE::Ed
 					if (AActor* ParentActor = ActorPtr->GetAttachParentActor())
 					{
 						FMapKeyView IdKey = FMapKeyView(ParentActor);
-						NewParentRowHandle = Context.LookupMappedRow(IdKey);
+						NewParentRowHandle = Context.LookupMappedRow(ICompatibilityProvider::ObjectMappingDomain, IdKey);
 					}
 					
 					UpdateParentData(DataStorage, Context, Parent, Row, NewParentRowHandle);
diff --git a/Engine/Plugins/Developer/PerforceSourceControl/Source/PerforceSourceControl/Private/PerforceSourceControlOperations.cpp b/Engine/Plugins/Developer/PerforceSourceControl/Source/PerforceSourceControl/Private/PerforceSourceControlOperations.cpp
--- a/Engine/Plugins/Developer/PerforceSourceControl/Source/PerforceSourceControl/Private/PerforceSourceControlOperations.cpp
+++ b/Engine/Plugins/Developer/PerforceSourceControl/Source/PerforceSourceControl/Private/PerforceSourceControlOperations.cpp
@@ -2059,6 +2059,7 @@ bool FPerforceUpdateStatusWorker::UpdateStates() const
 	{
 		using namespace UE::Editor::DataStorage;
 		using namespace UE::Editor::DataStorage::Queries;
+		using namespace UE::Editor::RevisionControl;
 
 		ICoreProvider* DataStorage = GetMutableDataStorageFeature<ICoreProvider>(StorageFeatureName);
 		if (!DataStorage)
@@ -2080,12 +2081,12 @@ bool FPerforceUpdateStatusWorker::UpdateStates() const
 			Filename = FPaths::ConvertRelativePathToFull(Filename);
 			
 			FMapKey Key = FMapKey(Filename);
-			RowHandle Row = DataStorage->LookupMappedRow(Key);
+			RowHandle Row = DataStorage->LookupMappedRow(MappingDomain, Key);
 
 			if (!DataStorage->IsRowAvailable(Row))
 			{
 				Row = DataStorage->AddRow(RevisionControlTable);
-				DataStorage->MapRow(Key, Row);
+				DataStorage->MapRow(MappingDomain, Key, Row);
 			}
 			return Row;
 		};
diff --git a/Engine/Plugins/Experimental/Dataflow/Source/DataflowEditor/Private/Dataflow/DataflowOutlinerFactory.cpp b/Engine/Plugins/Experimental/Dataflow/Source/DataflowEditor/Private/Dataflow/DataflowOutlinerFactory.cpp
--- a/Engine/Plugins/Experimental/Dataflow/Source/DataflowEditor/Private/Dataflow/DataflowOutlinerFactory.cpp
+++ b/Engine/Plugins/Experimental/Dataflow/Source/DataflowEditor/Private/Dataflow/DataflowOutlinerFactory.cpp
@@ -1,4 +1,4 @@
-﻿// Copyright Epic Games, Inc. All Rights Reserved.
+// Copyright Epic Games, Inc. All Rights Reserved.
 
 #include "DataflowOutlinerFactory.h"
 #include "Dataflow/DataflowOutlinerFactory.h"
@@ -99,7 +99,7 @@ void UDataflowObjectFactory::RegisterHierarchyQueries(UE::Editor::DataStorage::I
 			if(ParentObject)
 			{
 				FMapKeyView IdKey = FMapKeyView(ParentObject);
-				RowHandle ParentHandle = Context.LookupMappedRow(IdKey);
+				RowHandle ParentHandle = Context.LookupMappedRow(ICompatibilityProvider::ObjectMappingDomain, IdKey);
 
 				ParentColumn.Parent = ParentHandle;
 			}
@@ -122,7 +122,7 @@ void UDataflowObjectFactory::RegisterHierarchyQueries(UE::Editor::DataStorage::I
 					if(SceneObject->ParentElement)
 					{
 						FMapKeyView IdKey = FMapKeyView(SceneObject->ParentElement);
-						RowHandle ParentHandle = Context.LookupMappedRow(IdKey);
+						RowHandle ParentHandle = Context.LookupMappedRow(ICompatibilityProvider::ObjectMappingDomain, IdKey);
 
 						ParentColumn.Parent = ParentHandle;
 					}
diff --git a/Engine/Plugins/Experimental/EditorDataStorage/Source/TedsCore/Private/Commands/EditorDataStorageCompatibilityCommands.cpp b/Engine/Plugins/Experimental/EditorDataStorage/Source/TedsCore/Private/Commands/EditorDataStorageCompatibilityCommands.cpp
--- a/Engine/Plugins/Experimental/EditorDataStorage/Source/TedsCore/Private/Commands/EditorDataStorageCompatibilityCommands.cpp
+++ b/Engine/Plugins/Experimental/EditorDataStorage/Source/TedsCore/Private/Commands/EditorDataStorageCompatibilityCommands.cpp
@@ -265,7 +265,7 @@ namespace UE::Editor::DataStorage
 		if (Command.ObjectRow == InvalidRowHandle)
 		{
 			FMapKeyView Key = FMapKeyView(Command.Object);
-			Command.ObjectRow = Storage.LookupMappedRow(Key);
+			Command.ObjectRow = Storage.LookupMappedRow(ICompatibilityProvider::ObjectMappingDomain, Key);
 			if (!Storage.IsRowAvailable(Command.ObjectRow))
 			{
 				Commands.ReplaceCommand<FNopCommand>(CurrentIndex);
@@ -276,7 +276,7 @@ namespace UE::Editor::DataStorage
 	void FPrepareCommands::operator()(FRemoveCompatibleExternalObject& Command)
 	{
 		FMapKeyView Key = FMapKeyView(Command.Object);
-		Command.ObjectRow = Storage.LookupMappedRow(Key);
+		Command.ObjectRow = Storage.LookupMappedRow(ICompatibilityProvider::ObjectMappingDomain, Key);
 		if (!Storage.IsRowAvailable(Command.ObjectRow))
 		{
 			Commands.ReplaceCommand<FNopCommand>(CurrentIndex);
diff --git a/Engine/Plugins/Experimental/EditorDataStorage/Source/TedsCore/Private/Compatibility/TypedElementActorParentQueries.cpp b/Engine/Plugins/Experimental/EditorDataStorage/Source/TedsCore/Private/Compatibility/TypedElementActorParentQueries.cpp
--- a/Engine/Plugins/Experimental/EditorDataStorage/Source/TedsCore/Private/Compatibility/TypedElementActorParentQueries.cpp
+++ b/Engine/Plugins/Experimental/EditorDataStorage/Source/TedsCore/Private/Compatibility/TypedElementActorParentQueries.cpp
@@ -46,7 +46,7 @@ void UActorParentDataStorageFactory::RegisterAddParentColumn(UE::Editor::DataSto
 					if (AActor* Parent = ActorInstance->GetAttachParentActor())
 					{
 						FMapKeyView IdKey = FMapKeyView(Parent);
-						RowHandle ParentRow = Context.LookupMappedRow(IdKey);
+						RowHandle ParentRow = Context.LookupMappedRow(ICompatibilityProvider::ObjectMappingDomain, IdKey);
 						if (Context.IsRowAvailable(ParentRow))
 						{
 							Context.AddColumn(Row, FTableRowParentColumn{ .Parent = ParentRow });
@@ -80,7 +80,7 @@ void UActorParentDataStorageFactory::RegisterUpdateOrRemoveParentColumn(UE::Edit
 					if (AActor* ParentActor = ActorInstance->GetAttachParentActor())
 					{
 						FMapKeyView IdKey = FMapKeyView(ParentActor);
-						RowHandle ParentRow = Context.LookupMappedRow(IdKey);
+						RowHandle ParentRow = Context.LookupMappedRow(ICompatibilityProvider::ObjectMappingDomain, IdKey);
 
 						if (Parent.Parent != ParentRow)
 						{
diff --git a/Engine/Plugins/Experimental/EditorDataStorage/Source/TedsCore/Private/Processors/TypedElementProcessorAdaptors.cpp b/Engine/Plugins/Experimental/EditorDataStorage/Source/TedsCore/Private/Processors/TypedElementProcessorAdaptors.cpp
--- a/Engine/Plugins/Experimental/EditorDataStorage/Source/TedsCore/Private/Processors/TypedElementProcessorAdaptors.cpp
+++ b/Engine/Plugins/Experimental/EditorDataStorage/Source/TedsCore/Private/Processors/TypedElementProcessorAdaptors.cpp
@@ -1194,12 +1194,12 @@ namespace UE::Editor::DataStorage
 				}
 			}
 	
-			RowHandle LookupMappedRow(const FMapKeyView& Key) const
+			RowHandle LookupMappedRow(const FName& Domain, const FMapKeyView& Key) const
 			{
 				EGlobalLockScope Scope = FGlobalLock::GetLockStatus(EGlobalLockScope::Internal) == EGlobalLockStatus::Unlocked
 					? EGlobalLockScope::Public // There's no internal lock so use a public lock instead.
 					: EGlobalLockScope::Internal; // There's an internal lock set so use that.
-				return Environment.GetMappingTable().Lookup(Scope, Key);
+				return Environment.GetMappingTable().Lookup(Scope, Domain, Key);
 			}
 
 			UE::Editor::DataStorage::FQueryResult RunQuery(QueryHandle Query)
@@ -1290,7 +1290,7 @@ namespace UE::Editor::DataStorage
 			const UObject* GetDependency(const UClass* DependencyClass) override { return Implementation.GetDependency(DependencyClass); }
 			UObject* GetMutableDependency(const UClass* DependencyClass) override { return Implementation.GetMutableDependency(DependencyClass); }
 			void GetDependencies(TArrayView<UObject*> RetrievedAddresses, TConstArrayView<TWeakObjectPtr<const UClass>> DependencyTypes, TConstArrayView<EQueryAccessType> AccessTypes) override { return Implementation.GetDependencies(RetrievedAddresses, DependencyTypes, AccessTypes); }
-			RowHandle LookupMappedRow(const FMapKeyView& Index) const override { return Implementation.LookupMappedRow(Index); }
+			RowHandle LookupMappedRow(const FName& Domain, const FMapKeyView& Index) const override { return Implementation.LookupMappedRow(Domain, Index); }
 			FQueryResult RunQuery(QueryHandle Query) override { return Implementation.RunQuery(Query); }
 			FQueryResult RunSubquery(int32 SubqueryIndex) override { return Implementation.RunSubquery(SubqueryIndex); }
 			FQueryResult RunSubquery(int32 SubqueryIndex, UE::Editor::DataStorage::SubqueryCallbackRef Callback) override { return Implementation.RunSubquery(SubqueryIndex, Callback); }
diff --git a/Engine/Plugins/Experimental/EditorDataStorage/Source/TedsCore/Private/Queries/TypedElementHierarchyQueries.cpp b/Engine/Plugins/Experimental/EditorDataStorage/Source/TedsCore/Private/Queries/TypedElementHierarchyQueries.cpp
--- a/Engine/Plugins/Experimental/EditorDataStorage/Source/TedsCore/Private/Queries/TypedElementHierarchyQueries.cpp
+++ b/Engine/Plugins/Experimental/EditorDataStorage/Source/TedsCore/Private/Queries/TypedElementHierarchyQueries.cpp
@@ -15,7 +15,7 @@ void UTypedElementHiearchyQueriesFactory::RegisterQueries(UE::Editor::DataStorag
 			FProcessor(EQueryTickPhase::PrePhysics, DataStorage.GetQueryTickGroupName(EQueryTickGroups::Default)),
 			[](IQueryContext& Context, RowHandle Row, const FUnresolvedTableRowParentColumn& UnresolvedParent)
 			{
-				RowHandle ParentRow = Context.LookupMappedRow(UnresolvedParent.ParentIdKey);
+				RowHandle ParentRow = Context.LookupMappedRow(ICompatibilityProvider::ObjectMappingDomain, UnresolvedParent.ParentIdKey);
 				if (Context.IsRowAvailable(ParentRow))
 				{
 					Context.RemoveColumns<FUnresolvedTableRowParentColumn>(Row);
diff --git a/Engine/Plugins/Experimental/EditorDataStorage/Source/TedsCore/Private/TypedElementDatabase.cpp b/Engine/Plugins/Experimental/EditorDataStorage/Source/TedsCore/Private/TypedElementDatabase.cpp
--- a/Engine/Plugins/Experimental/EditorDataStorage/Source/TedsCore/Private/TypedElementDatabase.cpp
+++ b/Engine/Plugins/Experimental/EditorDataStorage/Source/TedsCore/Private/TypedElementDatabase.cpp
@@ -1154,29 +1154,33 @@ void UEditorDataStorage::ActivateQueries(FName ActivationName)
 	}
 }
 
-UE::Editor::DataStorage::RowHandle UEditorDataStorage::LookupMappedRow(const UE::Editor::DataStorage::FMapKeyView& Key) const
+UE::Editor::DataStorage::RowHandle UEditorDataStorage::LookupMappedRow(
+	const FName& Domain, const UE::Editor::DataStorage::FMapKeyView& Key) const
 {
-	return Environment->GetMappingTable().Lookup(UE::Editor::DataStorage::EGlobalLockScope::Public, Key);
+	return Environment->GetMappingTable().Lookup(UE::Editor::DataStorage::EGlobalLockScope::Public, Domain, Key);
 }
 
-void UEditorDataStorage::MapRow(UE::Editor::DataStorage::FMapKey Key, UE::Editor::DataStorage::RowHandle Row)
+void UEditorDataStorage::MapRow(const FName& Domain, UE::Editor::DataStorage::FMapKey Key, UE::Editor::DataStorage::RowHandle Row)
 {
-	Environment->GetMappingTable().Map(UE::Editor::DataStorage::EGlobalLockScope::Public, MoveTemp(Key), Row);
+	Environment->GetMappingTable().Map(UE::Editor::DataStorage::EGlobalLockScope::Public, Domain, MoveTemp(Key), Row);
 }
 
-void UEditorDataStorage::BatchMapRows(TArrayView<TPair<UE::Editor::DataStorage::FMapKey, UE::Editor::DataStorage::RowHandle>> MapRowPairs)
+void UEditorDataStorage::BatchMapRows(
+	const FName& Domain, TArrayView<TPair<UE::Editor::DataStorage::FMapKey, UE::Editor::DataStorage::RowHandle>> MapRowPairs)
 {
-	Environment->GetMappingTable().BatchMap(UE::Editor::DataStorage::EGlobalLockScope::Public, MapRowPairs);
+	Environment->GetMappingTable().BatchMap(UE::Editor::DataStorage::EGlobalLockScope::Public, Domain, MapRowPairs);
 }
 
-void UEditorDataStorage::RemapRow(const UE::Editor::DataStorage::FMapKeyView& OriginalKey, UE::Editor::DataStorage::FMapKey NewKey)
+void UEditorDataStorage::RemapRow(
+	const FName& Domain, const UE::Editor::DataStorage::FMapKeyView& OriginalKey, UE::Editor::DataStorage::FMapKey NewKey)
 {
-	Environment->GetMappingTable().Remap(UE::Editor::DataStorage::EGlobalLockScope::Public, OriginalKey, MoveTemp(NewKey));
+	Environment->GetMappingTable().Remap(UE::Editor::DataStorage::EGlobalLockScope::Public, Domain, OriginalKey, MoveTemp(NewKey));
 }
 
-void UEditorDataStorage::RemoveRowMapping(const UE::Editor::DataStorage::FMapKeyView& Key)
+void UEditorDataStorage::RemoveRowMapping(
+	const FName& Domain, const UE::Editor::DataStorage::FMapKeyView& Key)
 {
-	Environment->GetMappingTable().Remove(UE::Editor::DataStorage::EGlobalLockScope::Public, Key);
+	Environment->GetMappingTable().Remove(UE::Editor::DataStorage::EGlobalLockScope::Public, Domain, Key);
 }
 
 UE::Editor::DataStorage::FTypedElementOnDataStorageUpdate& UEditorDataStorage::OnUpdate()
diff --git a/Engine/Plugins/Experimental/EditorDataStorage/Source/TedsCore/Private/TypedElementDatabase.h b/Engine/Plugins/Experimental/EditorDataStorage/Source/TedsCore/Private/TypedElementDatabase.h
--- a/Engine/Plugins/Experimental/EditorDataStorage/Source/TedsCore/Private/TypedElementDatabase.h
+++ b/Engine/Plugins/Experimental/EditorDataStorage/Source/TedsCore/Private/TypedElementDatabase.h
@@ -149,11 +149,16 @@ public:
 		UE::Editor::DataStorage::DirectQueryCallbackRef Callback) override;
 	virtual void ActivateQueries(FName ActivationName) override;
 
-	virtual UE::Editor::DataStorage::RowHandle LookupMappedRow(const UE::Editor::DataStorage::FMapKeyView& Key) const override;
-	virtual void MapRow(UE::Editor::DataStorage::FMapKey Key, UE::Editor::DataStorage::RowHandle Row) override;
-	virtual void BatchMapRows(TArrayView<TPair<UE::Editor::DataStorage::FMapKey, UE::Editor::DataStorage::RowHandle>> MapRowPairs) override;
-	virtual void RemapRow(const UE::Editor::DataStorage::FMapKeyView& OriginalKey, UE::Editor::DataStorage::FMapKey NewKey) override;
-	virtual void RemoveRowMapping(const UE::Editor::DataStorage::FMapKeyView& Key) override;
+	virtual UE::Editor::DataStorage::RowHandle LookupMappedRow(
+		const FName& Domain, const UE::Editor::DataStorage::FMapKeyView& Key) const override;
+	virtual void MapRow(
+		const FName& Domain, UE::Editor::DataStorage::FMapKey Key, UE::Editor::DataStorage::RowHandle Row) override;
+	virtual void BatchMapRows(
+		const FName& Domain, TArrayView<TPair<UE::Editor::DataStorage::FMapKey, UE::Editor::DataStorage::RowHandle>> MapRowPairs) override;
+	virtual void RemapRow(
+		const FName& Domain, const UE::Editor::DataStorage::FMapKeyView& OriginalKey, UE::Editor::DataStorage::FMapKey NewKey) override;
+	virtual void RemoveRowMapping(
+		const FName& Domain, const UE::Editor::DataStorage::FMapKeyView& Key) override;
 
 	virtual UE::Editor::DataStorage::FTypedElementOnDataStorageUpdate& OnUpdate() override;
 	virtual UE::Editor::DataStorage::FTypedElementOnDataStorageUpdate& OnUpdateCompleted() override;
diff --git a/Engine/Plugins/Experimental/EditorDataStorage/Source/TedsCore/Private/TypedElementDatabaseCompatibility.cpp b/Engine/Plugins/Experimental/EditorDataStorage/Source/TedsCore/Private/TypedElementDatabaseCompatibility.cpp
--- a/Engine/Plugins/Experimental/EditorDataStorage/Source/TedsCore/Private/TypedElementDatabaseCompatibility.cpp
+++ b/Engine/Plugins/Experimental/EditorDataStorage/Source/TedsCore/Private/TypedElementDatabaseCompatibility.cpp
@@ -236,7 +236,7 @@ UE::Editor::DataStorage::RowHandle UEditorDataStorageCompatibility::AddCompatibl
 	if (!Storage->IsRowAvailable(Result))
 	{
 		Result = Storage->ReserveRow();
-		Storage->MapRow(FMapKey(Object), Result);
+		Storage->MapRow(ICompatibilityProvider::ObjectMappingDomain, FMapKey(Object), Result);
 		if (Private::bUseCommandBuffer)
 		{
 			QueuedCommands.AddCommand(FAddCompatibleExternalObject{ .Object = Object, .TypeInfo = TypeInfo, .Row = Result });
@@ -268,7 +268,7 @@ void UEditorDataStorageCompatibility::RemoveCompatibleObjectExplicit(void* Objec
 	{
 		FScopedExclusiveLock Lock(EGlobalLockScope::Public);
 
-		RowHandle Row = Storage->LookupMappedRow(FMapKeyView(Object));
+		RowHandle Row = Storage->LookupMappedRow(ICompatibilityProvider::ObjectMappingDomain, FMapKeyView(Object));
 		if (Storage->IsRowAvailable(Row))
 		{
 			const FTypedElementScriptStructTypeInfoColumn* TypeInfoColumn = Storage->GetColumn<FTypedElementScriptStructTypeInfoColumn>(Row);
@@ -289,7 +289,7 @@ UE::Editor::DataStorage::RowHandle UEditorDataStorageCompatibility::FindRowWithC
 	{
 		FScopedSharedLock Lock(EGlobalLockScope::Public);
 
-		RowHandle Row = Storage->LookupMappedRow(FMapKeyView(Object));
+		RowHandle Row = Storage->LookupMappedRow(ICompatibilityProvider::ObjectMappingDomain, FMapKeyView(Object));
 		return Storage->IsRowAvailable(Row) ? Row : DealiasObject(Object);
 	}
 	return InvalidRowHandle;
@@ -301,7 +301,9 @@ UE::Editor::DataStorage::RowHandle UEditorDataStorageCompatibility::FindRowWithC
 
 	using namespace UE::Editor::DataStorage;
 
-	return (Object && Storage && Storage->IsAvailable()) ? Storage->LookupMappedRow(FMapKeyView(Object)) : InvalidRowHandle;
+	return (Object && Storage && Storage->IsAvailable()) 
+		? Storage->LookupMappedRow(ICompatibilityProvider::ObjectMappingDomain, FMapKeyView(Object)) 
+		: InvalidRowHandle;
 }
 
 bool UEditorDataStorageCompatibility::SupportsExtension(FName Extension) const
@@ -408,7 +410,7 @@ bool UEditorDataStorageCompatibility::ShouldAddObject(const UObject* Object) con
 	FScopedSharedLock Lock(EGlobalLockScope::Public);
 
 	bool Include = true;
-	if (!Storage->IsRowAvailable(Storage->LookupMappedRow(FMapKeyView(Object))))
+	if (!Storage->IsRowAvailable(Storage->LookupMappedRow(ICompatibilityProvider::ObjectMappingDomain, FMapKeyView(Object))))
 	{
 		const ObjectRegistrationFilter* Filter = ObjectRegistrationFilters.GetData();
 		const ObjectRegistrationFilter* FilterEnd = Filter + ObjectRegistrationFilters.Num();
@@ -449,7 +451,7 @@ UE::Editor::DataStorage::RowHandle UEditorDataStorageCompatibility::AddCompatibl
 	if (!Storage->IsRowAvailable(Result))
 	{
 		Result = Storage->ReserveRow();
-		Storage->MapRow(FMapKey(Object), Result);
+		Storage->MapRow(ICompatibilityProvider::ObjectMappingDomain, FMapKey(Object), Result);
 		if (Private::bUseCommandBuffer)
 		{
 			QueuedCommands.AddCommand(FAddCompatibleUObject{ .Object = Object, .Row = Result });
@@ -492,7 +494,7 @@ void UEditorDataStorageCompatibility::RemoveCompatibleObjectExplicitTransactiona
 	// non-buffered to reduce the additional locks/unlocks while the buffered version doesn't need any locking beyond the
 	// shared lock FindIndexedRow does. Not adding an exclusive here means some additional lock/unlocking but doesn't make
 	// the code thread unsafe.
-	RowHandle Row = Storage->LookupMappedRow(FMapKeyView(Object));
+	RowHandle Row = Storage->LookupMappedRow(ICompatibilityProvider::ObjectMappingDomain, FMapKeyView(Object));
 	if (Storage->IsRowAvailable(Row))
 	{
 		RemoveCompatibleObjectExplicitTransactionable<bEnableTransactions>(Object, Row);
diff --git a/Engine/Plugins/Experimental/EditorDataStorage/Source/TedsCore/Private/TypedElementDatabaseIndexTable.cpp b/Engine/Plugins/Experimental/EditorDataStorage/Source/TedsCore/Private/TypedElementDatabaseIndexTable.cpp
--- a/Engine/Plugins/Experimental/EditorDataStorage/Source/TedsCore/Private/TypedElementDatabaseIndexTable.cpp
+++ b/Engine/Plugins/Experimental/EditorDataStorage/Source/TedsCore/Private/TypedElementDatabaseIndexTable.cpp
@@ -23,80 +23,89 @@ namespace UE::Editor::DataStorage
 	{
 	}
 
-	RowHandle FMappingTable::Lookup(EGlobalLockScope LockScope, const FMapKeyView& Key) const
+	RowHandle FMappingTable::Lookup(EGlobalLockScope LockScope, const FName& Domain, const FMapKeyView& Key) const
 	{
 		FScopedSharedLock Lock(LockScope);
 
-		if (int32 Index = FindIndexUnguarded(Key); Index >= 0)
+		if (const FDomain* DomainInfo = Domains.Find(Domain))
 		{
-			RowHandle Result = Rows[Index];
-			if (IsDirty() && !DataStorage.IsRowAvailableUnsafe(Result))
+			if (int32 Index = DomainInfo->FindIndexUnguarded(Key); Index >= 0)
 			{
-				Result = InvalidRowHandle;
+				RowHandle Result = DomainInfo->Rows[Index];
+				if (IsDirty() && !DataStorage.IsRowAvailableUnsafe(Result))
+				{
+					Result = InvalidRowHandle;
+				}
+				return Result;
 			}
-			return Result;
-		}
-		else
-		{
-			return InvalidRowHandle;
 		}
+		return InvalidRowHandle;
 	}
 
-	void FMappingTable::Map(EGlobalLockScope LockScope, FMapKey Key, RowHandle Row)
+	void FMappingTable::Map(EGlobalLockScope LockScope, const FName& Domain, FMapKey Key, RowHandle Row)
 	{
 		FScopedExclusiveLock Lock(LockScope);
-		IndexRowUnguarded(MoveTemp(Key), Row);
+		
+		FDomain& DomainInfo = Domains.FindOrAdd(Domain);
+		DomainInfo.IndexRowUnguarded(DataStorage, MoveTemp(Key), Row, IsDirty());
 	}
 
-	void FMappingTable::BatchMap(EGlobalLockScope LockScope, TArrayView<TPair<FMapKey, RowHandle>> MapRowPairs)
+	void FMappingTable::BatchMap(EGlobalLockScope LockScope, const FName& Domain, TArrayView<TPair<FMapKey, RowHandle>> MapRowPairs)
 	{
 		FScopedExclusiveLock Lock(LockScope);
 
-		IndexLookupMap.Reserve(IndexLookupMap.Num() + MapRowPairs.Num());
+		FDomain& DomainInfo = Domains.FindOrAdd(Domain);
+		DomainInfo.IndexLookupMap.Reserve(DomainInfo.IndexLookupMap.Num() + MapRowPairs.Num());
 
-		checkf(Rows.Num() >= FreeList.Num(),
+		checkf(DomainInfo.Rows.Num() >= DomainInfo.FreeList.Num(),
 			TEXT("There can't be less rows than there are rows stored in the free list as free list is a subset of rows."));
-		int32 ArraySize = Rows.Num() - FreeList.Num() + MapRowPairs.Num();
-		Rows.Reserve(ArraySize);
-		Keys.Reserve(ArraySize);
+		int32 ArraySize = DomainInfo.Rows.Num() - DomainInfo.FreeList.Num() + MapRowPairs.Num();
+		DomainInfo.Rows.Reserve(ArraySize);
+		DomainInfo.Keys.Reserve(ArraySize);
 
 		for (TPair<FMapKey, RowHandle>& IndexAndRow : MapRowPairs)
 		{
-			IndexRowUnguarded(MoveTemp(IndexAndRow.Key), IndexAndRow.Value);
+			DomainInfo.IndexRowUnguarded(DataStorage, MoveTemp(IndexAndRow.Key), IndexAndRow.Value, IsDirty());
 		}
 	}
 
-	void FMappingTable::Remap(EGlobalLockScope LockScope, const FMapKeyView& OriginalKey, FMapKey NewKey)
+	void FMappingTable::Remap(EGlobalLockScope LockScope, const FName& Domain, const FMapKeyView& OriginalKey, FMapKey NewKey)
 	{
 		FScopedExclusiveLock Lock(LockScope);
 
-		uint64 Hash = OriginalKey.CalculateHash();
-		for (IndexLookupMapType::TKeyIterator It = IndexLookupMap.CreateKeyIterator(Hash); It; ++It)
+		if (FDomain* DomainInfo = Domains.Find(Domain))
 		{
-			if (Keys[It->Value] == OriginalKey)
+			uint64 Hash = OriginalKey.CalculateHash();
+			for (FDomain::IndexLookupMapType::TKeyIterator It = DomainInfo->IndexLookupMap.CreateKeyIterator(Hash); It; ++It)
 			{
-				int32 Index = It->Value;
-				It.RemoveCurrent();
-				IndexLookupMap.Add(NewKey.CalculateHash(), Index);
-				Keys[Index] = MoveTemp(NewKey);
+				if (DomainInfo->Keys[It->Value] == OriginalKey)
+				{
+					int32 Index = It->Value;
+					It.RemoveCurrent();
+					DomainInfo->IndexLookupMap.Add(NewKey.CalculateHash(), Index);
+					DomainInfo->Keys[Index] = MoveTemp(NewKey);
+				}
 			}
 		}
 	}
 
-	void FMappingTable::Remove(EGlobalLockScope LockScope, const FMapKeyView& Key)
+	void FMappingTable::Remove(EGlobalLockScope LockScope, const FName& Domain, const FMapKeyView& Key)
 	{
 		FScopedExclusiveLock Lock(LockScope);
 		
-		uint64 Hash = Key.CalculateHash();
-		for (IndexLookupMapType::TKeyIterator It = IndexLookupMap.CreateKeyIterator(Hash); It; ++It)
+		if (FDomain* DomainInfo = Domains.Find(Domain))
 		{
-			if (Keys[It->Value] == Key)
+			uint64 Hash = Key.CalculateHash();
+			for (FDomain::IndexLookupMapType::TKeyIterator It = DomainInfo->IndexLookupMap.CreateKeyIterator(Hash); It; ++It)
 			{
-				int32 Index = It->Value;
-				Rows[Index] = InvalidRowHandle;
-				Keys[Index].Clear();
-				FreeList.PushFirst(Index);
-				It.RemoveCurrent();
+				if (DomainInfo->Keys[It->Value] == Key)
+				{
+					int32 Index = It->Value;
+					DomainInfo->Rows[Index] = InvalidRowHandle;
+					DomainInfo->Keys[Index].Clear();
+					DomainInfo->FreeList.PushFirst(Index);
+					It.RemoveCurrent();
+				}
 			}
 		}
 	}
@@ -104,11 +113,76 @@ namespace UE::Editor::DataStorage
 	void FMappingTable::MarkDirty()
 	{
 		DirtyDueToRemoval = true;
+		CleanUpFrameDelay = DefaultCleanUpFrameDelay;
+	}
+
+	FMappingTable::FDomain::ERemoveInvalidRowsResult FMappingTable::FDomain::RemoveInvalidRows(
+		const FName& DomainName, UEditorDataStorage& DataStorage, FTimespan& RemainingTime)
+	{
+		FCleanUpInfo& Info = *CleanUpInfo;
+
+		// Only do work if there's any left
+		if (Info.RemainingJobs > 0 || !Info.DeletionQueue.IsEmpty())
+		{
+			uint64 DomainStartTime = FPlatformTime::Cycles64();
+			FTimespan DomainRemainingTime = RemainingTime;
+
+			// If there are still batches left from the previous frame or a new pass was just started, then collect invalid rows.
+			if (Info.RemainingJobs > 0)
+			{
+				int32 NumBatches = FMath::Min(Info.RemainingJobs, Info.MaxNumJobs);
+				ParallelForTemplate(NumBatches, [this, &DataStorage](int32 Block)
+					{
+						InspectRowBlockForCleanUp(DataStorage, Block);
+					}, EParallelForFlags::Unbalanced);
+
+				Info.RemainingJobs -= NumBatches;
+
+				// Follow up with any adjustments needed.
+				DomainRemainingTime -= FTimespan(FPlatformTime::Cycles64() - DomainStartTime);
+				if (NumBatches == Info.MaxNumJobs) // Only adjust if a full set of batches was used to avoid skewing.
+				{
+					AdjustJobCount(DomainRemainingTime);
+				}
+				if (Info.RemainingJobs == 0) // Don't adjust the batches when there are more to be run.
+				{
+					AdjustBatchSize(Info.JobCount);
+				}
+			}
+
+			// If there's time left and there are rows to remove, start removing rows.
+			if (Info.RemainingJobs == 0 && DomainRemainingTime > FTimespan::Zero() && !Info.DeletionQueue.IsEmpty())
+			{
+				DrainDeletionQueue(DomainRemainingTime);
+			}
+
+			UE_LOG(LogEditorDataStorage, Verbose,
+				TEXT("Mapping cleanup - %7.2fms - Domain %s has%sremaining rows, Batch size: %i, Job count: %i, Remaining jobs: %i, Max batches: %i"),
+				FTimespan(FPlatformTime::Cycles64() - DomainStartTime).GetTotalMilliseconds(),
+				*DomainName.ToString(),
+				Info.DeletionQueue.IsEmpty() ? TEXT(" no ") : TEXT(" "),
+				Info.BatchSize,
+				Info.JobCount,
+				Info.RemainingJobs,
+				Info.MaxNumJobs);
+
+			// See if there's time left for another pass for the next domain.
+			RemainingTime -= FTimespan(FPlatformTime::Cycles64() - DomainStartTime);
+			
+			// If there are no more remaining jobs and nothing queued for deletion that the domain has been fully cleaned up and should return false.
+			return (Info.RemainingJobs > 0 || !Info.DeletionQueue.IsEmpty()) 
+				? ERemoveInvalidRowsResult::HasRemainingWork 
+				: ERemoveInvalidRowsResult::CompletedWork;
+		}
+		else
+		{
+			return ERemoveInvalidRowsResult::NoWork;
+		}
 	}
 
 	void FMappingTable::RemoveInvalidRows()
 	{
-		TEDS_EVENT_SCOPE("Index Table clean up");
+		TEDS_EVENT_SCOPE("Mapping clean up");
 
 		FScopedExclusiveLock Lock(EGlobalLockScope::Public);
 
@@ -116,66 +190,55 @@ namespace UE::Editor::DataStorage
 		uint64 StartTime = FPlatformTime::Cycles64();
 
 		// If there's no work from a previous pass, then start a new pass if needed.
-		if (CleanUpInfo.RemainingJobs == 0 && CleanUpInfo.DeletionQueue.IsEmpty())
+		if (RemainingUnverifiedDomains == 0)
 		{
-			if (DirtyDueToRemoval.load())
+			if (DirtyDueToRemoval.load() && CleanUpFrameDelay-- == 0)
 			{
 				DirtyDueToRemoval = false;
-				
-				CleanUpInfo.JobCount = (Rows.Num() / CleanUpInfo.BatchSize) + 1;
-				CleanUpInfo.RemainingJobs = CleanUpInfo.JobCount;
-			}
-			else
-			{
-				return;
+				RemainingUnverifiedDomains = Domains.Num();
+
+				// Setup all domains for cleanup.
+				for (DomainMap::ElementType& It : Domains)
+				{
+					It.Value.CleanUpInfo->JobCount = (It.Value.Rows.Num() / It.Value.CleanUpInfo->BatchSize) + 1;
+					It.Value.CleanUpInfo->RemainingJobs = It.Value.CleanUpInfo->JobCount;
+				}
 			}
+			return;
 		}
 
-		// If there are still batches left from the previous frame or a new pass was just started, then collect invalid rows.
-		if (CleanUpInfo.RemainingJobs > 0)
+		if (RemainingUnverifiedDomains > 0)
 		{
-			int32 NumBatches = FMath::Min(CleanUpInfo.RemainingJobs, CleanUpInfo.MaxNumJobs);
-			ParallelForTemplate(NumBatches, [this](int32 Block)
-				{
-					InspectRowBlockForCleanUp(Block);
-				}, EParallelForFlags::Unbalanced);
-
-			CleanUpInfo.RemainingJobs -= NumBatches;
-
-			// Follow up with any adjustments needed.
-			RemainingTime -= FTimespan(FPlatformTime::Cycles64() - StartTime);
-			if (NumBatches == CleanUpInfo.MaxNumJobs) // Only adjust if a full set of batches was used to avoid skewing.
+			for (DomainMap::ElementType& It : Domains)
 			{
-				AdjustJobCount(RemainingTime);
-			}
-			if (CleanUpInfo.RemainingJobs == 0) // Don't adjust the batches when there are more to be run.
-			{
-				AdjustBatchSize(CleanUpInfo.JobCount);
+				switch (It.Value.RemoveInvalidRows(It.Key, DataStorage, RemainingTime))
+				{
+				case FDomain::ERemoveInvalidRowsResult::NoWork:
+					continue;
+				case FDomain::ERemoveInvalidRowsResult::CompletedWork:
+					RemainingUnverifiedDomains--;
+					// Fall through.
+				case FDomain::ERemoveInvalidRowsResult::HasRemainingWork:
+					if (RemainingTime <= 0)
+					{
+						return;
+					}
+				}
 			}
 		}
+	}
 
-		// If there's time left and there are rows to remove, start removing rows.
-		if (CleanUpInfo.RemainingJobs == 0 && RemainingTime > FTimespan::Zero() && !CleanUpInfo.DeletionQueue.IsEmpty())
-		{
-			DrainDeletionQueue(RemainingTime);
-		}
-		
-		UE_LOG(LogEditorDataStorage, Verbose,
-			TEXT("TEDS Index Table cleanup - %7.2fms - Has%sremaining rows, Batch size: %i, Job count: %i, Remaining jobs: %i, Max batches: %i"),
-			FTimespan(FPlatformTime::Cycles64() - StartTime).GetTotalMilliseconds(),
-			CleanUpInfo.DeletionQueue.IsEmpty() ? TEXT(" no ") : TEXT(" "),
-			CleanUpInfo.BatchSize,
-			CleanUpInfo.JobCount,
-			CleanUpInfo.RemainingJobs,
-			CleanUpInfo.MaxNumJobs);
+	bool FMappingTable::IsDirty() const
+	{
+		return RemainingUnverifiedDomains > 0 || DirtyDueToRemoval;
 	}
 
-	int32 FMappingTable::FindIndexUnguarded(const FMapKey& Key) const
+	int32 FMappingTable::FDomain::FindIndexUnguarded(const FMapKey& Key) const
 	{
 		return FindIndexUnguarded(Key.CalculateHash(), Key);
 	}
 
-	int32 FMappingTable::FindIndexUnguarded(uint64 Hash, const FMapKey& Key) const
+	int32 FMappingTable::FDomain::FindIndexUnguarded(uint64 Hash, const FMapKey& Key) const
 	{
 		for (IndexLookupMapType::TConstKeyIterator It = IndexLookupMap.CreateConstKeyIterator(Hash); It; ++It)
 		{
@@ -187,12 +250,12 @@ namespace UE::Editor::DataStorage
 		return -1;
 	}
 
-	int32 FMappingTable::FindIndexUnguarded(const FMapKeyView& Key) const
+	int32 FMappingTable::FDomain::FindIndexUnguarded(const FMapKeyView& Key) const
 	{
 		return FindIndexUnguarded(Key.CalculateHash(), Key);
 	}
 
-	int32 FMappingTable::FindIndexUnguarded(uint64 Hash, const FMapKeyView& Key) const
+	int32 FMappingTable::FDomain::FindIndexUnguarded(uint64 Hash, const FMapKeyView& Key) const
 	{
 		for (IndexLookupMapType::TConstKeyIterator It = IndexLookupMap.CreateConstKeyIterator(Hash); It; ++It)
 		{
@@ -204,17 +267,12 @@ namespace UE::Editor::DataStorage
 		return -1;
 	}
 
-	bool FMappingTable::IsDirty() const
-	{
-		return CleanUpInfo.RemainingJobs > 0 || DirtyDueToRemoval || !CleanUpInfo.DeletionQueue.IsEmpty();
-	}
-
-	void FMappingTable::IndexRowUnguarded(FMapKey&& Key, RowHandle Row)
+	void FMappingTable::FDomain::IndexRowUnguarded(UEditorDataStorage& DataStorage, FMapKey&& Key, RowHandle Row, bool bIsDirty)
 	{
 		uint64 Hash = Key.CalculateHash();
 		if (int32 Index = FindIndexUnguarded(Hash, Key); Index >= 0)
 		{
-			bool bUpdatingAllowed = IsDirty() && !DataStorage.IsRowAvailableUnsafe(Rows[Index]);
+			bool bUpdatingAllowed = bIsDirty && !DataStorage.IsRowAvailableUnsafe(Rows[Index]);
 			if (ensureMsgf(bUpdatingAllowed || Row == Rows[Index], TEXT("Another row has already been registered under key '%s'."), *Key.ToString()))
 			{
 				Rows[Index] = Row; // Update the stored row to the new row.
@@ -239,16 +297,16 @@ namespace UE::Editor::DataStorage
 		IndexLookupMap.Add(Hash, RowIndex);
 	}
 
-	void FMappingTable::InspectRowBlockForCleanUp(int32 Block)
+	void FMappingTable::FDomain::InspectRowBlockForCleanUp(UEditorDataStorage& DataStorage, int32 Block)
 	{
-		TEDS_EVENT_SCOPE("Index Table inspect rows");
+		TEDS_EVENT_SCOPE("Mapping inspect rows");
 
-		int32 Index = ((CleanUpInfo.JobCount - CleanUpInfo.RemainingJobs + Block) * CleanUpInfo.BatchSize);
+		int32 Index = ((CleanUpInfo->JobCount - CleanUpInfo->RemainingJobs + Block) * CleanUpInfo->BatchSize);
 
 		const RowHandle* Begin = Rows.GetData();
 		const RowHandle* Front = Rows.GetData() + Index;
-		const RowHandle* End = FMath::Min(Begin + Rows.Num(), Front + CleanUpInfo.BatchSize);
-		bool bIsFullBatch = (End - Front) == CleanUpInfo.BatchSize;
+		const RowHandle* End = FMath::Min(Begin + Rows.Num(), Front + CleanUpInfo->BatchSize);
+		bool bIsFullBatch = (End - Front) == CleanUpInfo->BatchSize;
 
 		uint64 StartTime = FPlatformTime::Cycles64();
 		
@@ -264,7 +322,7 @@ namespace UE::Editor::DataStorage
 
 		if (!Result.IsEmpty())
 		{
-			CleanUpInfo.DeletionQueue.Enqueue(MoveTemp(Result));
+			CleanUpInfo->DeletionQueue.Enqueue(MoveTemp(Result));
 		}
 
 		// Only adjust if this is a full batch, otherwise partial batches skew performance stats.
@@ -273,24 +331,24 @@ namespace UE::Editor::DataStorage
 			FTimespan Duration(FPlatformTime::Cycles64() - StartTime);
 			if (Duration >= FCleanUpInfo::MaxBatchDuration)
 			{
-				++CleanUpInfo.BatchWentOverTime;
+				++CleanUpInfo->BatchWentOverTime;
 			}
 			else if (Duration <= FCleanUpInfo::MinBatchDuration)
 			{
-				++CleanUpInfo.BatchWentUnderTime;
+				++CleanUpInfo->BatchWentUnderTime;
 			}
 		}
 	}
 
-	void FMappingTable::DrainDeletionQueue(FTimespan RemainingFrameTime)
+	void FMappingTable::FDomain::DrainDeletionQueue(FTimespan RemainingFrameTime)
 	{
-		TEDS_EVENT_SCOPE("Index Table drain deletion queue");
+		TEDS_EVENT_SCOPE("Mapping drain deletion queue");
 
 		uint64 StartTime = FPlatformTime::Cycles64();
 
-		while (!CleanUpInfo.DeletionQueue.IsEmpty() && FTimespan(FPlatformTime::Cycles64() - StartTime) < RemainingFrameTime)
+		while (!CleanUpInfo->DeletionQueue.IsEmpty() && FTimespan(FPlatformTime::Cycles64() - StartTime) < RemainingFrameTime)
 		{
-			TOptional<FCleanUpInfo::InvalidRowContainer> OptionalContainer = CleanUpInfo.DeletionQueue.Dequeue();
+			TOptional<FCleanUpInfo::InvalidRowContainer> OptionalContainer = CleanUpInfo->DeletionQueue.Dequeue();
 			checkf(OptionalContainer.IsSet(),
 				TEXT("Retrieved an invalid row container from a non-empty deletion queue, but no result was returned"));
 			
@@ -302,7 +360,7 @@ namespace UE::Editor::DataStorage
 		}
 	}
 
-	void FMappingTable::ClearRow(int32 Index)
+	void FMappingTable::FDomain::ClearRow(int32 Index)
 	{
 		FMapKey& Key = Keys[Index];
 		uint64 Hash = Key.CalculateHash();
@@ -313,37 +371,37 @@ namespace UE::Editor::DataStorage
 		FreeList.PushFirst(Index);
 	}
 
-	void FMappingTable::AdjustJobCount(FTimespan RemainingFrameTime)
+	void FMappingTable::FDomain::AdjustJobCount(FTimespan RemainingFrameTime)
 	{
 		if (RemainingFrameTime >= FCleanUpInfo::JobGrowthThreshold)
 		{
-			CleanUpInfo.MaxNumJobs = FMath::RoundToInt32(static_cast<float>(CleanUpInfo.MaxNumJobs) * FCleanUpInfo::JobGrowthFactor);
+			CleanUpInfo->MaxNumJobs = FMath::RoundToInt32(static_cast<float>(CleanUpInfo->MaxNumJobs) * FCleanUpInfo::JobGrowthFactor);
 		}
 		else if (RemainingFrameTime <= FCleanUpInfo::JobShrinkThreshold)
 		{
-			CleanUpInfo.MaxNumJobs = FMath::RoundToInt32(static_cast<float>(CleanUpInfo.MaxNumJobs) * FCleanUpInfo::JobShrinkFactor);
+			CleanUpInfo->MaxNumJobs = FMath::RoundToInt32(static_cast<float>(CleanUpInfo->MaxNumJobs) * FCleanUpInfo::JobShrinkFactor);
 		}
 
-		CleanUpInfo.MaxNumJobs = FMath::Clamp(CleanUpInfo.MaxNumJobs, FCleanUpInfo::MinJobCount, FCleanUpInfo::MaxJobCount);
+		CleanUpInfo->MaxNumJobs = FMath::Clamp(CleanUpInfo->MaxNumJobs, FCleanUpInfo::MinJobCount, FCleanUpInfo::MaxJobCount);
 	}
 
-	void FMappingTable::AdjustBatchSize(int32 JobCount)
+	void FMappingTable::FDomain::AdjustBatchSize(int32 JobCount)
 	{
-		if (CleanUpInfo.BatchWentOverTime.load() >= FCleanUpInfo::BatchShrinkThreshold)
+		if (CleanUpInfo->BatchWentOverTime.load() >= FCleanUpInfo::BatchShrinkThreshold)
 		{
 			// The batch size was too big to complete in the allotted time so take a sizable chunk off.
-			CleanUpInfo.BatchSize = FMath::RoundToInt32(static_cast<float>(CleanUpInfo.BatchSize) * FCleanUpInfo::BatchShrinkFactor);
+			CleanUpInfo->BatchSize = FMath::RoundToInt32(static_cast<float>(CleanUpInfo->BatchSize) * FCleanUpInfo::BatchShrinkFactor);
 		}
 		// Check if more than the minimum percentage of jobs took less time than currently allocated.
-		else if (CleanUpInfo.BatchWentUnderTime.load() >= 
+		else if (CleanUpInfo->BatchWentUnderTime.load() >=
 			static_cast<uint32>(FMath::RoundToInt32(static_cast<float>(JobCount) * FCleanUpInfo::BatchIncreaseThreshold)))
 		{
 			// The batch size fitted within the allotted time so slightly increase it.
-			CleanUpInfo.BatchSize = FMath::RoundToInt32(static_cast<float>(CleanUpInfo.BatchSize) * FCleanUpInfo::BatchGrowthFactor);
+			CleanUpInfo->BatchSize = FMath::RoundToInt32(static_cast<float>(CleanUpInfo->BatchSize) * FCleanUpInfo::BatchGrowthFactor);
 		}
 		// Clamp the job batch size within reasonable sizes to avoid extremes.
-		CleanUpInfo.BatchSize = FMath::Clamp(CleanUpInfo.BatchSize, FCleanUpInfo::MinBatchSize, FCleanUpInfo::MaxBatchSize);
-		CleanUpInfo.BatchWentOverTime = 0;
-		CleanUpInfo.BatchWentUnderTime = 0;
+		CleanUpInfo->BatchSize = FMath::Clamp(CleanUpInfo->BatchSize, FCleanUpInfo::MinBatchSize, FCleanUpInfo::MaxBatchSize);
+		CleanUpInfo->BatchWentOverTime = 0;
+		CleanUpInfo->BatchWentUnderTime = 0;
 	}
 } // namespace UE::Editor::DataStorage
diff --git a/Engine/Plugins/Experimental/EditorDataStorage/Source/TedsCore/Private/TypedElementDatabaseIndexTable.h b/Engine/Plugins/Experimental/EditorDataStorage/Source/TedsCore/Private/TypedElementDatabaseIndexTable.h
--- a/Engine/Plugins/Experimental/EditorDataStorage/Source/TedsCore/Private/TypedElementDatabaseIndexTable.h
+++ b/Engine/Plugins/Experimental/EditorDataStorage/Source/TedsCore/Private/TypedElementDatabaseIndexTable.h
@@ -12,6 +12,7 @@
 #include "Elements/Interfaces/TypedElementQueryStorageInterfaces.h"
 #include "GlobalLock.h"
 #include "Misc/Timespan.h"
+#include "Templates/UniquePtr.h"
 
 class UEditorDataStorage;
 
@@ -27,11 +28,11 @@ namespace UE::Editor::DataStorage
 	public:
 		explicit FMappingTable(UEditorDataStorage& DataStorage);
 
-		RowHandle Lookup(EGlobalLockScope LockScope, const FMapKeyView& Key) const;
-		void Map(EGlobalLockScope LockScope, FMapKey Key, RowHandle Row);
-		void BatchMap(EGlobalLockScope LockScope, TArrayView<TPair<FMapKey, RowHandle>> MapRowPairs);
-		void Remap(EGlobalLockScope LockScope, const FMapKeyView& OriginalKey, FMapKey NewKey);
-		void Remove(EGlobalLockScope LockScope, const FMapKeyView& Key);
+		RowHandle Lookup(EGlobalLockScope LockScope, const FName& Domain, const FMapKeyView& Key) const;
+		void Map(EGlobalLockScope LockScope, const FName& Domain, FMapKey Key, RowHandle Row);
+		void BatchMap(EGlobalLockScope LockScope, const FName& Domain, TArrayView<TPair<FMapKey, RowHandle>> MapRowPairs);
+		void Remap(EGlobalLockScope LockScope, const FName& Domain, const FMapKeyView& OriginalKey, FMapKey NewKey);
+		void Remove(EGlobalLockScope LockScope, const FName& Domain, const FMapKeyView& Key);
 
 		void MarkDirty();
 		void RemoveInvalidRows();
@@ -39,10 +40,7 @@ namespace UE::Editor::DataStorage
 	private:
 		struct FCleanUpInfo final
 		{
-			// The number of row indices to store on the stack before allocating memory on the heap.
-			static constexpr uint32 StackReservationCount = 1024;
-
-			using InvalidRowContainer = TArray<int32, TInlineAllocator<StackReservationCount>>;
+			using InvalidRowContainer = TArray<int32>;
 
 			// The target amount of time that should be spend per frame on cleaning up old rows.
 			static const FTimespan TargetDuration;
@@ -98,29 +96,51 @@ namespace UE::Editor::DataStorage
 			int32 RemainingJobs = 0;
 		};
 		
-		using IndexLookupMapType = TMultiMap<uint64, int32>;
-		IndexLookupMapType IndexLookupMap;
-		TArray<RowHandle> Rows;
-		TArray<FMapKey> Keys;
-		TDeque<int32> FreeList;
+		struct FDomain
+		{
+			using IndexLookupMapType = TMultiMap<uint64, int32>;
+			
+			IndexLookupMapType IndexLookupMap;
+			TArray<RowHandle> Rows;
+			TArray<FMapKey> Keys;
+			TDeque<int32> FreeList;
 
-		FCleanUpInfo CleanUpInfo;
-		
-		UEditorDataStorage& DataStorage;
-		std::atomic<bool> DirtyDueToRemoval = false;
+			TUniquePtr<FCleanUpInfo> CleanUpInfo = MakeUnique<FCleanUpInfo>();
 
-		int32 FindIndexUnguarded(const FMapKey& Key) const;
-		int32 FindIndexUnguarded(uint64 Hash, const FMapKey& Key) const;
-		int32 FindIndexUnguarded(const FMapKeyView& Key) const;
-		int32 FindIndexUnguarded(uint64 Hash, const FMapKeyView& Key) const;
 
-		bool IsDirty() const;
-		void IndexRowUnguarded(FMapKey&& Key, RowHandle Row);
+			int32 FindIndexUnguarded(const FMapKey& Key) const;
+			int32 FindIndexUnguarded(uint64 Hash, const FMapKey& Key) const;
+			int32 FindIndexUnguarded(const FMapKeyView& Key) const;
+			int32 FindIndexUnguarded(uint64 Hash, const FMapKeyView& Key) const;
+
+			void IndexRowUnguarded(UEditorDataStorage& DataStorage, FMapKey&& Key, RowHandle Row, bool bIsDirty);
+
+			enum class ERemoveInvalidRowsResult
+			{
+				NoWork,
+				HasRemainingWork,
+				CompletedWork
+			};
+			ERemoveInvalidRowsResult RemoveInvalidRows(const FName& DomainName, UEditorDataStorage& DataStorage, FTimespan& RemainingTime);
+
+			void InspectRowBlockForCleanUp(UEditorDataStorage& DataStorage, int32 Block);
+			void DrainDeletionQueue(FTimespan RemainingFrameTime);
+			void ClearRow(int32 Index);
+			void AdjustJobCount(FTimespan RemainingFrameTime);
+			void AdjustBatchSize(int32 JobCount);
+		};
+		using DomainMap = TMap<FName, FDomain>;
+		DomainMap Domains;
+
+		int32 RemainingUnverifiedDomains = 0;
+		// The number of frames to wait to start the cleanup. This delay is added as several operations in the editor do a delete spread out of over
+		// a few frames.
+		static constexpr int8 DefaultCleanUpFrameDelay = 5;
+		std::atomic<int8> CleanUpFrameDelay = 0;
+		std::atomic<bool> DirtyDueToRemoval = false;
+		
+		UEditorDataStorage& DataStorage;
 		
-		void InspectRowBlockForCleanUp(int32 Block);
-		void DrainDeletionQueue(FTimespan RemainingFrameTime);
-		void ClearRow(int32 Index);
-		void AdjustJobCount(FTimespan RemainingFrameTime);
-		void AdjustBatchSize(int32 JobCount);
+		bool IsDirty() const;
 	};
 } // namespace UE::Editor::DataStorage
diff --git a/Engine/Plugins/Experimental/EditorDataStorage/Source/TedsCore/Private/TypedElementDatabaseUI.cpp b/Engine/Plugins/Experimental/EditorDataStorage/Source/TedsCore/Private/TypedElementDatabaseUI.cpp
--- a/Engine/Plugins/Experimental/EditorDataStorage/Source/TedsCore/Private/TypedElementDatabaseUI.cpp
+++ b/Engine/Plugins/Experimental/EditorDataStorage/Source/TedsCore/Private/TypedElementDatabaseUI.cpp
@@ -97,7 +97,7 @@ void UEditorDataStorageUi::PostInitialize(UE::Editor::DataStorage::ICoreProvider
 
 		if (ensureMsgf(UnresolvedParent, TEXT("Expected valid parent for row in query with TColumn<FUnresolvedTableRowParentColumn>()")))
 		{
-			RowHandle ParentRow = StorageInterface->LookupMappedRow(UnresolvedParent->ParentIdKey);
+			RowHandle ParentRow = StorageInterface->LookupMappedRow(IUiProvider::PurposeMappingDomain, UnresolvedParent->ParentIdKey);
 			
 			if (Storage->IsRowAvailable(ParentRow))
 			{
@@ -107,7 +107,7 @@ void UEditorDataStorageUi::PostInitialize(UE::Editor::DataStorage::ICoreProvider
 		}
 	}
 
-	// Post Init, there is a processer run to resolve parent rows automatically so we don't need this query anymore
+	// Post Init, there is a processor run to resolve parent rows automatically so we don't need this query anymore
 	StorageInterface->UnregisterQuery(PurposesQuery);
 }
 
@@ -121,7 +121,7 @@ UE::Editor::DataStorage::RowHandle UEditorDataStorageUi::RegisterWidgetPurpose(c
 
 	// If a purpose is already registered against this name, let the user know
 	const FMapKey Key = FMapKey(PurposeID);
-	RowHandle ExistingRow = Storage->LookupMappedRow(Key);
+	RowHandle ExistingRow = Storage->LookupMappedRow(PurposeMappingDomain, Key);
 	if (Storage->IsRowAvailable(ExistingRow))
 	{
 		ensureMsgf(false, TEXT("Existing purpose found registered with name: %s"), *PurposeID.ToString());
@@ -130,7 +130,7 @@ UE::Editor::DataStorage::RowHandle UEditorDataStorageUi::RegisterWidgetPurpose(c
 
 	// Add the row and register the mapping
 	RowHandle PurposeRowHandle = Storage->AddRow(WidgetPurposeTable);
-	Storage->MapRow(Key, PurposeRowHandle);
+	Storage->MapRow(PurposeMappingDomain, Key, PurposeRowHandle);
 
 	// Setup the relevant columns
 	if (FWidgetPurposeColumn* PurposeColumn = Storage->GetColumn<FWidgetPurposeColumn>(PurposeRowHandle))
@@ -541,7 +541,7 @@ void UEditorDataStorageUi::RegisterWidgetPurpose(FName Purpose, EPurposeType Typ
 
 bool UEditorDataStorageUi::RegisterWidgetFactory(FName Purpose, const UScriptStruct* Constructor)
 {
-	UE::Editor::DataStorage::RowHandle PurposeRow = Storage->LookupMappedRow(UE::Editor::DataStorage::FMapKey(Purpose));
+	UE::Editor::DataStorage::RowHandle PurposeRow = Storage->LookupMappedRow(PurposeMappingDomain, UE::Editor::DataStorage::FMapKey(Purpose));
 
 	if (Storage->IsRowAvailable(PurposeRow))
 	{
@@ -557,7 +557,7 @@ bool UEditorDataStorageUi::RegisterWidgetFactory(FName Purpose, const UScriptStr
 bool UEditorDataStorageUi::RegisterWidgetFactory(
 	FName Purpose, const UScriptStruct* Constructor, UE::Editor::DataStorage::Queries::FConditions Columns)
 {
-	UE::Editor::DataStorage::RowHandle PurposeRow = Storage->LookupMappedRow(UE::Editor::DataStorage::FMapKey(Purpose));
+	UE::Editor::DataStorage::RowHandle PurposeRow = Storage->LookupMappedRow(PurposeMappingDomain, UE::Editor::DataStorage::FMapKey(Purpose));
 	
 	if (Storage->IsRowAvailable(PurposeRow))
 	{
@@ -572,7 +572,7 @@ bool UEditorDataStorageUi::RegisterWidgetFactory(
 
 bool UEditorDataStorageUi::RegisterWidgetFactory(FName Purpose, TUniquePtr<FTypedElementWidgetConstructor>&& Constructor)
 {
-	UE::Editor::DataStorage::RowHandle PurposeRow = Storage->LookupMappedRow(UE::Editor::DataStorage::FMapKey(Purpose));
+	UE::Editor::DataStorage::RowHandle PurposeRow = Storage->LookupMappedRow(PurposeMappingDomain, UE::Editor::DataStorage::FMapKey(Purpose));
 	if (Storage->IsRowAvailable(PurposeRow))
 	{
 		return RegisterWidgetFactory(PurposeRow, MoveTemp(Constructor));
@@ -586,7 +586,7 @@ bool UEditorDataStorageUi::RegisterWidgetFactory(FName Purpose, TUniquePtr<FType
 bool UEditorDataStorageUi::RegisterWidgetFactory(FName Purpose, TUniquePtr<FTypedElementWidgetConstructor>&& Constructor, 
 	UE::Editor::DataStorage::Queries::FConditions Columns)
 {
-	UE::Editor::DataStorage::RowHandle PurposeRow = Storage->LookupMappedRow(UE::Editor::DataStorage::FMapKey(Purpose));
+	UE::Editor::DataStorage::RowHandle PurposeRow = Storage->LookupMappedRow(PurposeMappingDomain, UE::Editor::DataStorage::FMapKey(Purpose));
 	
 	if (Storage->IsRowAvailable(PurposeRow))
 	{
@@ -601,7 +601,7 @@ bool UEditorDataStorageUi::RegisterWidgetFactory(FName Purpose, TUniquePtr<FType
 void UEditorDataStorageUi::CreateWidgetConstructors(FName Purpose,
 	const UE::Editor::DataStorage::FMetaDataView& Arguments, const WidgetConstructorCallback& Callback)
 {
-	UE::Editor::DataStorage::RowHandle PurposeRow = Storage->LookupMappedRow(UE::Editor::DataStorage::FMapKey(Purpose));
+	UE::Editor::DataStorage::RowHandle PurposeRow = Storage->LookupMappedRow(PurposeMappingDomain, UE::Editor::DataStorage::FMapKey(Purpose));
 	CreateWidgetConstructors(PurposeRow, Arguments, Callback);
 }
 
@@ -609,14 +609,14 @@ void UEditorDataStorageUi::CreateWidgetConstructors(FName Purpose, EMatchApproac
 	TArray<TWeakObjectPtr<const UScriptStruct>>& Columns, const UE::Editor::DataStorage::FMetaDataView& Arguments,
 	const WidgetConstructorCallback& Callback)
 {
-	UE::Editor::DataStorage::RowHandle PurposeRow = Storage->LookupMappedRow(UE::Editor::DataStorage::FMapKey(Purpose));
+	UE::Editor::DataStorage::RowHandle PurposeRow = Storage->LookupMappedRow(PurposeMappingDomain, UE::Editor::DataStorage::FMapKey(Purpose));
 	CreateWidgetConstructors(PurposeRow, MatchApproach, Columns, Arguments, Callback);
 }
 
 void UEditorDataStorageUi::ConstructWidgets(FName Purpose, const UE::Editor::DataStorage::FMetaDataView& Arguments,
 	const WidgetCreatedCallback& ConstructionCallback)
 {
-	UE::Editor::DataStorage::RowHandle PurposeRow = Storage->LookupMappedRow(UE::Editor::DataStorage::FMapKey(Purpose));
+	UE::Editor::DataStorage::RowHandle PurposeRow = Storage->LookupMappedRow(PurposeMappingDomain, UE::Editor::DataStorage::FMapKey(Purpose));
 	ConstructWidgets(PurposeRow, Arguments, ConstructionCallback);
 }
 
@@ -791,7 +791,7 @@ UE::Editor::DataStorage::IUiProvider::FPurposeID UEditorDataStorageUi::GetGenera
 
 UE::Editor::DataStorage::RowHandle UEditorDataStorageUi::FindPurpose(const FPurposeID& PurposeID) const
 {
-	return Storage->LookupMappedRow(UE::Editor::DataStorage::FMapKey(PurposeID));
+	return Storage->LookupMappedRow(PurposeMappingDomain, UE::Editor::DataStorage::FMapKey(PurposeID));
 }
 
 void UEditorDataStorageUi::CreateStandardArchetypes()
diff --git a/Engine/Plugins/Experimental/EditorDataStorageFeatures/Source/TedsAssetData/Private/Processors/AssetProcessors.cpp b/Engine/Plugins/Experimental/EditorDataStorageFeatures/Source/TedsAssetData/Private/Processors/AssetProcessors.cpp
--- a/Engine/Plugins/Experimental/EditorDataStorageFeatures/Source/TedsAssetData/Private/Processors/AssetProcessors.cpp
+++ b/Engine/Plugins/Experimental/EditorDataStorageFeatures/Source/TedsAssetData/Private/Processors/AssetProcessors.cpp
@@ -238,7 +238,7 @@ void UTedsAssetDataFactory::PreShutdown(UE::Editor::DataStorage::ICoreProvider&
 void UTedsAssetDataFactory::OnSetFolderColor(const FString& Path, UE::Editor::DataStorage::ICoreProvider* DataStorage)
 {
 	const UE::Editor::DataStorage::FMapKey PathKey = UE::Editor::DataStorage::FMapKey(FName(Path));
-	const UE::Editor::DataStorage::RowHandle Row = DataStorage->LookupMappedRow(PathKey);
+	const UE::Editor::DataStorage::RowHandle Row = DataStorage->LookupMappedRow(UE::Editor::AssetData::MappingDomain, PathKey);
 
 	if (DataStorage->IsRowAvailable(Row))
 	{
@@ -251,7 +251,7 @@ void UTedsAssetDataFactory::OnFavoritesChanged(const FContentBrowserItemPath& It
 	if (ItemPath.HasInternalPath())
 	{
 		const UE::Editor::DataStorage::FMapKey PathKey = UE::Editor::DataStorage::FMapKey(ItemPath.GetInternalPathName());
-		const UE::Editor::DataStorage::RowHandle Row = DataStorage->LookupMappedRow(PathKey);
+		const UE::Editor::DataStorage::RowHandle Row = DataStorage->LookupMappedRow(UE::Editor::AssetData::MappingDomain, PathKey);
 
 		if (DataStorage->IsRowAvailable(Row))
 		{
diff --git a/Engine/Plugins/Experimental/EditorDataStorageFeatures/Source/TedsAssetData/Private/TedsAssetData.cpp b/Engine/Plugins/Experimental/EditorDataStorageFeatures/Source/TedsAssetData/Private/TedsAssetData.cpp
--- a/Engine/Plugins/Experimental/EditorDataStorageFeatures/Source/TedsAssetData/Private/TedsAssetData.cpp
+++ b/Engine/Plugins/Experimental/EditorDataStorageFeatures/Source/TedsAssetData/Private/TedsAssetData.cpp
@@ -37,7 +37,7 @@ FPopulateAssetDataRowArgs ThreadSafe_PopulateAssetDataTableRow(TAssetData&& InAs
 	FPopulateAssetDataRowArgs Output;
 	Output.ObjectPathKey = FMapKey(InAssetData.GetSoftObjectPath());
 
-	if (Database.IsRowAssigned(Database.LookupMappedRow(Output.ObjectPathKey)))
+	if (Database.IsRowAssigned(Database.LookupMappedRow(MappingDomain, Output.ObjectPathKey)))
 	{
 		// No need to initialize the rest of the row here. The invalid's asset data will be used as flag to skip the data generated here.
 		return Output;
@@ -212,7 +212,7 @@ FTedsAssetData::FTedsAssetData(UE::Editor::DataStorage::ICoreProvider& InDatabas
 		{
 			KeysToReservedPathRows[Index] = TPair<FMapKey, RowHandle>(PopulatePathRowArgs[Index].AssetRegistryPathKey, ReservedPopulatePathRows[Index]);
 		});
-	Database.BatchMapRows(KeysToReservedPathRows);
+	Database.BatchMapRows(MappingDomain, KeysToReservedPathRows);
 
 	// Populate Path Rows
 	Database.BatchAddRow(PathsTable, ReservedPopulatePathRows, [PathRowArgs = MoveTemp(PopulatePathRowArgs), Index = 0, this](RowHandle InRowHandle) mutable
@@ -236,7 +236,7 @@ FTedsAssetData::FTedsAssetData(UE::Editor::DataStorage::ICoreProvider& InDatabas
 		{
 			KeysToResevedAssetRows[Index] = TPair<FMapKey, RowHandle>(MoveTemp(PopulateAssetDataRowArgs[Index].ObjectPathKey), ReservedAssetDataRows[Index]);
 		});
-	Database.BatchMapRows(KeysToResevedAssetRows);
+	Database.BatchMapRows(MappingDomain, KeysToResevedAssetRows);
 
 	// Populate Asset Rows
 	Database.BatchAddRow(AssetsDataTable, ReservedAssetDataRows,  [AssetDataRowArgs = MoveTemp(PopulateAssetDataRowArgs), Index = 0, this](RowHandle InRowHandle) mutable
@@ -274,7 +274,7 @@ FTedsAssetData::~FTedsAssetData()
 			AssetRegistry->EnumerateAllCachedPaths([this](FName InPath)
 				{
 					const FMapKeyView PathKey = FMapKeyView(InPath);
-					const RowHandle Row = Database.LookupMappedRow(PathKey);
+					const RowHandle Row = Database.LookupMappedRow(MappingDomain, PathKey);
 					Database.RemoveRow(Row);
 					return true;
 				});
@@ -282,7 +282,7 @@ FTedsAssetData::~FTedsAssetData()
 			AssetRegistry->EnumerateAllAssets([this](const FAssetData& InAssetData)
 				{
 					const FMapKey AssetPathKey = FMapKey(InAssetData.GetSoftObjectPath());
-					const RowHandle Row = Database.LookupMappedRow(AssetPathKey);
+					const RowHandle Row = Database.LookupMappedRow(MappingDomain, AssetPathKey);
 					Database.RemoveRow(Row);
 					return true;
 				});
@@ -354,7 +354,7 @@ void FTedsAssetData::OnAssetsAdded(TConstArrayView<FAssetData> InAssetsAdded)
 				++Index;
 			});
 
-		Database.BatchMapRows(KeyToRow);
+		Database.BatchMapRows(MappingDomain, KeyToRow);
 	}
 }
 
@@ -372,7 +372,7 @@ void FTedsAssetData::OnAssetsRemoved(TConstArrayView<FAssetData> InAssetsRemoved
 	for (const FAssetData& Asset : InAssetsRemoved)
 	{
 		const FMapKey AssetKey = FMapKey(Asset.GetSoftObjectPath());
-		const RowHandle AssetRow = Database.LookupMappedRow(AssetKey);
+		const RowHandle AssetRow = Database.LookupMappedRow(MappingDomain, AssetKey);
 		if (Database.IsRowAssigned(AssetRow))
 		{
 			RowsToRemove.Add(AssetRow);
@@ -393,7 +393,7 @@ void FTedsAssetData::OnAssetsUpdated(TConstArrayView<FAssetData> InAssetsUpdated
 	for (const FAssetData& Asset : InAssetsUpdated)
 	{
 		const FMapKey AssetKey = FMapKey(Asset.GetSoftObjectPath());
-		const RowHandle Row = Database.LookupMappedRow(AssetKey);
+		const RowHandle Row = Database.LookupMappedRow(MappingDomain, AssetKey);
 		if (Database.IsRowAssigned(Row))
 		{
 			Database.GetColumn<FAssetDataColumn_Experimental>(Row)->AssetData = Asset;
@@ -413,7 +413,7 @@ void FTedsAssetData::OnAssetsUpdatedOnDisk(TConstArrayView<FAssetData> InAssetsU
 	for (const FAssetData& Asset : InAssetsUpdated)
 	{
 		const FMapKey AssetKey = FMapKey(Asset.GetSoftObjectPath());
-		const RowHandle Row = Database.LookupMappedRow(AssetKey);
+		const RowHandle Row = Database.LookupMappedRow(MappingDomain, AssetKey);
 		if (Database.IsRowAssigned(Row))
 		{
 			Database.GetColumn<FAssetDataColumn_Experimental>(Row)->AssetData = Asset;
@@ -434,7 +434,7 @@ void FTedsAssetData::OnAssetRenamed(const FAssetData& InAsset, const FString& In
 
 	FMapKey NewAssetHash = FMapKey(InAsset.GetSoftObjectPath());
 	const FMapKey OldAssetHash = FMapKey(FSoftObjectPath(InOldObjectPath));
-	const RowHandle Row = Database.LookupMappedRow(OldAssetHash);
+	const RowHandle Row = Database.LookupMappedRow(MappingDomain, OldAssetHash);
 	if (Database.IsRowAssigned(Row))
 	{
 		Database.GetColumn<FAssetDataColumn_Experimental>(Row)->AssetData = InAsset;
@@ -452,7 +452,7 @@ void FTedsAssetData::OnAssetRenamed(const FAssetData& InAsset, const FString& In
 		const FMapKey OldFolderHash = FMapKey(FName(OldPackagePath));
 
 		Database.AddColumn<FUpdatedPathTag>(Row);
-		Database.RemapRow(OldAssetHash, MoveTemp(NewAssetHash));
+		Database.RemapRow(MappingDomain, OldAssetHash, MoveTemp(NewAssetHash));
 	}
 }
 
@@ -475,7 +475,7 @@ void FTedsAssetData::OnPathsAdded(TConstArrayView<FStringView> InPathsAdded)
 			FMapKey AssetRegistryPathKey = FMapKey(PathName);
 			FPopulatePathRowArgs RowArgs;
 
-			if (Database.LookupMappedRow(AssetRegistryPathKey) != InvalidRowHandle)
+			if (Database.LookupMappedRow(MappingDomain, AssetRegistryPathKey) != InvalidRowHandle)
 			{
 				RowArgs.MarkAsInvalid();
 			}
@@ -518,7 +518,7 @@ void FTedsAssetData::OnPathsAdded(TConstArrayView<FStringView> InPathsAdded)
 			++RowCount;
 		}
 
-		Database.BatchMapRows(KeysAndRows);
+		Database.BatchMapRows(MappingDomain, KeysAndRows);
 
 		int32 Index = 0;
 		Database.BatchAddRow(PathsTable, ReservedRow, [&PopulateRowArgs, &Index, this](RowHandle InRowHandle)
@@ -547,7 +547,7 @@ void FTedsAssetData::OnPathsRemoved(TConstArrayView<FStringView> InPathsRemoved)
 	for (const FStringView Path : InPathsRemoved)
 	{
 		FMapKey PathKey = FMapKey(FName(Path));
-		Database.RemoveRow(Database.LookupMappedRow(PathKey));
+		Database.RemoveRow(Database.LookupMappedRow(MappingDomain, PathKey));
 	}
 }
 } // namespace UE::Editor::AssetData::Private
diff --git a/Engine/Plugins/Experimental/EditorDataStorageFeatures/Source/TedsAssetData/Private/Tests/TedsAssetDataTest.cpp b/Engine/Plugins/Experimental/EditorDataStorageFeatures/Source/TedsAssetData/Private/Tests/TedsAssetDataTest.cpp
--- a/Engine/Plugins/Experimental/EditorDataStorageFeatures/Source/TedsAssetData/Private/Tests/TedsAssetDataTest.cpp
+++ b/Engine/Plugins/Experimental/EditorDataStorageFeatures/Source/TedsAssetData/Private/Tests/TedsAssetDataTest.cpp
@@ -70,7 +70,7 @@ bool FTedsAssetDataTest::RunTest(const FString& Parameters)
 	
 	for (const TPair<FName, int32>& PathAndAssetCount : AssetRegistryPathsAndAssetCount)
 	{
-		RowHandle Row = Database->LookupMappedRow(FMapKeyView(PathAndAssetCount.Key));
+		RowHandle Row = Database->LookupMappedRow(AssetData::MappingDomain, FMapKeyView(PathAndAssetCount.Key));
 		bool bHadError = false;
 		bHadError |= !TestTrue(TEXT("Asset registry folder/path is indexed in TEDS"), Row != InvalidRowHandle);
 
@@ -103,7 +103,7 @@ bool FTedsAssetDataTest::RunTest(const FString& Parameters)
 		// Verse package don't emit asset added event so we will ignore those for the test to avoid false errors
 		if (!FPackageName::IsVersePackage(Builder))
 		{ 
-			RowHandle Row = Database->LookupMappedRow(FMapKey(AssetData.GetSoftObjectPath()));
+			RowHandle Row = Database->LookupMappedRow(AssetData::MappingDomain, FMapKey(AssetData.GetSoftObjectPath()));
 
 			bool bHadError = false;
 
diff --git a/Engine/Plugins/Experimental/EditorDataStorageFeatures/Source/TedsAssetData/Private/Widgets/DynamicAssetDataColumnBaseWidget.cpp b/Engine/Plugins/Experimental/EditorDataStorageFeatures/Source/TedsAssetData/Private/Widgets/DynamicAssetDataColumnBaseWidget.cpp
--- a/Engine/Plugins/Experimental/EditorDataStorageFeatures/Source/TedsAssetData/Private/Widgets/DynamicAssetDataColumnBaseWidget.cpp
+++ b/Engine/Plugins/Experimental/EditorDataStorageFeatures/Source/TedsAssetData/Private/Widgets/DynamicAssetDataColumnBaseWidget.cpp
@@ -69,7 +69,7 @@ TSharedPtr<SWidget> FDynamicAssetDataColumnBaseWidgetConstructor::CreateWidget(U
 	const TAttribute<FAssetData>& AssetDataAttribute = Binder.BindData(*DynamicTemplates, &FItemStringAttributeColumn_Experimental::Value, [DataStorage] (const FString& InValue)
 		{
 			const FMapKey ReferencedAssetRowKey = FMapKey(FSoftObjectPath(InValue));
-			const RowHandle ReferencedRow = DataStorage->LookupMappedRow(ReferencedAssetRowKey);
+			const RowHandle ReferencedRow = DataStorage->LookupMappedRow(UE::Editor::AssetData::MappingDomain, ReferencedAssetRowKey);
 			if (ReferencedRow != InvalidRowHandle)
 			{
 				FAttributeBinder Binder(ReferencedRow, DataStorage);
diff --git a/Engine/Plugins/Experimental/EditorDataStorageFeatures/Source/TedsAssetData/Public/TedsAssetDataColumns.h b/Engine/Plugins/Experimental/EditorDataStorageFeatures/Source/TedsAssetData/Public/TedsAssetDataColumns.h
--- a/Engine/Plugins/Experimental/EditorDataStorageFeatures/Source/TedsAssetData/Public/TedsAssetDataColumns.h
+++ b/Engine/Plugins/Experimental/EditorDataStorageFeatures/Source/TedsAssetData/Public/TedsAssetDataColumns.h
@@ -17,6 +17,11 @@
 
 #include "TedsAssetDataColumns.generated.h"
 
+namespace UE::Editor::AssetData
+{
+	inline static const FName MappingDomain = "Asset";
+}
+
 USTRUCT()
 struct FAssetPathColumn_Experimental : public FEditorDataStorageColumn
 {
diff --git a/Engine/Plugins/Experimental/EditorDataStorageFeatures/Source/TedsContentBrowser/Private/TedsContentBrowserModule.cpp b/Engine/Plugins/Experimental/EditorDataStorageFeatures/Source/TedsContentBrowser/Private/TedsContentBrowserModule.cpp
--- a/Engine/Plugins/Experimental/EditorDataStorageFeatures/Source/TedsContentBrowser/Private/TedsContentBrowserModule.cpp
+++ b/Engine/Plugins/Experimental/EditorDataStorageFeatures/Source/TedsContentBrowser/Private/TedsContentBrowserModule.cpp
@@ -183,12 +183,12 @@ namespace UE::Editor::ContentBrowser
 		if (Item->GetItem().Legacy_TryGetAssetData(ItemAssetData))
 		{
 			FMapKey Key = FMapKey(ItemAssetData.GetSoftObjectPath());
-			RowHandle = DataStorage->LookupMappedRow(Key);
+			RowHandle = DataStorage->LookupMappedRow(AssetData::MappingDomain, Key);
 		}
 		else if(Item->GetItem().Legacy_TryGetPackagePath(PackagePath))
 		{
 			FMapKeyView Key = FMapKeyView(PackagePath);
-			RowHandle = DataStorage->LookupMappedRow(Key);
+			RowHandle = DataStorage->LookupMappedRow(AssetData::MappingDomain, Key);
 		}
 
 		return RowHandle;
diff --git a/Engine/Plugins/Experimental/EditorDataStorageFeatures/Source/TedsOutliner/Private/ActorFolderFactory.cpp b/Engine/Plugins/Experimental/EditorDataStorageFeatures/Source/TedsOutliner/Private/ActorFolderFactory.cpp
--- a/Engine/Plugins/Experimental/EditorDataStorageFeatures/Source/TedsOutliner/Private/ActorFolderFactory.cpp
+++ b/Engine/Plugins/Experimental/EditorDataStorageFeatures/Source/TedsOutliner/Private/ActorFolderFactory.cpp
@@ -251,15 +251,14 @@ void UTedsActorFolderFactory::OnLevelAdded(ULevel* InLevel, UWorld* InWorld)
 
 void UTedsActorFolderFactory::OnFolderMoved(UWorld& InWorld, const FFolder& InOldFolder, const FFolder& InNewFolder)
 {
-	UE::Editor::DataStorage::FMapKey OldIndex(static_cast<uint64>(GetTypeHash(InOldFolder)));
-	UE::Editor::DataStorage::RowHandle Row = DataStorage->LookupMappedRow(OldIndex);
+	using namespace UE::Editor::DataStorage;
 
-	// Remove the mapping to the old folder
-	DataStorage->RemoveRowMapping(OldIndex);
-	
-	// Map this row to the new folder information
-	UE::Editor::DataStorage::FMapKey NewIndex(static_cast<uint64>(GetTypeHash(InNewFolder)));
-	DataStorage->MapRow(NewIndex, Row);
+	FMapKey OldIndex(static_cast<uint64>(GetTypeHash(InOldFolder)));
+	RowHandle Row = DataStorage->LookupMappedRow(ActorFolders::MappingDomain, OldIndex);
+
+	// Remap this row to the new folder information
+	FMapKey NewIndex(static_cast<uint64>(GetTypeHash(InNewFolder)));
+	DataStorage->RemapRow(ActorFolders::MappingDomain, OldIndex, NewIndex);
 
 	// Update the default columns in TEDS as the data could have changed after the folder moved
 	SetFolderColumns(Row, InWorld, InNewFolder);
@@ -268,27 +267,29 @@ void UTedsActorFolderFactory::OnFolderMoved(UWorld& InWorld, const FFolder& InOl
 // Register the folder and fill in the columns with data (World, FFolder, Parent, ActorFolder etc)
 UE::Editor::DataStorage::RowHandle UTedsActorFolderFactory::RegisterFolder(UWorld& World, const FFolder& Folder)
 {
+	using namespace UE::Editor::DataStorage;
+
 	// Don't register invalid folders
 	if(Folder.IsNone())
 	{
 		return UE::Editor::DataStorage::InvalidRowHandle;
 	}
 
-	UE::Editor::DataStorage::FMapKey Index(static_cast<uint64>(GetTypeHash(Folder)));
-	UE::Editor::DataStorage::RowHandle Row = DataStorage->LookupMappedRow(Index);
+	FMapKey Index(static_cast<uint64>(GetTypeHash(Folder)));
+	RowHandle Row = DataStorage->LookupMappedRow(ActorFolders::MappingDomain, Index);
 	
 	if (!DataStorage->IsRowAvailable(Row))
 	{
 		// Add and index the row
-		static UE::Editor::DataStorage::RowHandle Table = DataStorage->FindTable(UE::TedsActorFolderFactory::Local::TableName);
+		static RowHandle Table = DataStorage->FindTable(UE::TedsActorFolderFactory::Local::TableName);
 
-		if(Table == UE::Editor::DataStorage::InvalidTableHandle)
+		if(Table == InvalidTableHandle)
 		{
-			return UE::Editor::DataStorage::InvalidRowHandle;
+			return InvalidRowHandle;
 		}
 	
 		Row = DataStorage->AddRow(Table);
-		DataStorage->MapRow(Index, Row);
+		DataStorage->MapRow(ActorFolders::MappingDomain, Index, Row);
 
 		SetFolderColumns(Row, World,  Folder);
 	}
@@ -326,8 +327,9 @@ void UTedsActorFolderFactory::SetFolderColumns(UE::Editor::DataStorage::RowHandl
 
 void UTedsActorFolderFactory::UnregisterFolder(const FFolder& Folder)
 {
-	UE::Editor::DataStorage::FMapKey Index(static_cast<uint64>(GetTypeHash(Folder)));
-	UE::Editor::DataStorage::RowHandle Row = DataStorage->LookupMappedRow(Index);
+	using namespace UE::Editor::DataStorage;
+	FMapKey Index(static_cast<uint64>(GetTypeHash(Folder)));
+	RowHandle Row = DataStorage->LookupMappedRow(ActorFolders::MappingDomain, Index);
 
 	if (DataStorage->IsRowAvailable(Row))
 	{
diff --git a/Engine/Plugins/Experimental/EditorDataStorageFeatures/Source/TedsOutliner/Public/Columns/ActorFolderColumns.h b/Engine/Plugins/Experimental/EditorDataStorageFeatures/Source/TedsOutliner/Public/Columns/ActorFolderColumns.h
--- a/Engine/Plugins/Experimental/EditorDataStorageFeatures/Source/TedsOutliner/Public/Columns/ActorFolderColumns.h
+++ b/Engine/Plugins/Experimental/EditorDataStorageFeatures/Source/TedsOutliner/Public/Columns/ActorFolderColumns.h
@@ -7,6 +7,11 @@
 
 #include "ActorFolderColumns.generated.h"
 
+namespace UE::Editor::DataStorage::ActorFolders
+{
+	inline static const FName MappingDomain = "ActorFolders";
+}
+
 /**
  * Column that stores a constructed FFolder
  */
diff --git a/Engine/Plugins/Experimental/EditorDataStorageFeatures/Source/TedsRevisionControl/Private/Queries/ObjectPackagePathToColumnQueries.cpp b/Engine/Plugins/Experimental/EditorDataStorageFeatures/Source/TedsRevisionControl/Private/Queries/ObjectPackagePathToColumnQueries.cpp
--- a/Engine/Plugins/Experimental/EditorDataStorageFeatures/Source/TedsRevisionControl/Private/Queries/ObjectPackagePathToColumnQueries.cpp
+++ b/Engine/Plugins/Experimental/EditorDataStorageFeatures/Source/TedsRevisionControl/Private/Queries/ObjectPackagePathToColumnQueries.cpp
@@ -4,6 +4,7 @@
 
 #include "Elements/Columns/TypedElementCompatibilityColumns.h"
 #include "Elements/Columns/TypedElementPackageColumns.h"
+#include "Elements/Columns/TypedElementRevisionControlColumns.h"
 #include "Elements/Framework/TypedElementQueryBuilder.h"
 #include "HAL/IConsoleManager.h"
 #include "Misc/PackageName.h"
@@ -65,7 +66,8 @@ void UTypedElementUObjectPackagePathFactory::RegisterQueries(UE::Editor::DataSto
 			FProcessor(EQueryTickPhase::FrameEnd, DataStorage.GetQueryTickGroupName(EQueryTickGroups::SyncExternalToDataStorage)),
 			[](IQueryContext& Context, RowHandle Row, const FTypedElementUObjectColumn& Object, const FTypedElementPackageUnresolvedReference& UnresolvedPackageReference)
 			{
-				RowHandle PackageRow = Context.LookupMappedRow(FMapKeyView(UnresolvedPackageReference.PathOnDisk));
+				RowHandle PackageRow = Context.LookupMappedRow(
+					UE::Editor::RevisionControl::MappingDomain, FMapKeyView(UnresolvedPackageReference.PathOnDisk));
 				if (!Context.IsRowAvailable(PackageRow))
 				{
 					return;
@@ -126,7 +128,7 @@ void UTypedElementUObjectPackagePathFactory::RegisterTryAddPackageRef(UE::Editor
 						FPaths::NormalizeFilename(PackageFilename);
 						FString FullPackageFilename = FPaths::ConvertRelativePathToFull(PackageFilename);
 						FMapKeyView Key = FMapKeyView(FullPackageFilename);
-						RowHandle PackageRow = Context.LookupMappedRow(Key);
+						RowHandle PackageRow = Context.LookupMappedRow(UE::Editor::RevisionControl::MappingDomain, Key);
 						if (Context.IsRowAvailable(PackageRow))
 						{
 							const UPackage* Package = ObjectInstance->GetPackage();
diff --git a/Engine/Plugins/Experimental/EditorDataStorageFeatures/Source/TedsSettings/Private/TedsSettingsManager.cpp b/Engine/Plugins/Experimental/EditorDataStorageFeatures/Source/TedsSettings/Private/TedsSettingsManager.cpp
--- a/Engine/Plugins/Experimental/EditorDataStorageFeatures/Source/TedsSettings/Private/TedsSettingsManager.cpp
+++ b/Engine/Plugins/Experimental/EditorDataStorageFeatures/Source/TedsSettings/Private/TedsSettingsManager.cpp
@@ -125,7 +125,7 @@ UE::Editor::DataStorage::RowHandle FTedsSettingsManager::FindSettingsSection(con
 
 			FMapKey SectionIndexKey = GenerateIndexKey(ContainerName, CategoryName, SectionName);
 
-			return DataStorage->LookupMappedRow(SectionIndexKey);
+			return DataStorage->LookupMappedRow(Settings::MappingDomain, SectionIndexKey);
 		}
 	}
 
@@ -147,7 +147,7 @@ UE::Editor::DataStorage::RowHandle FTedsSettingsManager::FindOrAddSettingsSectio
 
 			FMapKey SectionIndexKey = GenerateIndexKey(ContainerName, CategoryName, SectionName);
 
-			RowHandle SectionRow = DataStorage->LookupMappedRow(SectionIndexKey);
+			RowHandle SectionRow = DataStorage->LookupMappedRow(Settings::MappingDomain, SectionIndexKey);
 
 			if (SectionRow == InvalidRowHandle)
 			{
@@ -158,7 +158,7 @@ UE::Editor::DataStorage::RowHandle FTedsSettingsManager::FindOrAddSettingsSectio
 				DataStorage->AddColumn<FSettingsContainerReferenceColumn>(SectionRow, { .ContainerName = ContainerName, .ContainerRow = InvalidRowHandle });
 				DataStorage->AddColumn<FSettingsCategoryReferenceColumn>(SectionRow, { .CategoryName = CategoryName, .CategoryRow = InvalidRowHandle });
 
-				DataStorage->MapRow(MoveTemp(SectionIndexKey), SectionRow);
+				DataStorage->MapRow(Settings::MappingDomain, MoveTemp(SectionIndexKey), SectionRow);
 			}
 
 			return SectionRow;
@@ -291,7 +291,7 @@ void FTedsSettingsManager::RegisterSettingsContainer(const FName& ContainerName)
 	ISettingsContainerPtr ContainerPtr = SettingsModule->GetContainer(ContainerName);
 
 	FMapKey ContainerIndexKey = FMapKey(ContainerPtr.Get());
-	RowHandle ContainerRow = DataStorage->LookupMappedRow(ContainerIndexKey);
+	RowHandle ContainerRow = DataStorage->LookupMappedRow(Settings::MappingDomain, ContainerIndexKey);
 	if (ContainerRow == InvalidRowHandle)
 	{
 		ContainerRow = DataStorage->AddRow(SettingsContainerTable);
@@ -300,7 +300,7 @@ void FTedsSettingsManager::RegisterSettingsContainer(const FName& ContainerName)
 		DataStorage->AddColumn<FDescriptionColumn>(ContainerRow, { .Description = ContainerPtr->GetDescription() });
 		DataStorage->AddColumn<FSettingsContainerTag>(ContainerRow);
 
-		DataStorage->MapRow(MoveTemp(ContainerIndexKey), ContainerRow);
+		DataStorage->MapRow(Settings::MappingDomain, MoveTemp(ContainerIndexKey), ContainerRow);
 	}
 
 	TArray<ISettingsCategoryPtr> Categories;
@@ -374,7 +374,7 @@ void FTedsSettingsManager::UnregisterActiveSettings()
 			}
 
 			FMapKeyView CategoryIndexKey = FMapKeyView(CategoryPtr.Get());
-			RowHandle CategoryRow = DataStorage->LookupMappedRow(CategoryIndexKey);
+			RowHandle CategoryRow = DataStorage->LookupMappedRow(Settings::MappingDomain, CategoryIndexKey);
 			if (CategoryRow != InvalidRowHandle)
 			{
 				DataStorage->RemoveRow(CategoryRow);
@@ -382,7 +382,7 @@ void FTedsSettingsManager::UnregisterActiveSettings()
 		}
 
 		FMapKeyView ContainerIndexKey = FMapKeyView(ContainerPtr.Get());
-		RowHandle ContainerRow = DataStorage->LookupMappedRow(ContainerIndexKey);
+		RowHandle ContainerRow = DataStorage->LookupMappedRow(Settings::MappingDomain, ContainerIndexKey);
 		if (ContainerRow != InvalidRowHandle)
 		{
 			DataStorage->RemoveRow(ContainerRow);
@@ -434,7 +434,7 @@ void FTedsSettingsManager::UpdateSettingsCategory(TSharedPtr<ISettingsCategory>
 
 	FMapKey CategoryIndexKey = FMapKey(SettingsCategory.Get());
 
-	RowHandle CategoryRow = DataStorage->LookupMappedRow(CategoryIndexKey);
+	RowHandle CategoryRow = DataStorage->LookupMappedRow(Settings::MappingDomain, CategoryIndexKey);
 	if (CategoryRow == InvalidRowHandle)
 	{
 		CategoryRow = DataStorage->AddRow(SettingsCategoryTable);
@@ -445,7 +445,7 @@ void FTedsSettingsManager::UpdateSettingsCategory(TSharedPtr<ISettingsCategory>
 		DataStorage->AddColumn<FDescriptionColumn>(CategoryRow, { .Description = SettingsCategory->GetDescription() });
 		DataStorage->AddColumn<FSettingsCategoryTag>(CategoryRow);
 
-		DataStorage->MapRow(MoveTemp(CategoryIndexKey), CategoryRow);
+		DataStorage->MapRow(Settings::MappingDomain, MoveTemp(CategoryIndexKey), CategoryRow);
 	}
 
 	TArray<RowHandle> OldRowHandles;
@@ -499,7 +499,7 @@ void FTedsSettingsManager::UpdateSettingsCategory(TSharedPtr<ISettingsCategory>
 			TArray<void*> ColumnsToCopy;
 			TArray<const UScriptStruct*> ColumnTypesToCopy;
 
-			RowHandle OldSectionRow = DataStorage->LookupMappedRow(SectionIndexKey);
+			RowHandle OldSectionRow = DataStorage->LookupMappedRow(Settings::MappingDomain, SectionIndexKey);
 			if (OldSectionRow != InvalidRowHandle)
 			{
 				UE_LOG(LogTedsSettings, Verbose, TEXT("Settings Section : '%s' is already in data storage"), *SectionName.ToString());
@@ -537,7 +537,7 @@ void FTedsSettingsManager::UpdateSettingsCategory(TSharedPtr<ISettingsCategory>
 			UE_CLOG(OldSectionRow != NewSectionRow, LogTedsSettings, Verbose, TEXT("Added Settings Section : '%s'"), *SectionName.ToString());
 			UE_CLOG(OldSectionRow == NewSectionRow, LogTedsSettings, Verbose, TEXT("Updated Settings Section : '%s'"), *SectionName.ToString());
 
-			DataStorage->MapRow(MoveTemp(SectionIndexKey), NewSectionRow);			
+			DataStorage->MapRow(Settings::MappingDomain, MoveTemp(SectionIndexKey), NewSectionRow);
 		}
 	}
 
@@ -579,7 +579,7 @@ void FTedsSettingsManager::UpdateSettingsCategory(TSharedPtr<ISettingsCategory>
 
 		DataStorage->RemoveRow(OldSectionRow);
 
-		DataStorage->MapRow(GenerateIndexKey(ContainerName, CategoryName, OldSectionName), NewSectionRow);
+		DataStorage->MapRow(Settings::MappingDomain, GenerateIndexKey(ContainerName, CategoryName, OldSectionName), NewSectionRow);
 
 		UE_LOG(LogTedsSettings, Verbose, TEXT("Inactivated Settings Section : '%s'"), *OldSectionName.ToString());
 	}
diff --git a/Engine/Plugins/Experimental/EditorDataStorageFeatures/Source/TedsSettings/Public/TedsSettingsColumns.h b/Engine/Plugins/Experimental/EditorDataStorageFeatures/Source/TedsSettings/Public/TedsSettingsColumns.h
--- a/Engine/Plugins/Experimental/EditorDataStorageFeatures/Source/TedsSettings/Public/TedsSettingsColumns.h
+++ b/Engine/Plugins/Experimental/EditorDataStorageFeatures/Source/TedsSettings/Public/TedsSettingsColumns.h
@@ -7,6 +7,11 @@
 
 #include "TedsSettingsColumns.generated.h"
 
+namespace UE::Editor::DataStorage::Settings
+{
+	inline static const FName MappingDomain = "Settings";
+}
+
 USTRUCT(meta = (DisplayName = "Settings Container Reference"))
 struct FSettingsContainerReferenceColumn final : public FEditorDataStorageColumn
 {
diff --git a/Engine/Source/Runtime/TypedElementFramework/Public/Elements/Columns/TypedElementRevisionControlColumns.h b/Engine/Source/Runtime/TypedElementFramework/Public/Elements/Columns/TypedElementRevisionControlColumns.h
--- a/Engine/Source/Runtime/TypedElementFramework/Public/Elements/Columns/TypedElementRevisionControlColumns.h
+++ b/Engine/Source/Runtime/TypedElementFramework/Public/Elements/Columns/TypedElementRevisionControlColumns.h
@@ -6,6 +6,11 @@
 
 #include "TypedElementRevisionControlColumns.generated.h"
 
+namespace UE::Editor::RevisionControl
+{
+	inline static const FName MappingDomain = "SourceControl";
+}
+
 USTRUCT(meta = (DisplayName = "In a changelist"))
 struct FSCCInChangelistTag final : public FEditorDataStorageTag
 {
diff --git a/Engine/Source/Runtime/TypedElementFramework/Public/Elements/Interfaces/TypedElementDataStorageCompatibilityInterface.h b/Engine/Source/Runtime/TypedElementFramework/Public/Elements/Interfaces/TypedElementDataStorageCompatibilityInterface.h
--- a/Engine/Source/Runtime/TypedElementFramework/Public/Elements/Interfaces/TypedElementDataStorageCompatibilityInterface.h
+++ b/Engine/Source/Runtime/TypedElementFramework/Public/Elements/Interfaces/TypedElementDataStorageCompatibilityInterface.h
@@ -24,6 +24,8 @@ public:
 	using ObjectRegistrationFilter = TFunction<bool(const ICompatibilityProvider&, const UObject*)>;
 	using ObjectToRowDealiaser = TFunction<RowHandle(const ICompatibilityProvider&, const UObject*)>;
 	
+	inline static const FName ObjectMappingDomain = "Object";
+
 	/**
 	 * @section Type-agnostic functions
 	 * These allow compatibility with any type. These do eventually fall back to the explicit versions.
diff --git a/Engine/Source/Runtime/TypedElementFramework/Public/Elements/Interfaces/TypedElementDataStorageInterface.h b/Engine/Source/Runtime/TypedElementFramework/Public/Elements/Interfaces/TypedElementDataStorageInterface.h
--- a/Engine/Source/Runtime/TypedElementFramework/Public/Elements/Interfaces/TypedElementDataStorageInterface.h
+++ b/Engine/Source/Runtime/TypedElementFramework/Public/Elements/Interfaces/TypedElementDataStorageInterface.h
@@ -405,21 +405,25 @@ public:
 	 * one or more key to be created for a row for fast retrieval.
 	 */
 
-	/** Retrieves the row for a mapped object. Returns an invalid row handle if the no row with the provided key was found. */
-	virtual RowHandle LookupMappedRow(const FMapKeyView& Key) const = 0;
-	/** 
-	 * Registers a row under a key. The same row can be registered multiple, but an key can only be associated with a single row.
+	/**
+	 * Retrieves the row for a mapped object. Returns an invalid row handle if the no row with the provided key was found in the provided 
+	 * domain. 
 	 */
-	virtual void MapRow(FMapKey Key, RowHandle Row) = 0;
+	virtual RowHandle LookupMappedRow(const FName& Domain, const FMapKeyView& Key) const = 0;
 	/**
-	 * Register multiple rows under their key. The same row can be registered multiple times, but key can only be associated with a single
-	 * row. During processing the keys will be moved out into the final location.
+	 * Registers a row under a key in the provided domain. The same row can be registered multiple, but an key can only be associated with
+	 * a single row.
 	 */
-	virtual void BatchMapRows(TArrayView<TPair<FMapKey, RowHandle>> MapRowPairs) = 0;
-	/** Updates the key of a row to a new value. Effectively this is the same as removing a key and adding a new one. */
-	virtual void RemapRow(const FMapKeyView& OriginalKey, FMapKey NewKey) = 0;
-	/** Removes a previously registered key from the mapping table or does nothing if the key no longer exists. */
-	virtual void RemoveRowMapping(const FMapKeyView& Key) = 0;
+	virtual void MapRow(const FName& Domain, FMapKey Key, RowHandle Row) = 0;
+	/**
+	 * Register multiple rows under their key in the provided domain. The same row can be registered multiple times, but the key can only
+	 * be associated with a single row in the domain. During processing the keys will be moved out into the final location.
+	 */
+	virtual void BatchMapRows(const FName& Domain, TArrayView<TPair<FMapKey, RowHandle>> MapRowPairs) = 0;
+	/** Updates the key of a row in the provided domain to a new value. Effectively this is the same as removing a key and adding a new one. */
+	virtual void RemapRow(const FName& Domain, const FMapKeyView& OriginalKey, FMapKey NewKey) = 0;
+	/** Removes a previously registered key in the provided domain from the mapping table or does nothing if the key no longer exists. */
+	virtual void RemoveRowMapping(const FName& Domain, const FMapKeyView& Key) = 0;
 	
 	/**
 	 * @section Miscellaneous
@@ -457,7 +461,7 @@ public:
 	/**
 	 * @section Deprecated
 	 */
-	UE_DEPRECATED(5.6, "Use 'LookUpMappedRow' instead and use the new FMapKey(View) instead of the explicit 'GenerateIndexHash'-functions.")
+	/*UE_DEPRECATED(5.6, "Use 'LookUpMappedRow' instead and use the new FMapKey(View) instead of the explicit 'GenerateIndexHash'-functions.")
 	inline RowHandle FindIndexedRow(IndexHash Index) const;
 	UE_DEPRECATED(5.6, "Use 'MapRow' instead and use the new FMapKey(View) instead of the explicit 'GenerateIndexHash'-functions.")
 	inline void IndexRow(IndexHash Index, RowHandle Row);
@@ -467,9 +471,19 @@ public:
 	inline void Reindex(IndexHash OriginalIndex, IndexHash NewIndex);
 	UE_DEPRECATED(5.6, "Use 'RemovedRowMapping' instead and use the new FMapKey(View) instead of the explicit 'GenerateIndexHash'-functions.")
 	inline void RemoveIndex(IndexHash Index);
-
 	UE_DEPRECATED(5.6, "Use 'RemapRow' instead and use the new FMapKey(View) instead of the explicit 'GenerateIndexHash'-functions. The Row argument is no longer used.")
 	inline void ReindexRow(IndexHash OriginalIndex, IndexHash NewIndex, RowHandle Row);
+
+	UE_DEPRECATED(5.7, "Use the version of 'LookupMappedRow' that uses a domain.")
+	inline RowHandle LookupMappedRow(const FMapKeyView& Key) const;
+	UE_DEPRECATED(5.7, "Use the version of 'MapRow' that uses a domain.")
+	inline void MapRow(FMapKey Key, RowHandle Row);
+	UE_DEPRECATED(5.7, "Use the version of 'BatchMapRows' that uses a domain.")
+	inline void BatchMapRows(TArrayView<TPair<FMapKey, RowHandle>> MapRowPairs);
+	UE_DEPRECATED(5.7, "Use the version of 'RemapRow' that uses a domain.")
+	inline void RemapRow(const FMapKeyView& OriginalKey, FMapKey NewKey);
+	UE_DEPRECATED(5.7, "Use the version of 'RemoveRowMapping' that uses a domain.")
+	inline void RemoveRowMapping(const FMapKeyView& Key);*/
 };
 
 // Implementations
@@ -709,40 +723,65 @@ SystemType* ICoreProvider::GetExternalSystem()
 	return reinterpret_cast<SystemType*>(GetExternalSystemAddress(SystemType::StaticClass()));
 }
 
-void ICoreProvider::ReindexRow(
+/*void ICoreProvider::ReindexRow(
 	IndexHash OriginalIndex, IndexHash NewIndex, RowHandle Row)
 {
-	RemapRow(FMapKeyView(OriginalIndex), FMapKey(NewIndex));
+	RemapRow(NAME_None, FMapKeyView(OriginalIndex), FMapKey(NewIndex));
 }
 
 RowHandle ICoreProvider::FindIndexedRow(IndexHash Index) const
 {
-	return LookupMappedRow(FMapKeyView(Index));
+	return LookupMappedRow(NAME_None, FMapKeyView(Index));
 }
 
 void ICoreProvider::IndexRow(IndexHash Index, RowHandle Row)
 {
-	MapRow(FMapKey(Index), Row);
+	MapRow(NAME_None, FMapKey(Index), Row);
 }
 
 void ICoreProvider::BatchIndexRows(TConstArrayView<TPair<IndexHash, RowHandle>> IndexRowPairs)
 {
 	for (const TPair<IndexHash, RowHandle>& Pair : IndexRowPairs)
 	{
-		MapRow(FMapKey(Pair.Key), Pair.Value);
+		MapRow(NAME_None, FMapKey(Pair.Key), Pair.Value);
 	}
 }
 
 void ICoreProvider::Reindex(IndexHash OriginalIndex, IndexHash NewIndex)
 {
-	RemapRow(FMapKeyView(OriginalIndex), FMapKey(NewIndex));
+	RemapRow(NAME_None, FMapKeyView(OriginalIndex), FMapKey(NewIndex));
 }
 
 void ICoreProvider::RemoveIndex(IndexHash Index)
 {
-	RemoveRowMapping(FMapKeyView(Index));
+	RemoveRowMapping(NAME_None, FMapKeyView(Index));
+}
+
+RowHandle ICoreProvider::LookupMappedRow(const FMapKeyView& Key) const
+{
+	return LookupMappedRow(NAME_None, Key);
+}
+
+void ICoreProvider::MapRow(FMapKey Key, RowHandle Row)
+{
+	MapRow(NAME_None, Key, Row);
+}
+
+void ICoreProvider::BatchMapRows(TArrayView<TPair<FMapKey, RowHandle>> MapRowPairs)
+{
+	BatchMapRows(NAME_None, MapRowPairs);
 }
 
+void ICoreProvider::RemapRow(const FMapKeyView& OriginalKey, FMapKey NewKey)
+{
+	RemapRow(NAME_None, OriginalKey, NewKey);
+}
+
+void ICoreProvider::RemoveRowMapping(const FMapKeyView& Key)
+{
+	RemoveRowMapping(NAME_None, Key);
+}*/
+
 template<TDynamicColumnTemplate DynamicColumnTemplateType>
 void ICoreProvider::ForEachDynamicColumn(TFunctionRef<void(const UScriptStruct& Type)> Callback) const
 {
diff --git a/Engine/Source/Runtime/TypedElementFramework/Public/Elements/Interfaces/TypedElementDataStorageUiInterface.h b/Engine/Source/Runtime/TypedElementFramework/Public/Elements/Interfaces/TypedElementDataStorageUiInterface.h
--- a/Engine/Source/Runtime/TypedElementFramework/Public/Elements/Interfaces/TypedElementDataStorageUiInterface.h
+++ b/Engine/Source/Runtime/TypedElementFramework/Public/Elements/Interfaces/TypedElementDataStorageUiInterface.h
@@ -266,6 +266,8 @@ namespace UE::Editor::DataStorage
 	class IUiProvider : public IModularFeature
 	{
 	public:
+		inline static const FName PurposeMappingDomain = "WidgetPurpose";
+
 		enum class EPurposeType : uint8
 		{
 			/** General purpose name which allows multiple factory registrations. */
diff --git a/Engine/Source/Runtime/TypedElementFramework/Public/Elements/Interfaces/TypedElementQueryStorageInterfaces.h b/Engine/Source/Runtime/TypedElementFramework/Public/Elements/Interfaces/TypedElementQueryStorageInterfaces.h
--- a/Engine/Source/Runtime/TypedElementFramework/Public/Elements/Interfaces/TypedElementQueryStorageInterfaces.h
+++ b/Engine/Source/Runtime/TypedElementFramework/Public/Elements/Interfaces/TypedElementQueryStorageInterfaces.h
@@ -363,8 +363,8 @@ namespace UE::Editor::DataStorage
 		virtual void GetDependencies(TArrayView<UObject*> RetrievedAddresses, TConstArrayView<TWeakObjectPtr<const UClass>> DependencyTypes,
 			TConstArrayView<EQueryAccessType> AccessTypes) = 0;
 
-		/** Retrieves the row for a mapped object. Returns an invalid row handle if the key wasn't found. */
-		virtual RowHandle LookupMappedRow(const FMapKeyView& Index) const = 0;
+		/** Retrieves the row for a mapped object in the provided domain. Returns an invalid row handle if the key wasn't found. */
+		virtual RowHandle LookupMappedRow(const FName& Domain, const FMapKeyView& Index) const = 0;
 
 		/**
 		 * Runs a previously created query. This version takes an arbitrary query, but is limited to running queries that do not directly
@@ -390,8 +390,10 @@ namespace UE::Editor::DataStorage
 		 */
 		virtual FQueryResult RunSubquery(int32 SubqueryIndex, RowHandle Row, SubqueryCallbackRef Callback) = 0;
 
-		UE_DEPRECATED(5.6, "Use 'LookUpMappedRow' instead and use the new FMapKey(View) instead of the explicit 'GenerateIndexHash'-functions.")
+		/*UE_DEPRECATED(5.6, "Use 'LookUpMappedRow' instead and use the new FMapKey(View) instead of the explicit 'GenerateIndexHash'-functions.")
 		inline RowHandle FindIndexedRow(IndexHash Index) const;
+		UE_DEPRECATED(5.7, "Use the version of this 'LookupMappedRow' that uses a domain.")
+		inline RowHandle LookupMappedRow(const FMapKeyView& Index) const;*/
 	};
 } // namespace UE::Editor::DataStorage
 
@@ -694,9 +696,14 @@ namespace UE::Editor::DataStorage
 	// IQueryContext
 	//
 
-	RowHandle IQueryContext::FindIndexedRow(IndexHash Index) const
+	/*RowHandle IQueryContext::FindIndexedRow(IndexHash Index) const
 	{
-		return LookupMappedRow(FMapKeyView(Index));
+		return LookupMappedRow(NAME_None, FMapKeyView(Index));
 	}
 
+	RowHandle IQueryContext::LookupMappedRow(const FMapKeyView& Index) const
+	{
+		return LookupMappedRow(NAME_None, FMapKeyView(Index));
+	}*/
+
 } // namespace UE::Editor::DataStorage
