UE: Fix shadowed variables

diff --git a/Engine/Plugins/Experimental/EditorDataStorage/Source/TedsCore/Private/TypedElementDatabaseIndexTable.cpp b/Engine/Plugins/Experimental/EditorDataStorage/Source/TedsCore/Private/TypedElementDatabaseIndexTable.cpp
--- a/Engine/Plugins/Experimental/EditorDataStorage/Source/TedsCore/Private/TypedElementDatabaseIndexTable.cpp
+++ b/Engine/Plugins/Experimental/EditorDataStorage/Source/TedsCore/Private/TypedElementDatabaseIndexTable.cpp
@@ -117,7 +117,7 @@ namespace UE::Editor::DataStorage
 	}
 
 	FMappingTable::FDomain::ERemoveInvalidRowsResult FMappingTable::FDomain::RemoveInvalidRows(
-		const FName& DomainName, UEditorDataStorage& DataStorage, FTimespan& RemainingTime)
+		const FName& DomainName, UEditorDataStorage& InDataStorage, FTimespan& RemainingTime)
 	{
 		FCleanUpInfo& Info = *CleanUpInfo;
 
@@ -131,9 +131,9 @@ namespace UE::Editor::DataStorage
 			if (Info.RemainingJobs > 0)
 			{
 				int32 NumBatches = FMath::Min(Info.RemainingJobs, Info.MaxNumJobs);
-				ParallelForTemplate(NumBatches, [this, &DataStorage](int32 Block)
+				ParallelForTemplate(NumBatches, [this, &InDataStorage](int32 Block)
 					{
-						InspectRowBlockForCleanUp(DataStorage, Block);
+						InspectRowBlockForCleanUp(InDataStorage, Block);
 					}, EParallelForFlags::Unbalanced);
 
 				Info.RemainingJobs -= NumBatches;
@@ -267,12 +267,12 @@ namespace UE::Editor::DataStorage
 		return -1;
 	}
 
-	void FMappingTable::FDomain::IndexRowUnguarded(UEditorDataStorage& DataStorage, FMapKey&& Key, RowHandle Row, bool bIsDirty)
+	void FMappingTable::FDomain::IndexRowUnguarded(UEditorDataStorage& InDataStorage, FMapKey&& Key, RowHandle Row, bool bIsDirty)
 	{
 		uint64 Hash = Key.CalculateHash();
 		if (int32 Index = FindIndexUnguarded(Hash, Key); Index >= 0)
 		{
-			bool bUpdatingAllowed = bIsDirty && !DataStorage.IsRowAvailableUnsafe(Rows[Index]);
+			bool bUpdatingAllowed = bIsDirty && !InDataStorage.IsRowAvailableUnsafe(Rows[Index]);
 			if (ensureMsgf(bUpdatingAllowed || Row == Rows[Index], TEXT("Another row has already been registered under key '%s'."), *Key.ToString()))
 			{
 				Rows[Index] = Row; // Update the stored row to the new row.
@@ -297,7 +297,7 @@ namespace UE::Editor::DataStorage
 		IndexLookupMap.Add(Hash, RowIndex);
 	}
 
-	void FMappingTable::FDomain::InspectRowBlockForCleanUp(UEditorDataStorage& DataStorage, int32 Block)
+	void FMappingTable::FDomain::InspectRowBlockForCleanUp(UEditorDataStorage& InDataStorage, int32 Block)
 	{
 		TEDS_EVENT_SCOPE("Mapping inspect rows");
 
@@ -314,7 +314,7 @@ namespace UE::Editor::DataStorage
 		for (; Front != End; ++Front)
 		{
 			// Skip rows with an invalid handle as those are already freed.
-			if (*Front != InvalidRowHandle && !DataStorage.IsRowAvailableUnsafe(*Front))
+			if (*Front != InvalidRowHandle && !InDataStorage.IsRowAvailableUnsafe(*Front))
 			{
 				Result.Add(static_cast<int32>(Front - Begin));
 			}
