Cleanup TEDS TypedElementBridge on shutdown

diff --git a/Engine/Plugins/Experimental/EditorDataStorage/Source/TedsCore/Private/Compatibility/TedsTypedElementBridgeQueries.cpp b/Engine/Plugins/Experimental/EditorDataStorage/Source/TedsCore/Private/Compatibility/TedsTypedElementBridgeQueries.cpp
--- a/Engine/Plugins/Experimental/EditorDataStorage/Source/TedsCore/Private/Compatibility/TedsTypedElementBridgeQueries.cpp
+++ b/Engine/Plugins/Experimental/EditorDataStorage/Source/TedsCore/Private/Compatibility/TedsTypedElementBridgeQueries.cpp
@@ -68,6 +68,7 @@ void UTypedElementBridgeDataStorageFactory::RegisterQuery_NewUObject(UE::Editor:
 
 void UTypedElementBridgeDataStorageFactory::UnregisterQuery_NewUObject(UE::Editor::DataStorage::ICoreProvider& DataStorage)
 {
+	DataStorage.UnregisterQuery(RemoveTypedElementRowHandleQuery);
 }
 
 void UTypedElementBridgeDataStorageFactory::CleanupTypedElementColumns(UE::Editor::DataStorage::ICoreProvider& DataStorage)
