[FNameColumn] Add a new column to TEDS WidgetPurposes to represent the name column and replace uses of the FNameColumn

diff --git a/Engine/Plugins/Experimental/EditorDataStorage/Source/TedsCore/Private/TypedElementDatabaseUI.cpp b/Engine/Plugins/Experimental/EditorDataStorage/Source/TedsCore/Private/TypedElementDatabaseUI.cpp
--- a/Engine/Plugins/Experimental/EditorDataStorage/Source/TedsCore/Private/TypedElementDatabaseUI.cpp
+++ b/Engine/Plugins/Experimental/EditorDataStorage/Source/TedsCore/Private/TypedElementDatabaseUI.cpp
@@ -706,10 +706,10 @@ void UEditorDataStorageUi::ListWidgetPurposes(const WidgetPurposeCallback& Callb
 	using namespace UE::Editor::DataStorage::Queries;
 	
 	static QueryHandle PurposeQueryHandle =
-		Storage->RegisterQuery(Select().ReadOnly<FNameColumn, FWidgetPurposeColumn, FDescriptionColumn>().Compile());
+		Storage->RegisterQuery(Select().ReadOnly<FWidgetConstructorNameColumn, FWidgetPurposeColumn, FDescriptionColumn>().Compile());
 
 	Storage->RunQuery(PurposeQueryHandle, CreateDirectQueryCallbackBinding(
-		[Callback](const FNameColumn& NameColumn, const FWidgetPurposeColumn& PurposeColumn, const FDescriptionColumn& DescriptionColumn)
+		[Callback](const FWidgetConstructorNameColumn& NameColumn, const FWidgetPurposeColumn& PurposeColumn, const FDescriptionColumn& DescriptionColumn)
 	{
 			Callback(NameColumn.Name, PurposeColumn.PurposeType, DescriptionColumn.Description);
 	}));
@@ -813,11 +813,11 @@ void UEditorDataStorageUi::RegisterQueries()
 	
 	Storage->RegisterQuery(
 		Select(
-			TEXT("Add display name to widget factory with constructor"),
+			TEXT("Add the name to widget factory with constructor"),
 			FObserver::OnAdd<FWidgetFactoryConstructorTypeInfoColumn>().SetExecutionMode(EExecutionMode::GameThread),
 			[](IQueryContext& Context, RowHandle Row, const FWidgetFactoryConstructorTypeInfoColumn& ConstructorColumn)
 			{
-				Context.AddColumn(Row, FNameColumn{ .Name = ConstructorColumn.Constructor->GetFName() });
+				Context.AddColumn(Row, FWidgetConstructorNameColumn{ .Name = ConstructorColumn.Constructor->GetFName() });
 			})
 		.Where()
 			.All<FWidgetFactoryColumn>()
@@ -825,11 +825,11 @@ void UEditorDataStorageUi::RegisterQueries()
 	
 	Storage->RegisterQuery(
 		Select(
-			TEXT("Add display name to widget factory with constructor type info"),
+			TEXT("Add the name to widget factory with constructor type info"),
 			FObserver::OnAdd<FWidgetFactoryConstructorColumn>().SetExecutionMode(EExecutionMode::GameThread),
 			[this](IQueryContext& Context, RowHandle Row, const FWidgetFactoryConstructorColumn& ConstructorColumn)
 			{
-				Context.AddColumn(Row, FNameColumn{ .Name = ConstructorColumn.Constructor->GetTypeInfo()->GetFName() });
+				Context.AddColumn(Row, FWidgetConstructorNameColumn{ .Name = ConstructorColumn.Constructor->GetTypeInfo()->GetFName() });
 			})
 		.Where()
 			.All<FWidgetFactoryColumn>()
diff --git a/Engine/Source/Runtime/TypedElementFramework/Public/Elements/Columns/WidgetPurposeColumns.h b/Engine/Source/Runtime/TypedElementFramework/Public/Elements/Columns/WidgetPurposeColumns.h
--- a/Engine/Source/Runtime/TypedElementFramework/Public/Elements/Columns/WidgetPurposeColumns.h
+++ b/Engine/Source/Runtime/TypedElementFramework/Public/Elements/Columns/WidgetPurposeColumns.h
@@ -35,7 +35,6 @@ struct FWidgetPurposeNameColumn : public FEditorDataStorageColumn
 	FName Frame;
 };
 
-
 /** Column to store info about a widget factory */
 USTRUCT()
 struct FWidgetFactoryColumn : public FEditorDataStorageColumn
@@ -66,6 +65,16 @@ struct FWidgetFactoryConstructorColumn : public FEditorDataStorageColumn
 	TUniquePtr<FTypedElementWidgetConstructor> Constructor;
 };
 
+/** Column used to store the name of the widget constructor */
+USTRUCT()
+struct FWidgetConstructorNameColumn : public FEditorDataStorageColumn
+{
+	GENERATED_BODY()
+
+	UPROPERTY(meta = (Searchable, Sortable))
+	FName Name;
+};
+
 template<>
 struct TStructOpsTypeTraits<FWidgetFactoryConstructorColumn> : public TStructOpsTypeTraitsBase2<FWidgetFactoryConstructorColumn>
 {
