Convert TEDS activatable queries to use Mass activatable processors
- Remove custom activation infrastructure and forward activation calls to the internal processor for each query.
- Activation counts are no longer supported, since Mass processors don't have this feature. Nothing is using it.
- Remove support for phase amble queries. Existing queries moved to processors. Phase amble queries were implemented through pre and post-phase callbacks. These callbacks didn't use a processor and therefore couldn't benefit from activatable processors. They're also single threaded and shouldn't be necessary given sufficient control over processor ordering and command buffer flushing so it makes sense to avoid using them entirely.

diff --git a/Engine/Plugins/Experimental/EditorDataStorage/Source/TedsCore/Private/Processors/TypedElementProcessorAdaptors.cpp b/Engine/Plugins/Experimental/EditorDataStorage/Source/TedsCore/Private/Processors/TypedElementProcessorAdaptors.cpp
--- a/Engine/Plugins/Experimental/EditorDataStorage/Source/TedsCore/Private/Processors/TypedElementProcessorAdaptors.cpp
+++ b/Engine/Plugins/Experimental/EditorDataStorage/Source/TedsCore/Private/Processors/TypedElementProcessorAdaptors.cpp
@@ -2,18 +2,15 @@
 
 #include "Processors/TypedElementProcessorAdaptors.h"
 
-#include <utility>
-#include "Algo/Copy.h"
 #include "DataStorage/Queries/Types.h"
 #include "MassEntityTypes.h"
 #include "MassEntityView.h"
 #include "MassExecutionContext.h"
 #include "Queries/TypedElementExtendedQueryStore.h"
+#include "StructUtils/StructTypeBitSet.h"
 #include "TypedElementDatabase.h"
 #include "TypedElementDatabaseEnvironment.h"
 #include "TypedElementUtils.h"
-#include "Algo/Copy.h"
-#include "StructUtils/StructTypeBitSet.h"
 
 namespace UE::Editor::DataStorage
 {
@@ -1324,7 +1321,7 @@ protected:
 		FMassEntityQuery& NativeQuery,
 		QueryCallbackRef Callback)
 	{
-		if (Description.Callback.ActivationCount > 0)
+		if (Description.Callback.Active)
 		{
 			NativeQuery.ForEachEntityChunk(Context,
 				[&Callback, &QueryStore, &Environment, &Description](FMassExecutionContext& ExecutionContext)
@@ -1544,28 +1541,25 @@ UE::Editor::DataStorage::FQueryResult FTypedElementQueryProcessorData::Execute(
 	
 	if (EnumHasAnyFlags(ExecutionFlags, EDirectQueryExecutionFlags::AllowBoundQueries) || !Description.Callback.Function)
 	{
-		if (EnumHasAnyFlags(ExecutionFlags, EDirectQueryExecutionFlags::IgnoreActivationCount) || Description.Callback.ActivationCount > 0)
-		{
-			FMassExecutionContext Context(EntityManager);
-			auto ExecuteFunction = [&Result, &Callback, &Description, &Environment](FMassExecutionContext& Context)
-				{
-					// No need to cache any subsystem dependencies as these are not accessible from a direct query.
-					UE::Editor::DataStorage::Processors::Private::FMassDirectContextForwarder QueryContext(Context, Environment);
-					Callback(Description, QueryContext);
-					Result.Count += Context.GetNumEntities();
-				};
-			if (EnumHasAnyFlags(ExecutionFlags, EDirectQueryExecutionFlags::ParallelizeChunks))
-			{
-				FMassEntityQuery::EParallelExecutionFlags Mode =
-					EnumHasAnyFlags(ExecutionFlags, EDirectQueryExecutionFlags::AutoBalanceParallelChunkProcessing)
-					? FMassEntityQuery::EParallelExecutionFlags::AutoBalance
-					: FMassEntityQuery::EParallelExecutionFlags::Default;
-				NativeQuery.ParallelForEachEntityChunk(Context, ExecuteFunction, Mode);
-			}
-			else
+		FMassExecutionContext Context(EntityManager);
+		auto ExecuteFunction = [&Result, &Callback, &Description, &Environment](FMassExecutionContext& Context)
 			{
-				NativeQuery.ForEachEntityChunk(Context, ExecuteFunction);
-			}
+				// No need to cache any subsystem dependencies as these are not accessible from a direct query.
+				UE::Editor::DataStorage::Processors::Private::FMassDirectContextForwarder QueryContext(Context, Environment);
+				Callback(Description, QueryContext);
+				Result.Count += Context.GetNumEntities();
+			};
+		if (EnumHasAnyFlags(ExecutionFlags, EDirectQueryExecutionFlags::ParallelizeChunks))
+		{
+			FMassEntityQuery::EParallelExecutionFlags Mode =
+				EnumHasAnyFlags(ExecutionFlags, EDirectQueryExecutionFlags::AutoBalanceParallelChunkProcessing)
+				? FMassEntityQuery::EParallelExecutionFlags::AutoBalance
+				: FMassEntityQuery::EParallelExecutionFlags::Default;
+			NativeQuery.ParallelForEachEntityChunk(Context, ExecuteFunction, Mode);
+		}
+		else
+		{
+			NativeQuery.ForEachEntityChunk(Context, ExecuteFunction);
 		}
 	}
 	else
@@ -1588,25 +1582,22 @@ UE::Editor::DataStorage::FQueryResult FTypedElementQueryProcessorData::Execute(
 	FQueryResult Result;
 	Result.Completed = FQueryResult::ECompletion::Fully;
 
-	if (Description.Callback.ActivationCount > 0)
-	{
-		checkf(Description.Callback.ExecutionMode != EExecutionMode::ThreadedChunks,
-			TEXT("TEDS Sub-queries do not support parallel chunk processing."));
-		
-		FMassExecutionContext Context(EntityManager);
-		Context.SetDeferredCommandBuffer(ParentContext.GetSharedDeferredCommandBuffer());
-		Context.SetFlushDeferredCommands(false);
+	checkf(Description.Callback.ExecutionMode != EExecutionMode::ThreadedChunks,
+		TEXT("TEDS Sub-queries do not support parallel chunk processing."));
 
-		NativeQuery.ForEachEntityChunk(Context,
-			[&Result, &Callback, &Description, &Environment](FMassExecutionContext& Context)
-			{
-				// No need to cache any subsystem dependencies as these are not accessible from a subquery.
-				UE::Editor::DataStorage::Processors::Private::FMassSubqueryContextForwarder QueryContext(Context, Environment, Description.bShouldBatchModifications);
-				Callback(Description, QueryContext);
-				Result.Count += Context.GetNumEntities();
-			}
-		);
-	}
+	FMassExecutionContext Context(EntityManager);
+	Context.SetDeferredCommandBuffer(ParentContext.GetSharedDeferredCommandBuffer());
+	Context.SetFlushDeferredCommands(false);
+
+	NativeQuery.ForEachEntityChunk(Context,
+		[&Result, &Callback, &Description, &Environment](FMassExecutionContext& Context)
+		{
+			// No need to cache any subsystem dependencies as these are not accessible from a subquery.
+			UE::Editor::DataStorage::Processors::Private::FMassSubqueryContextForwarder QueryContext(Context, Environment, Description.bShouldBatchModifications);
+			Callback(Description, QueryContext);
+			Result.Count += Context.GetNumEntities();
+		}
+	);
 	return Result;
 }
 
@@ -1625,7 +1616,7 @@ UE::Editor::DataStorage::FQueryResult FTypedElementQueryProcessorData::Execute(
 	Result.Completed = FQueryResult::ECompletion::Fully;
 
 	FMassEntityHandle NativeEntity = FMassEntityHandle::FromNumber(RowHandle);
-	if (Description.Callback.ActivationCount > 0 && EntityManager.IsEntityActive(NativeEntity))
+	if (EntityManager.IsEntityActive(NativeEntity))
 	{
 		checkf(Description.Callback.ExecutionMode != EExecutionMode::ThreadedChunks,
 			TEXT("TEDS Sub-queries do not support parallel chunk processing."));
@@ -1659,25 +1650,22 @@ void FTypedElementQueryProcessorData::Execute(FMassEntityManager& EntityManager,
 	checkf(StoredQuery, TEXT("A query callback was registered for execution without an associated query. Processor: [%s]"), *GetProcessorName());
 	
 	FQueryDescription& Description = StoredQuery->Description;
-	if (Description.Callback.ActivationCount > 0)
-	{
-		auto ExceteFunction = [this, &Description](FMassExecutionContext& Context)
-			{
-				if (PrepareCachedDependenciesOnQuery(Description, Context))
-				{
-					Processors::Private::FMassContextForwarder QueryContext(Description, Context, *QueryStore, *Environment);
-					Description.Callback.Function(Description, QueryContext);
-				}
-			};
-		
-		if (StoredQuery->Description.Callback.ExecutionMode != EExecutionMode::ThreadedChunks)
-		{
-			NativeQuery.ForEachEntityChunk(Context, ExceteFunction);
-		}
-		else
+	auto ExecuteFunction = [this, &Description](FMassExecutionContext& Context)
 		{
-			NativeQuery.ParallelForEachEntityChunk(Context, ExceteFunction);
-		}
+			if (PrepareCachedDependenciesOnQuery(Description, Context))
+			{
+				Processors::Private::FMassContextForwarder QueryContext(Description, Context, *QueryStore, *Environment);
+				Description.Callback.Function(Description, QueryContext);
+			}
+		};
+
+	if (StoredQuery->Description.Callback.ExecutionMode != EExecutionMode::ThreadedChunks)
+	{
+		NativeQuery.ForEachEntityChunk(Context, ExecuteFunction);
+	}
+	else
+	{
+		NativeQuery.ParallelForEachEntityChunk(Context, ExecuteFunction);
 	}
 }
 
diff --git a/Engine/Plugins/Experimental/EditorDataStorage/Source/TedsCore/Private/Queries/TypedElementExtendedQueryStore.cpp b/Engine/Plugins/Experimental/EditorDataStorage/Source/TedsCore/Private/Queries/TypedElementExtendedQueryStore.cpp
--- a/Engine/Plugins/Experimental/EditorDataStorage/Source/TedsCore/Private/Queries/TypedElementExtendedQueryStore.cpp
+++ b/Engine/Plugins/Experimental/EditorDataStorage/Source/TedsCore/Private/Queries/TypedElementExtendedQueryStore.cpp
@@ -5,7 +5,6 @@
 #include "MassProcessingPhaseManager.h"
 #include "MassProcessor.h"
 #include "TypedElementDatabaseEnvironment.h"
-#include "Commands/EditorDataStorageCommandBuffer.h"
 #include "DataStorage/Debug/Log.h"
 #include "Misc/OutputDevice.h"
 #include "Processors/TypedElementProcessorAdaptors.h"
@@ -156,14 +155,14 @@ namespace UE::Editor::DataStorage
 
 		FMassEntityQuery& NativeQuery = SetupNativeQuery(StoredQuery.Description, StoredQuery, Environment);
 
-		bool bContinueSetup =               SetupDynamicColumns(StoredQuery.Description, Environment);
-		bContinueSetup = bContinueSetup &&	SetupSelectedColumns(StoredQuery.Description, NativeQuery);
+		bool bContinueSetup =              SetupDynamicColumns(StoredQuery.Description, Environment);
+		bContinueSetup = bContinueSetup && SetupSelectedColumns(StoredQuery.Description, NativeQuery);
 		bContinueSetup = bContinueSetup && SetupChunkFilters(Result, StoredQuery.Description, Environment, NativeQuery);
 		bContinueSetup = bContinueSetup && SetupConditions(StoredQuery.Description, Environment, NativeQuery);
 		bContinueSetup = bContinueSetup && SetupDependencies(StoredQuery.Description, NativeQuery);
 		bContinueSetup = bContinueSetup && SetupTickGroupDefaults(StoredQuery.Description);
 		bContinueSetup = bContinueSetup && SetupProcessors(Result, StoredQuery, Environment, EntityManager, PhaseManager);
-		bContinueSetup = bContinueSetup && SetupActivatable(Result, Query);
+		bContinueSetup = bContinueSetup && SetupActivatable(Result, StoredQuery);
 
 		if (!bContinueSetup)
 		{
@@ -286,9 +285,7 @@ namespace UE::Editor::DataStorage
 		for (Handle Query : ActiveActivatables)
 		{
 			FExtendedQuery& QueryData = Queries.GetMutable(Query);
-			checkf(QueryData.Description.Callback.ActivationCount > 0,
-				TEXT("Attempting to decrement the query '%s' which is already at zero."), *QueryData.Description.Callback.Name.ToString());
-			QueryData.Description.Callback.ActivationCount--;
+			QueryData.Description.Callback.Active = false;
 		}
 		ActiveActivatables.Reset();
 
@@ -296,10 +293,19 @@ namespace UE::Editor::DataStorage
 		for (Handle Query : PendingActivatables)
 		{
 			FExtendedQuery& QueryData = Queries.GetMutable(Query);
-			if (QueryData.Description.Callback.ActivationCount == 0)
+			if (!QueryData.Description.Callback.Active)
 			{
-				QueryData.Description.Callback.ActivationCount = 1;
+				QueryData.Description.Callback.Active = true;
 				ActiveActivatables.Add(Query);
+
+				if (QueryData.Description.Callback.Type != EQueryCallbackType::PhasePreparation &&
+					QueryData.Description.Callback.Type != EQueryCallbackType::PhaseFinalization)
+				{
+					if (ensure(QueryData.Processor.IsValid()))
+					{
+						QueryData.Processor->MakeOneShot();
+					}
+				}
 			}
 		}
 		PendingActivatables.Reset();
@@ -1093,11 +1099,20 @@ namespace UE::Editor::DataStorage
 		return true;
 	}
 
-	bool FExtendedQueryStore::SetupActivatable(Handle QueryHandle, FQueryDescription& Query)
+	bool FExtendedQueryStore::SetupActivatable(Handle QueryHandle, FExtendedQuery& StoredQuery)
 	{
-		if (!Query.Callback.ActivationName.IsNone())
+		if (!StoredQuery.Description.Callback.ActivationName.IsNone())
 		{
-			ActivatableMapping.Add(Query.Callback.ActivationName, QueryHandle);
+			ActivatableMapping.Add(StoredQuery.Description.Callback.ActivationName, QueryHandle);
+
+			if (StoredQuery.Description.Callback.Type != EQueryCallbackType::PhasePreparation &&
+				StoredQuery.Description.Callback.Type != EQueryCallbackType::PhaseFinalization)
+			{
+				if (ensure(StoredQuery.Processor.IsValid()))
+				{
+					StoredQuery.Processor->MakeInactive();
+				}
+			}
 		}
 		return true;
 	}
diff --git a/Engine/Plugins/Experimental/EditorDataStorage/Source/TedsCore/Private/Queries/TypedElementExtendedQueryStore.h b/Engine/Plugins/Experimental/EditorDataStorage/Source/TedsCore/Private/Queries/TypedElementExtendedQueryStore.h
--- a/Engine/Plugins/Experimental/EditorDataStorage/Source/TedsCore/Private/Queries/TypedElementExtendedQueryStore.h
+++ b/Engine/Plugins/Experimental/EditorDataStorage/Source/TedsCore/Private/Queries/TypedElementExtendedQueryStore.h
@@ -181,7 +181,7 @@ namespace UE::Editor::DataStorage
 			FEnvironment& Environment,
 			FMassExecutionContext* ParentContext,
 			Handle Query,
-		EDirectQueryExecutionFlags ExecutionFlags,
+			EDirectQueryExecutionFlags ExecutionFlags,
 			CallbackReference Callback);
 
 		FMassEntityQuery& SetupNativeQuery(FQueryDescription& Query, FExtendedQuery& StoredQuery, FEnvironment& Environment);
@@ -193,7 +193,7 @@ namespace UE::Editor::DataStorage
 		bool SetupTickGroupDefaults(FQueryDescription& Query);
 		bool SetupProcessors(Handle QueryHandle, FExtendedQuery& StoredQuery, FEnvironment& Environment,
 			FMassEntityManager& EntityManager, FMassProcessingPhaseManager& PhaseManager);
-		bool SetupActivatable(Handle QueryHandle, FQueryDescription& Query);
+		bool SetupActivatable(Handle QueryHandle, FExtendedQuery& StoredQuery);
 
 		EMassFragmentAccess ConvertToNativeAccessType(EQueryAccessType AccessType);
 		EMassFragmentPresence ConvertToNativePresenceType(EQueryAccessType AccessType);
@@ -219,4 +219,4 @@ namespace UE::Editor::DataStorage
 		TArray<Handle> ActiveActivatables;
 		FDynamicColumnGenerator* DynamicColumnGenerator;
 	};
-} // namepsace UE::Editor::DataStorage
+} // namespace UE::Editor::DataStorage
diff --git a/Engine/Plugins/Experimental/EditorDataStorageFeatures/Source/TedsQueryStack/Private/TedsRowQueryResultsNode.cpp b/Engine/Plugins/Experimental/EditorDataStorageFeatures/Source/TedsQueryStack/Private/TedsRowQueryResultsNode.cpp
--- a/Engine/Plugins/Experimental/EditorDataStorageFeatures/Source/TedsQueryStack/Private/TedsRowQueryResultsNode.cpp
+++ b/Engine/Plugins/Experimental/EditorDataStorageFeatures/Source/TedsQueryStack/Private/TedsRowQueryResultsNode.cpp
@@ -82,7 +82,7 @@ namespace UE::Editor::DataStorage::QueryStack
 		TargetRows.Reserve(Result.Count);
 
 		Storage.RunQuery(QueryNode->GetQuery(), 
-			EDirectQueryExecutionFlags::AllowBoundQueries | EDirectQueryExecutionFlags::IgnoreActivationCount,
+			EDirectQueryExecutionFlags::AllowBoundQueries | EDirectQueryExecutionFlags::IgnoreActiveState,
 			CreateDirectQueryCallbackBinding(
 				[&TargetRows](IDirectQueryContext& Context)
 				{
diff --git a/Engine/Source/Runtime/TypedElementFramework/Public/DataStorage/Queries/Description.h b/Engine/Source/Runtime/TypedElementFramework/Public/DataStorage/Queries/Description.h
--- a/Engine/Source/Runtime/TypedElementFramework/Public/DataStorage/Queries/Description.h
+++ b/Engine/Source/Runtime/TypedElementFramework/Public/DataStorage/Queries/Description.h
@@ -71,21 +71,34 @@ namespace UE::Editor::DataStorage
 
 		struct FCallbackData
 		{
+			// Necessary to avoid deprecation warnings in clang when using compiler generated copy and move ctors. Can
+			// be removed when ActivationCount is removed. Also causes spurious warnings in msvc, hence the conditional
+			// compilation.
+			#ifdef __clang__
+			PRAGMA_DISABLE_DEPRECATION_WARNINGS
+			FCallbackData() = default;
+			FCallbackData(const FCallbackData&) = default;
+			FCallbackData& operator=(const FCallbackData&) = default;
+			FCallbackData(FCallbackData&&) = default;
+			FCallbackData& operator=(FCallbackData&&) = default;
+			PRAGMA_ENABLE_DEPRECATION_WARNINGS
+			#endif
+
 			TArray<FName, TInlineAllocator<NumInlineGroups>> BeforeGroups;
 			TArray<FName, TInlineAllocator<NumInlineGroups>> AfterGroups;
 			QueryCallback Function;
 			FName Name;
 			FName Group;
-			/** If a name is set, it indicates the query callback will not be run unless the ActivationCount is greater than zero. */
+			/** If a name is set, it indicates the query callback will not be run unless explicitly activated with ActivateQueries. */
 			FName ActivationName;
 			const UScriptStruct* MonitoredType{ nullptr };
 			EQueryCallbackType Type{ EQueryCallbackType::None };
 			EQueryTickPhase Phase{ EQueryTickPhase::FrameEnd };
-			/**
-			 * The number of remaining iterations for a activatable query callback. If this is higher than 0, the query callback will be 
-			 * called. If ActivationName is set, this value will be decremented by one at the end of the update cycle.
-			 */
-			uint8 ActivationCount = 255;
+			bool Active = true;
+
+			UE_DEPRECATED("5.7", "Activation counts are no longer supported.")
+			uint8 ActivationCount;
+
 			EExecutionMode ExecutionMode = EExecutionMode::Default;
 		};
 		FCallbackData Callback;
diff --git a/Engine/Source/Runtime/TypedElementFramework/Public/Elements/Framework/TypedElementQueryBuilder.inl b/Engine/Source/Runtime/TypedElementFramework/Public/Elements/Framework/TypedElementQueryBuilder.inl
--- a/Engine/Source/Runtime/TypedElementFramework/Public/Elements/Framework/TypedElementQueryBuilder.inl
+++ b/Engine/Source/Runtime/TypedElementFramework/Public/Elements/Framework/TypedElementQueryBuilder.inl
@@ -1035,7 +1035,7 @@ namespace UE::Editor::DataStorage::Queries
 			Query.Callback.Phase = Processor.Phase;
 			Query.Callback.Group = Processor.Group;
 			Query.Callback.ActivationName = Processor.ActivationName;
-			Query.Callback.ActivationCount = Processor.ActivationName.IsNone() ? 255 : 0;
+			Query.Callback.Active = Processor.ActivationName.IsNone();
 			if (!Processor.BeforeGroup.IsNone())
 			{
 				Query.Callback.BeforeGroups.Add(Processor.BeforeGroup);
@@ -1061,7 +1061,7 @@ namespace UE::Editor::DataStorage::Queries
 				break;
 			}
 			Query.Callback.ActivationName = Observer.ActivationName;
-			Query.Callback.ActivationCount = Observer.ActivationName.IsNone() ? 255 : 0;
+			Query.Callback.Active = Observer.ActivationName.IsNone();
 			Query.Callback.MonitoredType = Observer.Monitor;
 			Query.Callback.ExecutionMode = Observer.ExecutionMode;
 		}
@@ -1078,7 +1078,7 @@ namespace UE::Editor::DataStorage::Queries
 				break;
 			}
 			Query.Callback.ActivationName = PhaseAmble.ActivationName;
-			Query.Callback.ActivationCount = PhaseAmble.ActivationName.IsNone() ? 255 : 0;
+			Query.Callback.Active = PhaseAmble.ActivationName.IsNone();
 			Query.Callback.Phase = PhaseAmble.Phase;
 			Query.Callback.ExecutionMode = PhaseAmble.ExecutionMode;
 		}
diff --git a/Engine/Source/Runtime/TypedElementFramework/Public/Elements/Interfaces/TypedElementQueryStorageInterfaces.h b/Engine/Source/Runtime/TypedElementFramework/Public/Elements/Interfaces/TypedElementQueryStorageInterfaces.h
--- a/Engine/Source/Runtime/TypedElementFramework/Public/Elements/Interfaces/TypedElementQueryStorageInterfaces.h
+++ b/Engine/Source/Runtime/TypedElementFramework/Public/Elements/Interfaces/TypedElementQueryStorageInterfaces.h
@@ -323,8 +323,10 @@ namespace UE::Editor::DataStorage
 		 * distributed among threads, but comes at additional scheduling cost. Add this flag if the processing time per chunk varies.
 		 */
 		AutoBalanceParallelChunkProcessing = 1 << 1,
-		IgnoreActivationCount = 1 << 2, //< If set, a direct call will not check activatable queries if they're set.
+		IgnoreActiveState = 1 << 2, //< If set, a direct call will not check activatable queries if they're set.
 		AllowBoundQueries = 1 << 3, //< Stops checking if queries are bound when set, otherwise bound queries will assert.
+
+		IgnoreActivationCount UE_DEPRECATED("5.7", "Activation counts are no longer supported.") = IgnoreActiveState,
 	};
 	ENUM_CLASS_FLAGS(EDirectQueryExecutionFlags);
 
