TEDS Hierarchy API

- An abstraction layer for tree-like hierarchies in TEDS
- QueryStack node to walk a hierarchy and collate all objects within that hierarchy (FHierarchyRowNode)

Usage:
1. Register your hierarchy before registering a query that uses it
`ICoreProvider::RegisterHierarchy()`

2. Via CoreProvider - an interface for the Hierarchy can be provided:
`ICoreProvider::GetHierarchyAccessInterface()`
This interface will let you walk down and visit children of the hierarchy

3. Via the various QueryContexts - use the provided GetParent and SetParent functions to hook together a hierarchy and walk up a hierarchy.

Note: QueryContext API will not let you walk down. Eventually a query's iteration will be configurable to operate to walk the hierarchy during the callback however this functionality is not yet ready.

diff --git a/Engine/Plugins/Experimental/EditorDataStorage/Source/TedsCore/Private/Hierarchy/EditorDataStorageHierarchyColumns.h b/Engine/Plugins/Experimental/EditorDataStorage/Source/TedsCore/Private/Hierarchy/EditorDataStorageHierarchyColumns.h
new file mode 100644
--- /dev/null
+++ b/Engine/Plugins/Experimental/EditorDataStorage/Source/TedsCore/Private/Hierarchy/EditorDataStorageHierarchyColumns.h
@@ -0,0 +1,38 @@
+﻿// Copyright Epic Games, Inc. All Rights Reserved.
+
+#pragma once
+
+#include "DataStorage/CommonTypes.h"
+
+#include "EditorDataStorageHierarchyColumns.generated.h"
+
+USTRUCT(meta = (DisplayName = "EditorDataHierarchyData (Template)", EditorDataStorage_DynamicColumnTemplate))
+struct FEditorDataHierarchyData_Template final : public FEditorDataStorageColumn
+{
+	GENERATED_BODY()
+
+	UPROPERTY()
+	FTedsRowHandle Parent;
+
+	UPROPERTY()
+	TArray<FTedsRowHandle> Children;
+};
+
+USTRUCT(meta = (DisplayName = "EditorDataHierarchyParentTag (Template)", EditorDataStorage_DynamicColumnTemplate))
+struct FEditorDataHierarchyParentTag_Template : public FEditorDataStorageTag
+{
+	GENERATED_BODY()
+};
+
+USTRUCT(meta = (DisplayName = "EditorDataHierarchyChildTag (Template)", EditorDataStorage_DynamicColumnTemplate))
+struct FEditorDataHierarchyChildTag_Template : public FEditorDataStorageTag
+{
+	GENERATED_BODY()
+};
+
+namespace UE::Editor::DataStorage
+{
+	using FHierarchyData_Template = FEditorDataHierarchyData_Template;
+	using FHierarchyParentTag_Template = FEditorDataHierarchyParentTag_Template;
+	using FHierarchyChildTag_Template = FEditorDataHierarchyChildTag_Template;
+}
diff --git a/Engine/Plugins/Experimental/EditorDataStorage/Source/TedsCore/Private/Hierarchy/EditorDataStorageHierarchyImplementation.cpp b/Engine/Plugins/Experimental/EditorDataStorage/Source/TedsCore/Private/Hierarchy/EditorDataStorageHierarchyImplementation.cpp
new file mode 100644
--- /dev/null
+++ b/Engine/Plugins/Experimental/EditorDataStorage/Source/TedsCore/Private/Hierarchy/EditorDataStorageHierarchyImplementation.cpp
@@ -0,0 +1,518 @@
+﻿// Copyright Epic Games, Inc. All Rights Reserved.
+
+#include "EditorDataStorageHierarchyImplementation.h"
+
+#include "EditorDataStorageHierarchyColumns.h"
+
+namespace UE::Editor::DataStorage
+{
+	// Interpretation for the opaque FHierarchyHandle
+	struct FHierarchyHandleInterpretation
+	{
+		int32 Index;
+		bool IsSet;
+	};
+
+	static int32 GetHandleIndex(const FHierarchyHandle& Handle)
+	{
+		const FHierarchyHandleInterpretation& Interp = reinterpret_cast<const FHierarchyHandleInterpretation&>(Handle);
+		return Interp.IsSet ? Interp.Index : INDEX_NONE;
+	}
+
+	static FHierarchyHandle CreateHandle(int32 Index)
+	{
+		FHierarchyHandle Handle;
+		
+		FHierarchyHandleInterpretation& Interp = reinterpret_cast<FHierarchyHandleInterpretation&>(Handle);
+		new (&Interp) FHierarchyHandleInterpretation();
+		Interp.Index = Index;
+		Interp.IsSet = true;
+		
+		return Handle;
+	}
+
+	static_assert(sizeof(FHierarchyHandleInterpretation) <= sizeof(FHierarchyHandle));
+	static_assert(alignof(FHierarchyHandleInterpretation) <= alignof(FHierarchyHandle));
+	static_assert(std::is_trivially_copyable_v<FHierarchyHandle>);
+	static_assert(std::is_trivially_copyable_v<FHierarchyHandleInterpretation>);
+	
+	const FTedsHierarchyAccessInterface* FTedsHierarchyRegistrar::GetAccessInterface(FHierarchyHandle Handle) const
+	{
+		if (int32 Index = GetHandleIndex(Handle); Index != INDEX_NONE)
+		{
+			return RegisteredHierarchies[Index].AccessInterface.Get();
+		}
+		return nullptr;
+	}
+
+	FHierarchyHandle FTedsHierarchyRegistrar::RegisterHierarchy(ICoreProvider* InProvider, const FHierarchyRegistrationParams& RegistrationParams)
+	{
+		// Create an access interface
+		
+		// Generate the tags and HierarchyData column per Hierarchy type
+		const UScriptStruct* HierarchyData = [InProvider, &RegistrationParams]()
+		{
+			FDynamicColumnDescription Desc
+			{
+				.TemplateType = FEditorDataHierarchyData_Template::StaticStruct(),
+				.Identifier = RegistrationParams.Name
+			};
+			return InProvider->GenerateDynamicColumn(Desc);
+		}();
+		
+		const UScriptStruct* ParentTag = [InProvider, &RegistrationParams]()
+		{
+			FDynamicColumnDescription Desc
+			{
+				.TemplateType = FEditorDataHierarchyParentTag_Template::StaticStruct(),
+				.Identifier = RegistrationParams.Name
+			};
+			return InProvider->GenerateDynamicColumn(Desc);
+		}();
+		
+		const UScriptStruct* ChildTag = [InProvider, &RegistrationParams]()
+		{
+			FDynamicColumnDescription Desc
+			{
+				.TemplateType = FEditorDataHierarchyChildTag_Template::StaticStruct(),
+				.Identifier = RegistrationParams.Name
+			};
+			return InProvider->GenerateDynamicColumn(Desc);
+		}();
+		
+		TUniquePtr<FTedsHierarchyAccessInterface> AccessInterface = [&]()
+		{
+			FTedsHierarchyAccessInterface::FConstructParams Params
+			{
+				.ChildTag = ChildTag,
+				.ParentTag = ParentTag,
+				.HierarchyData = HierarchyData
+			};
+			return MakeUnique<FTedsHierarchyAccessInterface>(Params);
+		}();
+		
+		FRegisteredHierarchy RegisteredHierarchy
+		{
+			.Name = RegistrationParams.Name,
+			.ChildTag = ChildTag,
+			.ParentTag = ParentTag,
+			.HierarchyData = HierarchyData,
+			.AccessInterface = MoveTemp(AccessInterface)
+		};
+
+		int32 Index = RegisteredHierarchies.Num();
+		RegisteredHierarchies.Emplace(MoveTemp(RegisteredHierarchy));
+
+		FHierarchyHandle HierarchyHandle = CreateHandle(Index);
+		
+		ICoreProvider* DataStorage = GetMutableDataStorageFeature<ICoreProvider>(StorageFeatureName);
+		RegisterObservers(*DataStorage, HierarchyHandle);
+		
+		return HierarchyHandle;
+	}
+
+	FHierarchyHandle FTedsHierarchyRegistrar::FindHierarchyByName(const FName& Name) const
+	{
+		for (int32 HierarchyIndex = 0; HierarchyIndex < RegisteredHierarchies.Num(); ++HierarchyIndex)
+		{
+			if (RegisteredHierarchies[HierarchyIndex].Name == Name)
+			{
+				return CreateHandle(HierarchyIndex);
+			}
+		}
+		return FHierarchyHandle();
+	}
+
+	auto Hack_FindRowIndexInChunk = [](const ICommonQueryContext& Context, RowHandle Row) -> int32
+	{
+		int32 RowIndex = INDEX_NONE;
+		TConstArrayView<RowHandle> RowHandes = Context.GetRowHandles();
+		const int32 RowCount = RowHandes.Num();
+		for (int32 CandidateRowIndex = 0; CandidateRowIndex < RowCount; ++CandidateRowIndex)
+		{
+			if (RowHandes[CandidateRowIndex] == Row)
+			{
+				RowIndex = CandidateRowIndex;
+				break;
+			}
+		}
+		return RowIndex;
+	};
+
+	void FTedsHierarchyRegistrar::RegisterObservers(ICoreProvider& DataStorage, const FHierarchyHandle& Handle)
+	{
+		using namespace UE::Editor::DataStorage::Queries;
+
+		const FRegisteredHierarchy* Hierarchy = FindRegisteredHierarchy(Handle);
+		check(Hierarchy != nullptr);
+
+		QueryHandle Subquery = DataStorage.RegisterQuery(
+				Select()
+					.ReadWrite(Hierarchy->HierarchyData)
+				.Compile());
+
+		DataStorage.RegisterQuery(
+		Select(
+			TEXT("Remove child from parent's child array"),
+			FObserver(FObserver::EEvent::Remove, Hierarchy->ChildTag),
+			[this, Handle](IQueryContext& Context, RowHandle Row)
+			{
+				const FRegisteredHierarchy* Hierarchy = FindRegisteredHierarchy(Handle);
+				
+				FEditorDataHierarchyData_Template* ChildHierarchyDataComponent;
+				{
+					const int RowIndexInChunk = Hack_FindRowIndexInChunk(Context, Row);
+					check(RowIndexInChunk != INDEX_NONE);
+					void* HierarchyDataComponentArrayStart = Context.GetMutableColumn(Hierarchy->HierarchyData);
+					uint8* HierarchyDataComponentObjectAddress = static_cast<uint8*>(HierarchyDataComponentArrayStart) + (RowIndexInChunk * Hierarchy->HierarchyData->GetStructureSize());
+					ChildHierarchyDataComponent = reinterpret_cast<FEditorDataHierarchyData_Template*>(HierarchyDataComponentObjectAddress);
+				}
+				if (ChildHierarchyDataComponent->Parent == InvalidRowHandle)
+				{
+					return;
+				}
+				// Remove the ParentTag from the parent row if it has no more children
+				bool bRemoveParentTag = false;
+				Context.RunSubquery(0, ChildHierarchyDataComponent->Parent, CreateSubqueryCallbackBinding([Row, &bRemoveParentTag, Hierarchy](ISubqueryContext& SubQueryContext, RowHandle ParentRow)
+				{
+					FEditorDataHierarchyData_Template* ParentHierarchyDataComponent;
+					{
+						const int RowIndexInChunk = Hack_FindRowIndexInChunk(SubQueryContext, ParentRow);
+						check(RowIndexInChunk != INDEX_NONE);
+						void* HierarchyDataComponentArrayStart = SubQueryContext.GetMutableColumn(Hierarchy->HierarchyData);
+						uint8* HierarchyDataComponentObjectAddress = static_cast<uint8*>(HierarchyDataComponentArrayStart) + (RowIndexInChunk * Hierarchy->HierarchyData->GetStructureSize());
+						ParentHierarchyDataComponent = reinterpret_cast<FEditorDataHierarchyData_Template*>(HierarchyDataComponentObjectAddress);
+					}
+					
+					ParentHierarchyDataComponent->Children.RemoveSwap(FTedsRowHandle(Row));
+					if (ParentHierarchyDataComponent->Children.IsEmpty())
+					{
+						bRemoveParentTag = true;
+					}
+				}));
+				if (bRemoveParentTag)
+				{
+					Context.RemoveColumns(ChildHierarchyDataComponent->Parent, {Hierarchy->ParentTag});
+				}
+				
+				ChildHierarchyDataComponent->Parent = InvalidRowHandle;
+			}
+		)
+		.ReadWrite(Hierarchy->HierarchyData)
+		.DependsOn().SubQuery(Subquery)
+		.Compile());
+
+		DataStorage.RegisterQuery(
+		Select(
+			TEXT("Remove parent reference from children of destroyed parent"),
+			FObserver(FObserver::EEvent::Remove, Hierarchy->ParentTag),
+			[this, Handle](IQueryContext& Context, RowHandle ParentRow)
+			{
+				const FRegisteredHierarchy* Hierarchy = FindRegisteredHierarchy(Handle);
+				
+				FEditorDataHierarchyData_Template* ParentHierarchyDataComponent;
+				{
+					const int RowIndexInChunk = Hack_FindRowIndexInChunk(Context, ParentRow);
+					check(RowIndexInChunk != INDEX_NONE);
+					void* HierarchyDataComponentArrayStart = Context.GetMutableColumn(Hierarchy->HierarchyData);
+					uint8* HierarchyDataComponentObjectAddress = static_cast<uint8*>(HierarchyDataComponentArrayStart) + (RowIndexInChunk * Hierarchy->HierarchyData->GetStructureSize());
+					ParentHierarchyDataComponent = reinterpret_cast<FEditorDataHierarchyData_Template*>(HierarchyDataComponentObjectAddress);
+				}
+				if (ParentHierarchyDataComponent->Children.IsEmpty())
+				{
+					return;
+				}
+				for (int32 ChildIndex = 0; ChildIndex < ParentHierarchyDataComponent->Children.Num(); ++ChildIndex)
+				{
+					RowHandle ChildRow = ParentHierarchyDataComponent->Children[ChildIndex];
+					Context.RunSubquery(0, ChildRow, CreateSubqueryCallbackBinding([ParentRow, Hierarchy](ISubqueryContext& SubQueryContext, RowHandle ChildRow)
+					{
+						FEditorDataHierarchyData_Template* ChildHierarchyDataComponent;
+						{
+							const int RowIndexInChunk = Hack_FindRowIndexInChunk(SubQueryContext, ChildRow);
+							check(RowIndexInChunk != INDEX_NONE);
+							void* HierarchyDataComponentArrayStart = SubQueryContext.GetMutableColumn(Hierarchy->HierarchyData);
+							uint8* HierarchyDataComponentObjectAddress = static_cast<uint8*>(HierarchyDataComponentArrayStart) + (RowIndexInChunk * Hierarchy->HierarchyData->GetStructureSize());
+							ChildHierarchyDataComponent = reinterpret_cast<FEditorDataHierarchyData_Template*>(HierarchyDataComponentObjectAddress);
+						}
+						ChildHierarchyDataComponent->Parent = InvalidRowHandle;
+					}));
+
+					Context.RemoveColumns(ChildRow, {Hierarchy->ChildTag});
+				}
+
+				ParentHierarchyDataComponent->Children.Reset();
+			}
+		)
+		.ReadWrite(Hierarchy->HierarchyData)
+		.DependsOn().SubQuery(Subquery)
+		.Compile());
+	}
+
+	const FTedsHierarchyRegistrar::FRegisteredHierarchy* FTedsHierarchyRegistrar::FindRegisteredHierarchy(const FHierarchyHandle& Handle)
+	{
+		int32 Index = GetHandleIndex(Handle);
+		if (Index != INDEX_NONE)
+		{
+			return &RegisteredHierarchies[Index];
+		}
+		return nullptr;
+	}
+
+	FTedsHierarchyAccessInterface::FTedsHierarchyAccessInterface(const FConstructParams& Params)
+		: ChildTag(Params.ChildTag)
+		, ParentTag(Params.ParentTag)
+		, HierarchyDataColumnType(Params.HierarchyData)
+	{
+		check(HierarchyDataColumnType->IsChildOf(FEditorDataHierarchyData_Template::StaticStruct()));
+	}
+
+	const UScriptStruct* FTedsHierarchyAccessInterface::GetChildTagType() const
+	{
+		return ChildTag;
+	}
+
+	const UScriptStruct* FTedsHierarchyAccessInterface::GetParentTagType() const
+	{
+		return ParentTag;
+	}
+
+	const UScriptStruct* FTedsHierarchyAccessInterface::GetHierarchyDataColumnType() const
+	{
+		return HierarchyDataColumnType;
+	}
+	
+	bool FTedsHierarchyAccessInterface::HasChildren(const ICoreProvider& Context, RowHandle Row) const
+	{
+		const FEditorDataHierarchyData_Template* Hierarchy = static_cast<const FEditorDataHierarchyData_Template*>(Context.GetColumnData(Row, HierarchyDataColumnType));
+		return Hierarchy && !Hierarchy->Children.IsEmpty();
+	}
+
+	void FTedsHierarchyAccessInterface::WalkDepthFirst(
+		const ICoreProvider& Context,
+		RowHandle Row,
+		TFunction<void(const ICoreProvider& Context, RowHandle Owner,RowHandle Target)> OnVisitedFn) const
+	{
+		// Recursively walk depth first to visit each child
+		auto WalkDepthFirstImpl = [this](const ICoreProvider& Context_, RowHandle Row, decltype(OnVisitedFn)& OnVisitedFn_, auto& WalkDepthFirstImplRef) -> void
+		{
+			const FEditorDataHierarchyData_Template* HierarchyData = static_cast<const FEditorDataHierarchyData_Template*>(Context_.GetColumnData(Row, HierarchyDataColumnType));
+			if (HierarchyData)
+			{
+				for (RowHandle ChildRow : HierarchyData->Children)
+				{					
+					OnVisitedFn_(Context_, Row, ChildRow);
+					
+					WalkDepthFirstImplRef(Context_, ChildRow, OnVisitedFn_, WalkDepthFirstImplRef);
+				}
+			}
+		};
+
+		// Call the top level object
+		OnVisitedFn(Context, InvalidRowHandle /*no parent*/, Row);
+
+		WalkDepthFirstImpl(Context, Row, OnVisitedFn, WalkDepthFirstImpl);
+	}
+
+	void FTedsHierarchyAccessInterface::SetParentRow(ICoreProvider* CoreProvider, RowHandle Target, RowHandle Parent) const
+	{
+		CoreProvider->RemoveColumn(Target, ChildTag);
+
+		auto AddHierarchyDataColumn = [this](ICoreProvider* CoreProvider, RowHandle Row, RowHandle ParentRow)
+		{
+			CoreProvider->AddColumnData(
+				Row,
+				HierarchyDataColumnType,
+				[ParentRow](void* Dest, const UScriptStruct& ColumnType)
+				{
+					ColumnType.InitializeStruct(Dest);
+					FEditorDataHierarchyData_Template* HierarchyDataColumn = static_cast<FEditorDataHierarchyData_Template*>(Dest);
+					check(HierarchyDataColumn->Parent == InvalidRowHandle);
+					HierarchyDataColumn->Parent = ParentRow;
+				},
+				[](const UScriptStruct& ColumnType, void* Destination, void* Source)
+				{
+					// We could get away with defining this as the move operator of the base class as
+					// it should be the same as whatever is in HierarchyDataColumnType
+					// However, it is technically slicing.  
+					new (Destination) FEditorDataHierarchyData_Template(MoveTemp(*static_cast<FEditorDataHierarchyData_Template*>(Source)));
+				}
+			);
+		};
+
+		FEditorDataHierarchyData_Template* HierarchyDataColumn_Parent;
+		if (void* ParentHierarchyDataColumnAddress = CoreProvider->GetColumnData(Parent, HierarchyDataColumnType); ParentHierarchyDataColumnAddress == nullptr)
+		{
+			AddHierarchyDataColumn(CoreProvider, Parent, InvalidRowHandle);
+			ParentHierarchyDataColumnAddress = CoreProvider->GetColumnData(Parent, HierarchyDataColumnType);
+			check(ParentHierarchyDataColumnAddress);
+			HierarchyDataColumn_Parent = static_cast<FEditorDataHierarchyData_Template*>(ParentHierarchyDataColumnAddress);
+		}
+		else
+		{
+			HierarchyDataColumn_Parent = static_cast<FEditorDataHierarchyData_Template*>(ParentHierarchyDataColumnAddress);
+		}
+		
+		FEditorDataHierarchyData_Template* HierarchyDataColumn_Target;
+		if (void* TargetHierarchyDataColumnAddress = CoreProvider->GetColumnData(Target, HierarchyDataColumnType); TargetHierarchyDataColumnAddress == nullptr)
+		{
+			AddHierarchyDataColumn(CoreProvider, Target, Parent);
+			TargetHierarchyDataColumnAddress = CoreProvider->GetColumnData(Target, HierarchyDataColumnType);
+			check(TargetHierarchyDataColumnAddress);
+			HierarchyDataColumn_Target = static_cast<FEditorDataHierarchyData_Template*>(TargetHierarchyDataColumnAddress);
+		}
+		else
+		{
+			HierarchyDataColumn_Target = static_cast<FEditorDataHierarchyData_Template*>(TargetHierarchyDataColumnAddress);
+		}
+		
+		if (Parent != InvalidRowHandle)
+		{
+			HierarchyDataColumn_Target->Parent = Parent;
+			HierarchyDataColumn_Parent->Children.Add(FTedsRowHandle(Target));
+			CoreProvider->AddColumn(Target, ChildTag);
+			CoreProvider->AddColumn(Parent, ParentTag);
+		}
+		else
+		{
+			CoreProvider->RemoveColumns(Target, {ChildTag});
+
+			// All this work is handled by observer...
+			// Commented out just to make it clear what should happen
+			
+			// HierarchyDataColumn_Target->Parent = InvalidRowHandle;
+			// HierarchyDataColumn_Parent->Children.Remove(FTedsRowHandle(Target));
+			// if (HierarchyDataColumn_Parent->Children.IsEmpty())
+			// {
+			// 	CoreProvider->RemoveColumns(Parent, {ParentTag});
+			// }
+		}
+	}
+
+	RowHandle FTedsHierarchyAccessInterface::GetParentRow(const ICoreProvider* CoreProvider, RowHandle Target) const
+	{
+		const UScriptStruct* HierarchyDataColumn = GetHierarchyDataColumnType();
+		const void* ColumnData = CoreProvider->GetColumnData(Target, HierarchyDataColumn);
+		if (ColumnData != nullptr)
+		{
+			const FEditorDataHierarchyData_Template* Column = static_cast<const FEditorDataHierarchyData_Template*>(ColumnData);
+			return Column->Parent.RowHandle;
+		}
+		return InvalidRowHandle;
+	}
+
+	TFunction<RowHandle(const void*, const UScriptStruct*)> FTedsHierarchyAccessInterface::CreateParentExtractionFunction() const
+	{
+		return [this](const void* ColumnData, const UScriptStruct* ColumnType)->RowHandle
+		{
+			check(ColumnType == GetHierarchyDataColumnType());
+			const FEditorDataHierarchyData_Template* HierarchyDataColumn = static_cast<const FEditorDataHierarchyData_Template*>(ColumnData);
+			return HierarchyDataColumn->Parent.RowHandle;
+		};
+	}
+
+	// Deferred command to set the parent of a target row
+	// Setting a parent may result in the Target and Parent to change archetypes due to the addition of data and tags
+	// to store the relationships if they are not already established
+	struct FSetParentCommand
+	{
+		RowHandle Parent;
+		RowHandle Target;
+		const FTedsHierarchyAccessInterface* HierarchyAccessInterface = nullptr;
+		ICoreProvider* CoreProvider = nullptr;
+
+		void operator()()
+		{
+			if ((Parent == InvalidRowHandle || CoreProvider->IsRowAssigned(Parent)) && CoreProvider->IsRowAssigned(Target))
+			{
+				HierarchyAccessInterface->SetParentRow(CoreProvider, Target, Parent);
+			}
+		}
+	};
+
+	
+	void FTedsHierarchyAccessInterface::SetParent(ICommonQueryContext& Context, RowHandle Target, RowHandle Parent) const
+	{
+		FSetParentCommand Command;
+		Command.Parent = Parent;
+		Command.Target = Target;
+		Command.HierarchyAccessInterface = this;
+		Command.CoreProvider = GetMutableDataStorageFeature<ICoreProvider>(StorageFeatureName);
+		Context.PushCommand(MoveTemp(Command));
+	}
+
+	RowHandle FTedsHierarchyAccessInterface::GetParent(const ICommonQueryContext& Context, RowHandle Target) const
+	{
+		// NOTE: It is not clear to the user that they will need to have registered the HierarchyDataColumnType as at least ReadOnly
+		// for this function to return the Parent.
+		// This can be improved with a function on the the context (and other contexts) to check for the access requirements of a
+		// query and warn or ensure about a missing requirement - though this does not currently exist at time of writing.
+		const FEditorDataHierarchyData_Template* ParentHierarchyDataComponent;
+		{
+			const int RowIndexInChunk = Hack_FindRowIndexInChunk(Context, Target);
+			if (RowIndexInChunk == INDEX_NONE)
+			{
+				return InvalidRowHandle;
+			}
+			const void* HierarchyDataComponentArrayStart = Context.GetColumn(HierarchyDataColumnType);
+			if (HierarchyDataComponentArrayStart == nullptr)
+			{
+				return InvalidRowHandle;
+			}
+			const uint8* HierarchyDataComponentObjectAddress = static_cast<const uint8*>(HierarchyDataComponentArrayStart) + (RowIndexInChunk * HierarchyDataColumnType->GetStructureSize());
+			ParentHierarchyDataComponent = reinterpret_cast<const FEditorDataHierarchyData_Template*>(HierarchyDataComponentObjectAddress);
+		}
+
+		return ParentHierarchyDataComponent->Parent;
+	}
+
+	void FTedsHierarchyAccessInterface::SetParentRow(IQueryContext& Context, RowHandle Target, RowHandle Parent) const
+	{
+		SetParent(static_cast<ICommonQueryContext&>(Context), Target, Parent);
+	}
+
+	RowHandle FTedsHierarchyAccessInterface::GetParentRow(const IQueryContext& Context, RowHandle Target) const
+	{
+		return GetParent(static_cast<const ICommonQueryContext&>(Context), Target);
+	}
+
+	void FTedsHierarchyAccessInterface::SetParentRow(ISubqueryContext& Context, RowHandle Target, RowHandle Parent) const
+	{
+		SetParent(static_cast<ICommonQueryContext&>(Context), Target, Parent);
+	}
+
+	RowHandle FTedsHierarchyAccessInterface::GetParentRow(const ISubqueryContext& Context, RowHandle Target) const
+	{
+		return GetParent(static_cast<const ICommonQueryContext&>(Context), Target);
+	}
+
+	void FTedsHierarchyAccessInterface::SetParentRow(IDirectQueryContext& Context, RowHandle Target, RowHandle Parent) const
+	{
+		SetParent(static_cast<ICommonQueryContext&>(Context), Target, Parent);
+	}
+
+	RowHandle FTedsHierarchyAccessInterface::GetParentRow(const IDirectQueryContext& Context, RowHandle Target) const
+	{
+		// NOTE: It is not clear to the user that they will need to have registered the HierarchyDataColumnType as at least ReadOnly
+		// for this function to return the Parent.
+		// This can be improved with a function on the the context (and other contexts) to check for the access requirements of a
+		// query and warn or ensure about a missing requirement - though this does not currently exist at time of writing.
+		const FEditorDataHierarchyData_Template* ParentHierarchyDataComponent;
+		{
+			const int RowIndexInChunk = Hack_FindRowIndexInChunk(Context, Target);
+			if (RowIndexInChunk == INDEX_NONE)
+			{
+				return InvalidRowHandle;
+			}
+			const void* HierarchyDataComponentArrayStart = Context.GetColumn(HierarchyDataColumnType);
+			if (HierarchyDataComponentArrayStart == nullptr)
+			{
+				return InvalidRowHandle;
+			}
+			const uint8* HierarchyDataComponentObjectAddress = static_cast<const uint8*>(HierarchyDataComponentArrayStart) + (RowIndexInChunk * HierarchyDataColumnType->GetStructureSize());
+			ParentHierarchyDataComponent = reinterpret_cast<const FEditorDataHierarchyData_Template*>(HierarchyDataComponentObjectAddress);
+		}
+
+		return ParentHierarchyDataComponent->Parent;
+	}
+}
diff --git a/Engine/Plugins/Experimental/EditorDataStorage/Source/TedsCore/Private/Hierarchy/EditorDataStorageHierarchyImplementation.h b/Engine/Plugins/Experimental/EditorDataStorage/Source/TedsCore/Private/Hierarchy/EditorDataStorageHierarchyImplementation.h
new file mode 100644
--- /dev/null
+++ b/Engine/Plugins/Experimental/EditorDataStorage/Source/TedsCore/Private/Hierarchy/EditorDataStorageHierarchyImplementation.h
@@ -0,0 +1,89 @@
+﻿// Copyright Epic Games, Inc. All Rights Reserved.
+
+#pragma once
+
+#include "DataStorage/Features.h"
+#include "Elements/Framework/TypedElementQueryBuilder.h"
+#include "Elements/Interfaces/TypedElementDataStorageInterface.h"
+
+namespace UE::Editor::DataStorage
+{
+	class FTedsHierarchyAccessInterface;
+	class FTedsHierarchyRegistrar final
+	{
+	public:
+		/**
+		 * Registers a hierarchy type with TEDS
+		 * This sets up the internal columns and processors to kee the bidirectional relationship in sync
+		 */
+		FHierarchyHandle RegisterHierarchy(ICoreProvider* InProvider, const FHierarchyRegistrationParams& Params);
+		/**
+		 * Finds a registered hierarchy by name
+		 */
+		FHierarchyHandle FindHierarchyByName(const FName& Name) const;
+		/**
+		 * Gets an interface to add and remove rows from a hierarchy and walk the hierarchy
+		 */
+		const FTedsHierarchyAccessInterface* GetAccessInterface(FHierarchyHandle) const;
+	private:
+		void RegisterObservers(ICoreProvider& CoreProvider, const FHierarchyHandle& Handle);
+		
+		struct FRegisteredHierarchy
+		{
+			FName Name;
+			const UScriptStruct* ChildTag;
+			const UScriptStruct* ParentTag;
+			const UScriptStruct* HierarchyData;
+			TUniquePtr<FTedsHierarchyAccessInterface> AccessInterface;
+		};
+
+		const FRegisteredHierarchy* FindRegisteredHierarchy(const FHierarchyHandle& Handle);
+	
+	private:
+		TArray<FRegisteredHierarchy> RegisteredHierarchies;
+	};
+
+	class FTedsHierarchyAccessInterface final
+	{
+	public:
+		struct FConstructParams
+		{
+			const UScriptStruct* ChildTag;
+			const UScriptStruct* ParentTag;
+			const UScriptStruct* HierarchyData;
+		};
+		
+		explicit FTedsHierarchyAccessInterface(const FConstructParams& Params);
+		
+		const UScriptStruct* GetChildTagType() const;
+		const UScriptStruct* GetParentTagType() const;
+		const UScriptStruct* GetHierarchyDataColumnType() const;
+		
+		void SetParentRow(ICoreProvider* CoreProvider, RowHandle Target, RowHandle Parent) const;
+		RowHandle GetParentRow(const ICoreProvider* CoreProvider, RowHandle Target) const;
+		
+		bool HasChildren(const ICoreProvider& Context, RowHandle Row) const;
+		void WalkDepthFirst(
+			const ICoreProvider& Context,
+			RowHandle Row,
+			TFunction<void(const ICoreProvider& Context, RowHandle Owner, RowHandle Target)> OnVisitedFn) const;
+		
+		TFunction<RowHandle(const void*, const UScriptStruct*)> CreateParentExtractionFunction() const;
+		
+		void SetParentRow(IQueryContext& Context, RowHandle Target, RowHandle Parent) const;
+		RowHandle GetParentRow(const IQueryContext& Context, RowHandle Target) const;
+		void SetParentRow(ISubqueryContext& Context, RowHandle Target, RowHandle Parent) const;
+		RowHandle GetParentRow(const ISubqueryContext& Context, RowHandle Target) const;
+		void SetParentRow(IDirectQueryContext& Context, RowHandle Target, RowHandle Parent) const;
+		RowHandle GetParentRow(const IDirectQueryContext& Context, RowHandle Target) const;
+
+	private:
+
+		void SetParent(ICommonQueryContext& Context, RowHandle Target, RowHandle Parent) const;
+		RowHandle GetParent(const ICommonQueryContext& Context, RowHandle Target) const;
+		
+		const UScriptStruct* ChildTag;
+		const UScriptStruct* ParentTag;
+		const UScriptStruct* HierarchyDataColumnType;
+	};
+}
diff --git a/Engine/Plugins/Experimental/EditorDataStorage/Source/TedsCore/Private/Processors/TypedElementProcessorAdaptors.cpp b/Engine/Plugins/Experimental/EditorDataStorage/Source/TedsCore/Private/Processors/TypedElementProcessorAdaptors.cpp
--- a/Engine/Plugins/Experimental/EditorDataStorage/Source/TedsCore/Private/Processors/TypedElementProcessorAdaptors.cpp
+++ b/Engine/Plugins/Experimental/EditorDataStorage/Source/TedsCore/Private/Processors/TypedElementProcessorAdaptors.cpp
@@ -118,10 +118,58 @@
 				return DynamicColumnType;
 			}
 
-			float GetDeltaTimeSeconds() const 
+			float GetDeltaTimeSeconds() const
 			{
 				return Context.GetDeltaTimeSeconds();
 			}
+
+			void PushCommand(void (*CommandFunction)(void*), void* InCommandData) const
+			{
+				if (!ensure(CommandFunction))
+				{
+					return;
+				}
+				const FEnvironment::FEnvironmentCommand Command
+				{
+					.CommandFunction = CommandFunction,
+					.CommandData = InCommandData
+				};
+				Environment.PushCommands(MakeConstArrayView(&Command, 1));
+			}
+
+			void* EmplaceObjectInScratch(size_t ObjectSize, size_t Alignment, void(* Construct)(void*, void*), void(* Destroy)(void*), void* SourceCommandContext) const
+			{
+				struct FDestructor final
+				{
+					using DestroyFnType = decltype(Destroy);
+					FDestructor(DestroyFnType InDestroyFn, void* InObjectPtr)
+						: DestroyFn(InDestroyFn)
+						, ObjectPtr(InObjectPtr)
+					{
+					}
+					~FDestructor()
+					{
+						DestroyFn(ObjectPtr);
+					}
+					DestroyFnType DestroyFn;
+					void* ObjectPtr;
+				};
+
+				FScratchBuffer& ScratchBuffer = Environment.GetScratchBuffer();
+				
+				void* ObjectMemory = ScratchBuffer.AllocateUninitialized(ObjectSize, Alignment);
+				Construct(ObjectMemory, SourceCommandContext);
+
+				// The presence of a Destroy function implies that the objects that was just added to the scratch buffer
+				// is not trivially destructible, hence needs its destructor called.
+				// The API for the scratch buffer's internal memory allocator needs us to emplace a non-trivially destructible object
+				// of some type. FDestructor is used to fulfill that role to destroy the object that was just constructed.
+				if (Destroy)
+				{
+					ScratchBuffer.Emplace<FDestructor>(Destroy, ObjectMemory);	
+				}
+				return ObjectMemory;
+			}
 		};
 
 		struct FRowColumnModifications
@@ -930,54 +978,25 @@
 				}
 			}
 
-			void PushCommand(void (*CommandFunction)(void*), void* InCommandData)
+			template<typename ContextType>
+			void SetParent(ContextType& Context, RowHandle Target, RowHandle Parent) const
 			{
-				if (!ensure(CommandFunction))
+				if (HierarchyAccessInterface)
 				{
-					return;
+					HierarchyAccessInterface->SetParentRow(Context, Target, Parent);
 				}
-				const FEnvironment::FEnvironmentCommand Command
-				{
-					.CommandFunction = CommandFunction,
-					.CommandData = InCommandData
-				};
-				Environment.PushCommands(MakeConstArrayView(&Command, 1));
 			}
 
-			void* EmplaceObjectInScratch(size_t ObjectSize, size_t Alignment, void(* Construct)(void*, void*), void(* Destroy)(void*), void* SourceCommandContext)
+			template<typename ContextType>
+			RowHandle GetParent(ContextType& Context, RowHandle Target) const
 			{
-				struct FDestructor final
+				if (HierarchyAccessInterface)
 				{
-					using DestroyFnType = decltype(Destroy);
-					FDestructor(DestroyFnType InDestroyFn, void* InObjectPtr)
-						: DestroyFn(InDestroyFn)
-						, ObjectPtr(InObjectPtr)
-					{
-					}
-					~FDestructor()
-					{
-						DestroyFn(ObjectPtr);
-					}
-					DestroyFnType DestroyFn;
-					void* ObjectPtr;
-				};
-
-				FScratchBuffer& ScratchBuffer = Environment.GetScratchBuffer();
-				
-				void* ObjectMemory = ScratchBuffer.AllocateUninitialized(ObjectSize, Alignment);
-				Construct(ObjectMemory, SourceCommandContext);
-
-				// The presence of a Destroy function implies that the objects that was just added to the scratch buffer
-				// is not trivially destructible, hence needs its destructor called.
-				// The API for the scratch buffer's internal memory allocator needs us to emplace a non-trivially destructible object
-				// of some type. FDestructor is used to fulfill that role to destroy the object that was just constructed.
-				if (Destroy)
-				{
-					ScratchBuffer.Emplace<FDestructor>(Destroy, ObjectMemory);	
+					return HierarchyAccessInterface->GetParentRow(Context, Target);
 				}
-				return ObjectMemory;
+				return InvalidRowHandle;
 			}
-
+			
 			struct FAddedColumns
 			{
 				FMassArchetypeCompositionDescriptor AddDescriptor;
@@ -1078,13 +1097,17 @@
 			TChunkedArray<FRemovedColumnsBatch*> RemovedColumnsBatchQueue;
 			TChunkedArray<FAddRow> AddedRowsQueue;
 			TChunkedArray<RowHandle> RemovedRowsQueue;
+			const FTedsHierarchyAccessInterface* HierarchyAccessInterface = nullptr;
 		};
 
 		struct FMassDirectContextForwarder final : public IDirectQueryContext
 		{
-			FMassDirectContextForwarder(FMassExecutionContext& InContext, FEnvironment& InEnvironment)
+			FMassDirectContextForwarder(FMassExecutionContext& InContext, FEnvironment& InEnvironment, FQueryDescription& Description)
 				: Implementation(InContext, InEnvironment)
-			{}
+			{
+				FHierarchyHandle Handle = InEnvironment.GetHierarchyRegistrar().FindHierarchyByName(Description.Hierarchy);
+				HierarchyAccessInterface = InEnvironment.GetHierarchyRegistrar().GetAccessInterface(Handle);
+			}
 	
 			uint32 GetRowCount() const override { return Implementation.GetRowCount(); }
 			TConstArrayView<RowHandle> GetRowHandles() const override { return Implementation.GetRowHandles(); }
@@ -1096,17 +1119,33 @@
 			bool HasColumn(RowHandle Row, const UScriptStruct* ColumnType) const override { return Implementation.HasColumn(Row, ColumnType); }
 			const UScriptStruct* FindDynamicColumnType(const UE::Editor::DataStorage::FDynamicColumnDescription& Description) const override { return Implementation.FindDynamicColumnType(Description); }
 			float GetDeltaTimeSeconds() const override { return Implementation.GetDeltaTimeSeconds(); }
+			virtual void SetParentRow(RowHandle Target, RowHandle Parent) override { if (HierarchyAccessInterface) { HierarchyAccessInterface->SetParentRow(*static_cast<IDirectQueryContext*>(this), Target, Parent); } }
+			virtual RowHandle GetParentRow(RowHandle Target) const override { return HierarchyAccessInterface ? HierarchyAccessInterface->GetParentRow(*static_cast<const IDirectQueryContext*>(this), Target) : InvalidRowHandle; }
+			virtual void PushCommand(void(* CommandFunction)(void*), void* InCommandData) override  { return Implementation.PushCommand(CommandFunction, InCommandData); }
+		
+		protected:
+			void* EmplaceObjectInScratch(const FEmplaceObjectParams& Params) override { return Implementation.EmplaceObjectInScratch(Params.ObjectSize, Params.Alignment, Params.Construct, Params.Destroy, Params.SourceObject); }
 
+		public:
+
 			FMassContextCommon Implementation;
+			const FTedsHierarchyAccessInterface* HierarchyAccessInterface = nullptr;
 		};
 
 		struct FMassSubqueryContextForwarder  : public ISubqueryContext
 		{
-			FMassSubqueryContextForwarder(FMassExecutionContext& InContext, FEnvironment& InEnvironment, bool bInBatchDeferredCommands)
-				: Implementation(InContext, InEnvironment, bInBatchDeferredCommands)
-			{}
+			FMassSubqueryContextForwarder(FMassExecutionContext& InContext, FEnvironment& InEnvironment, FQueryDescription& Description)
+				: Implementation(InContext, InEnvironment, Description.bShouldBatchModifications)
+			{
+				FHierarchyHandle Handle = InEnvironment.GetHierarchyRegistrar().FindHierarchyByName(Description.Hierarchy);
+				Implementation.HierarchyAccessInterface = InEnvironment.GetHierarchyRegistrar().GetAccessInterface(Handle);
+			}
 
 			~FMassSubqueryContextForwarder() override = default;
+			void SetHierarchyAccessInterface(const FTedsHierarchyAccessInterface* InHierarchyAccessInterface)
+			{
+				Implementation.HierarchyAccessInterface = InHierarchyAccessInterface;
+			}
 			uint32 GetRowCount() const override { return Implementation.GetRowCount(); }
 			TConstArrayView<RowHandle> GetRowHandles() const override { return Implementation.GetRowHandles(); }
 			const void* GetColumn(const UScriptStruct* ColumnType) const override { return Implementation.GetColumn(ColumnType); }
@@ -1133,10 +1172,12 @@
 			void RemoveColumns(TConstArrayView<RowHandle> Rows, TConstArrayView<const UScriptStruct*> ColumnTypes) override { Implementation.RemoveColumns(Rows, ColumnTypes); }
 			const UScriptStruct* FindDynamicColumnType(const UE::Editor::DataStorage::FDynamicColumnDescription& Description) const override { return Implementation.FindDynamicColumnType(Description); }
 			float GetDeltaTimeSeconds() const override { return Implementation.GetDeltaTimeSeconds(); }
+			virtual void SetParentRow(RowHandle Target, RowHandle Parent) override { Implementation.SetParent(*static_cast<ISubqueryContext*>(this), Target, Parent); }
+			virtual RowHandle GetParentRow(RowHandle Target) const override { return Implementation.GetParent(*static_cast<const ISubqueryContext*>(this), Target); }
 			void PushCommand(void (*CommandFunction)(void*), void* CommandData) override { return Implementation.PushCommand(CommandFunction, CommandData); }
 		protected:
 			void* EmplaceObjectInScratch(const FEmplaceObjectParams& Params) override { return Implementation.EmplaceObjectInScratch(Params.ObjectSize, Params.Alignment, Params.Construct, Params.Destroy, Params.SourceObject); }
-
+			
 			FMassWithEnvironmentContextCommon Implementation;
 		};
 
@@ -1151,7 +1192,8 @@
 				, QueryDescription(InQueryDescription)
 				, QueryStore(InQueryStore)
 			{
-		
+				FHierarchyHandle Handle = InEnvironment.GetHierarchyRegistrar().FindHierarchyByName(InQueryDescription.Hierarchy);
+				HierarchyAccessInterface = InEnvironment.GetHierarchyRegistrar().GetAccessInterface(Handle);
 			}
 
 			~FMassQueryContextImplementation() = default;
@@ -1254,7 +1296,8 @@
 				FExtendedQueryStore& InQueryStore,
 				FEnvironment& InEnvironment)
 					: Implementation(InQueryDescription, InContext, InQueryStore, InEnvironment)
-			{}
+			{
+			}
 	
 			uint32 GetRowCount() const override { return Implementation.GetRowCount(); }
 			TConstArrayView<RowHandle> GetRowHandles() const override { return Implementation.GetRowHandles(); }
@@ -1282,6 +1325,8 @@
 			void RemoveColumns(TConstArrayView<RowHandle> Rows, TConstArrayView<const UScriptStruct*> ColumnTypes) override { Implementation.RemoveColumns(Rows, ColumnTypes); }
 			const UScriptStruct* FindDynamicColumnType(const UE::Editor::DataStorage::FDynamicColumnDescription& Description) const override { return Implementation.FindDynamicColumnType(Description); }
 			float GetDeltaTimeSeconds() const override { return Implementation.GetDeltaTimeSeconds(); }
+			virtual void SetParentRow(RowHandle Target, RowHandle Parent) override { Implementation.SetParent(*static_cast<IQueryContext*>(this), Target, Parent); }
+			virtual RowHandle GetParentRow(RowHandle Target) const override { return Implementation.GetParent(*static_cast<const IQueryContext*>(this), Target); }
 			void PushCommand(void (*CommandFunction)(void*), void* Context) override { return Implementation.PushCommand(CommandFunction, Context); }
 
 			const UObject* GetDependency(const UClass* DependencyClass) override { return Implementation.GetDependency(DependencyClass); }
@@ -1292,10 +1337,11 @@
 			FQueryResult RunSubquery(int32 SubqueryIndex) override { return Implementation.RunSubquery(SubqueryIndex); }
 			FQueryResult RunSubquery(int32 SubqueryIndex, UE::Editor::DataStorage::SubqueryCallbackRef Callback) override { return Implementation.RunSubquery(SubqueryIndex, Callback); }
 			FQueryResult RunSubquery(int32 SubqueryIndex, RowHandle Row, UE::Editor::DataStorage::SubqueryCallbackRef Callback) override { return Implementation.RunSubquery(SubqueryIndex, Row, Callback); }
-	
+
 protected:
 			void* EmplaceObjectInScratch(const FEmplaceObjectParams& Params) override { return Implementation.EmplaceObjectInScratch(Params.ObjectSize, Params.Alignment, Params.Construct, Params.Destroy, Params.SourceObject); }
-			
+		
+		protected:
 			FMassQueryContextImplementation Implementation;
 		};
 	} // namespace Processors::Private
@@ -1544,8 +1590,9 @@
 		FMassExecutionContext Context(EntityManager);
 		auto ExecuteFunction = [&Result, &Callback, &Description, &Environment](FMassExecutionContext& Context)
 			{
+				using namespace Processors::Private;
 				// No need to cache any subsystem dependencies as these are not accessible from a direct query.
-				UE::Editor::DataStorage::Processors::Private::FMassDirectContextForwarder QueryContext(Context, Environment);
+				FMassDirectContextForwarder QueryContext(Context, Environment, Description);
 				Callback(Description, QueryContext);
 				Result.Count += Context.GetNumEntities();
 			};
@@ -1561,6 +1608,7 @@
 		{
 			NativeQuery.ForEachEntityChunk(Context, ExecuteFunction);
 		}
+		Context.FlushDeferred();
 	}
 	else
 	{
@@ -1592,8 +1640,9 @@
 	NativeQuery.ForEachEntityChunk(Context,
 		[&Result, &Callback, &Description, &Environment](FMassExecutionContext& Context)
 		{
+			using namespace Processors::Private;
 			// No need to cache any subsystem dependencies as these are not accessible from a subquery.
-			UE::Editor::DataStorage::Processors::Private::FMassSubqueryContextForwarder QueryContext(Context, Environment, Description.bShouldBatchModifications);
+			FMassSubqueryContextForwarder QueryContext(Context, Environment, Description);
 			Callback(Description, QueryContext);
 			Result.Count += Context.GetNumEntities();
 		}
@@ -1626,12 +1675,22 @@
 		Context.SetEntityCollection(FMassArchetypeEntityCollection(NativeArchetype, { NativeEntity }, FMassArchetypeEntityCollection::NoDuplicates));
 		Context.SetDeferredCommandBuffer(ParentContext.GetSharedDeferredCommandBuffer());
 		Context.SetFlushDeferredCommands(false);
+		
+		const FTedsHierarchyAccessInterface* HierarchyAccessInterface = nullptr;
+		if (!Description.Hierarchy.IsNone())
+		{
+			FHierarchyHandle Handle = Environment.GetHierarchyRegistrar().FindHierarchyByName(Description.Hierarchy);
+			HierarchyAccessInterface = Environment.GetHierarchyRegistrar().GetAccessInterface(Handle);
+		}
 
 		NativeQuery.ForEachEntityChunk(Context,
-			[&Result, &Callback, &Description, &Environment, &EntityManager, RowHandle](FMassExecutionContext& Context)
+			[&Result, &Callback, &Description, &Environment,  HierarchyAccessInterface](FMassExecutionContext& Context)
 			{
+				using namespace Processors::Private;
 				// No need to cache any subsystem dependencies as these are not accessible from a subquery.
-				Processors::Private::FMassSubqueryContextForwarder QueryContext(Context, Environment, Description.bShouldBatchModifications);
+				FMassSubqueryContextForwarder QueryContext(Context, Environment, Description);
+				QueryContext.SetHierarchyAccessInterface(HierarchyAccessInterface);
+				
 				Callback(Description, QueryContext);
 				Result.Count += Context.GetNumEntities();
 			}
diff --git a/Engine/Plugins/Experimental/EditorDataStorage/Source/TedsCore/Private/Queries/TypedElementExtendedQueryStore.cpp b/Engine/Plugins/Experimental/EditorDataStorage/Source/TedsCore/Private/Queries/TypedElementExtendedQueryStore.cpp
--- a/Engine/Plugins/Experimental/EditorDataStorage/Source/TedsCore/Private/Queries/TypedElementExtendedQueryStore.cpp
+++ b/Engine/Plugins/Experimental/EditorDataStorage/Source/TedsCore/Private/Queries/TypedElementExtendedQueryStore.cpp
@@ -156,6 +156,7 @@ namespace UE::Editor::DataStorage
 		FMassEntityQuery& NativeQuery = SetupNativeQuery(StoredQuery.Description, StoredQuery, Environment);
 
 		bool bContinueSetup =              SetupDynamicColumns(StoredQuery.Description, Environment);
+		bContinueSetup = bContinueSetup && SetupHierarchyDependencies(StoredQuery.Description, Environment);
 		bContinueSetup = bContinueSetup && SetupSelectedColumns(StoredQuery.Description, NativeQuery);
 		bContinueSetup = bContinueSetup && SetupChunkFilters(Result, StoredQuery.Description, Environment, NativeQuery);
 		bContinueSetup = bContinueSetup && SetupConditions(StoredQuery.Description, Environment, NativeQuery);
@@ -647,6 +648,32 @@ namespace UE::Editor::DataStorage
 		return true;
 	}
 
+	bool FExtendedQueryStore::SetupHierarchyDependencies(FQueryDescription& Query, FEnvironment& Environment)
+	{
+		if (!Query.Hierarchy.IsNone())
+		{
+			FHierarchyHandle Handle = Environment.GetHierarchyRegistrar().FindHierarchyByName(Query.Hierarchy);
+			const FTedsHierarchyAccessInterface* HierarchyAccessInterface = Environment.GetHierarchyRegistrar().GetAccessInterface(Handle);
+			ensureMsgf(
+				HierarchyAccessInterface,
+				TEXT("Hierarchy named '%s' not registered. Query will return no hierarchy information. Register query after registering hierarchy"),
+				*Query.Hierarchy.ToString());
+			if (HierarchyAccessInterface)
+			{
+				// Must bind the HierarchyDataColumn for Read to allow for GetParent
+				// Note: It would be better to do late binding to avoid the order dependency on query creation and hierarchy registration
+				// but this would require resetting a query's requirements after it has been bound to Mass or swapping out a new query
+				// when the hierarchy is registered
+				// Note: Optional used to allow a query to match rows that are not currently participating in a relationship (ie. childless root objects)
+				const UScriptStruct* DataColumnType = HierarchyAccessInterface->GetHierarchyDataColumnType();
+				Query.SelectionTypes.Emplace(DataColumnType);
+				Query.SelectionAccessTypes.Emplace(EQueryAccessType::OptionalReadOnly);
+				Query.SelectionMetaData.Emplace(FColumnMetaData(DataColumnType, FColumnMetaData::EFlags::None));
+			}
+		}
+		return true;
+	}
+
 	FMassEntityQuery& FExtendedQueryStore::SetupNativeQuery(
 		FQueryDescription& Query, FExtendedQuery& StoredQuery, FEnvironment& Environment)
 	{
diff --git a/Engine/Plugins/Experimental/EditorDataStorage/Source/TedsCore/Private/Queries/TypedElementExtendedQueryStore.h b/Engine/Plugins/Experimental/EditorDataStorage/Source/TedsCore/Private/Queries/TypedElementExtendedQueryStore.h
--- a/Engine/Plugins/Experimental/EditorDataStorage/Source/TedsCore/Private/Queries/TypedElementExtendedQueryStore.h
+++ b/Engine/Plugins/Experimental/EditorDataStorage/Source/TedsCore/Private/Queries/TypedElementExtendedQueryStore.h
@@ -186,6 +186,7 @@ namespace UE::Editor::DataStorage
 
 		FMassEntityQuery& SetupNativeQuery(FQueryDescription& Query, FExtendedQuery& StoredQuery, FEnvironment& Environment);
 		bool SetupDynamicColumns(FQueryDescription& Query, FEnvironment& Environment);
+		bool SetupHierarchyDependencies(FQueryDescription& Query, FEnvironment& Environment);
 		bool SetupSelectedColumns(FQueryDescription& Query, FMassEntityQuery& NativeQuery);
 		bool SetupConditions(FQueryDescription& Query, FEnvironment& Environment, FMassEntityQuery& NativeQuery);
 		bool SetupChunkFilters(Handle QueryHandle, FQueryDescription& Query, FEnvironment& Environment, FMassEntityQuery& NativeQuery);
diff --git a/Engine/Plugins/Experimental/EditorDataStorage/Source/TedsCore/Private/TypedElementDatabase.cpp b/Engine/Plugins/Experimental/EditorDataStorage/Source/TedsCore/Private/TypedElementDatabase.cpp
--- a/Engine/Plugins/Experimental/EditorDataStorage/Source/TedsCore/Private/TypedElementDatabase.cpp
+++ b/Engine/Plugins/Experimental/EditorDataStorage/Source/TedsCore/Private/TypedElementDatabase.cpp
@@ -1057,6 +1057,98 @@ void UEditorDataStorage::ForEachDynamicColumn(const UScriptStruct& Template, TFu
 	return Environment->ForEachDynamicColumn(Template, Callback);
 }
 
+UE::Editor::DataStorage::FHierarchyHandle UEditorDataStorage::RegisterHierarchy(
+	const UE::Editor::DataStorage::FHierarchyRegistrationParams& Params)
+{
+	return Environment->GetHierarchyRegistrar().RegisterHierarchy(this, Params);
+}
+
+UE::Editor::DataStorage::FHierarchyHandle UEditorDataStorage::FindHierarchyByName(const FName& Name) const
+{
+	return Environment->GetHierarchyRegistrar().FindHierarchyByName(Name);
+}
+
+bool UEditorDataStorage::IsValidHierachyHandle(UE::Editor::DataStorage::FHierarchyHandle Handle) const
+{
+	return Environment->GetHierarchyRegistrar().GetAccessInterface(Handle) != nullptr;
+}
+
+const UScriptStruct* UEditorDataStorage::GetChildTagType(UE::Editor::DataStorage::FHierarchyHandle InHierarchyHandle) const
+{
+	if (const UE::Editor::DataStorage::FTedsHierarchyAccessInterface* AccessInterface = Environment->GetHierarchyRegistrar().GetAccessInterface(InHierarchyHandle))
+	{
+		return AccessInterface->GetChildTagType();
+	}
+	return nullptr;
+}
+
+const UScriptStruct* UEditorDataStorage::GetParentTagType(UE::Editor::DataStorage::FHierarchyHandle InHierarchyHandle) const
+{
+	if (const UE::Editor::DataStorage::FTedsHierarchyAccessInterface* AccessInterface = Environment->GetHierarchyRegistrar().GetAccessInterface(InHierarchyHandle))
+	{
+		return AccessInterface->GetParentTagType();
+	}
+	return nullptr;
+}
+
+const UScriptStruct* UEditorDataStorage::GetHierarchyDataColumnType(UE::Editor::DataStorage::FHierarchyHandle InHierarchyHandle) const
+{
+	if (const UE::Editor::DataStorage::FTedsHierarchyAccessInterface* AccessInterface = Environment->GetHierarchyRegistrar().GetAccessInterface(InHierarchyHandle))
+	{
+		return AccessInterface->GetHierarchyDataColumnType();
+	}
+	return nullptr;
+}
+
+void UEditorDataStorage::SetParentRow(UE::Editor::DataStorage::FHierarchyHandle InHierarchyHandle, UE::Editor::DataStorage::RowHandle Target,
+	UE::Editor::DataStorage::RowHandle Parent)
+{
+	if (const UE::Editor::DataStorage::FTedsHierarchyAccessInterface* AccessInterface = Environment->GetHierarchyRegistrar().GetAccessInterface(InHierarchyHandle))
+	{
+		AccessInterface->SetParentRow(this, Target, Parent);
+	}
+}
+
+UE::Editor::DataStorage::RowHandle UEditorDataStorage::GetParentRow(UE::Editor::DataStorage::FHierarchyHandle InHierarchyHandle,
+	UE::Editor::DataStorage::RowHandle Target) const
+{
+	if (const UE::Editor::DataStorage::FTedsHierarchyAccessInterface* AccessInterface = Environment->GetHierarchyRegistrar().GetAccessInterface(InHierarchyHandle))
+	{
+		return AccessInterface->GetParentRow(this, Target);
+	}
+	return UE::Editor::DataStorage::InvalidRowHandle;
+}
+
+TFunction<UE::Editor::DataStorage::RowHandle(const void*, const UScriptStruct*)> UEditorDataStorage::CreateParentExtractionFunction(
+	UE::Editor::DataStorage::FHierarchyHandle InHierarchyHandle) const
+{
+	if (const UE::Editor::DataStorage::FTedsHierarchyAccessInterface* AccessInterface = Environment->GetHierarchyRegistrar().GetAccessInterface(InHierarchyHandle))
+	{
+		return AccessInterface->CreateParentExtractionFunction();
+	}
+	return [](const void*, const UScriptStruct*) { return UE::Editor::DataStorage::InvalidRowHandle; };
+}
+
+bool UEditorDataStorage::HasChildren(UE::Editor::DataStorage::FHierarchyHandle InHierarchyHandle, UE::Editor::DataStorage::RowHandle Row) const
+{
+	if (const UE::Editor::DataStorage::FTedsHierarchyAccessInterface* AccessInterface = Environment->GetHierarchyRegistrar().GetAccessInterface(InHierarchyHandle))
+	{
+		return AccessInterface->HasChildren(*this, Row);
+	}
+	return false;
+}
+
+void UEditorDataStorage::WalkDepthFirst(
+	UE::Editor::DataStorage::FHierarchyHandle InHierarchyHandle,
+	UE::Editor::DataStorage::RowHandle Row,
+	TFunction<void(const ICoreProvider& Context, UE::Editor::DataStorage::RowHandle Owner, UE::Editor::DataStorage::RowHandle Target)> VisitFn) const
+{
+	if (const UE::Editor::DataStorage::FTedsHierarchyAccessInterface* AccessInterface = Environment->GetHierarchyRegistrar().GetAccessInterface(InHierarchyHandle))
+	{
+		return AccessInterface->WalkDepthFirst(*this, Row, MoveTemp(VisitFn));
+	}
+}
+
 void UEditorDataStorage::RegisterTickGroup(
 	FName GroupName, UE::Editor::DataStorage::EQueryTickPhase Phase, FName BeforeGroup, FName AfterGroup, UE::Editor::DataStorage::EExecutionMode ExecutionMode)
 {
diff --git a/Engine/Plugins/Experimental/EditorDataStorage/Source/TedsCore/Private/TypedElementDatabase.h b/Engine/Plugins/Experimental/EditorDataStorage/Source/TedsCore/Private/TypedElementDatabase.h
--- a/Engine/Plugins/Experimental/EditorDataStorage/Source/TedsCore/Private/TypedElementDatabase.h
+++ b/Engine/Plugins/Experimental/EditorDataStorage/Source/TedsCore/Private/TypedElementDatabase.h
@@ -140,6 +140,28 @@ public:
 	UE_API virtual const UScriptStruct* GenerateDynamicColumn(const UE::Editor::DataStorage::FDynamicColumnDescription& Description) override;
 	UE_API virtual void ForEachDynamicColumn(const UScriptStruct& Template, TFunctionRef<void(const UScriptStruct& Type)> Callback) const override;
 
+	UE_API virtual UE::Editor::DataStorage::FHierarchyHandle RegisterHierarchy(const UE::Editor::DataStorage::FHierarchyRegistrationParams& Params) override;
+	UE_API virtual UE::Editor::DataStorage::FHierarchyHandle FindHierarchyByName(const FName& Name) const override;
+	UE_API virtual bool IsValidHierachyHandle(UE::Editor::DataStorage::FHierarchyHandle) const override;
+	UE_API virtual const UScriptStruct* GetChildTagType(
+		UE::Editor::DataStorage::FHierarchyHandle) const override;
+	UE_API virtual const UScriptStruct* GetParentTagType(
+		UE::Editor::DataStorage::FHierarchyHandle) const override;
+	UE_API virtual const UScriptStruct* GetHierarchyDataColumnType(
+		UE::Editor::DataStorage::FHierarchyHandle) const override;
+	UE_API virtual void SetParentRow(
+		UE::Editor::DataStorage::FHierarchyHandle, UE::Editor::DataStorage::RowHandle Target,	UE::Editor::DataStorage::RowHandle Parent) override;
+	UE_API virtual UE::Editor::DataStorage::RowHandle GetParentRow(
+		UE::Editor::DataStorage::FHierarchyHandle, UE::Editor::DataStorage::RowHandle Target) const override;
+	UE_API virtual TFunction<UE::Editor::DataStorage::RowHandle(const void*, const UScriptStruct*)> CreateParentExtractionFunction(
+		UE::Editor::DataStorage::FHierarchyHandle) const override;
+	UE_API virtual bool HasChildren(
+		UE::Editor::DataStorage::FHierarchyHandle, UE::Editor::DataStorage::RowHandle Row) const override;
+	UE_API virtual void WalkDepthFirst(
+		UE::Editor::DataStorage::FHierarchyHandle,
+		UE::Editor::DataStorage::RowHandle Row, 
+		TFunction<void(const ICoreProvider& Context, UE::Editor::DataStorage::RowHandle Owner, UE::Editor::DataStorage::RowHandle Target)> VisitFn) const override;
+
 	UE_API void RegisterTickGroup(FName GroupName, UE::Editor::DataStorage::EQueryTickPhase Phase, FName BeforeGroup, FName AfterGroup, UE::Editor::DataStorage::EExecutionMode ExecutionMode);
 	UE_API void UnregisterTickGroup(FName GroupName, UE::Editor::DataStorage::EQueryTickPhase Phase);
 
@@ -187,7 +209,10 @@ private:
 	UE_API void Reset();
 
 	UE_API int32 GetTableChunkSize(FName TableName) const;
-	
+
+public:
+
+private:
 	struct FFactoryTypePair
 	{
 		// Used to find the factory by type without needing to dereference each one
diff --git a/Engine/Plugins/Experimental/EditorDataStorage/Source/TedsCore/Private/TypedElementDatabaseEnvironment.cpp b/Engine/Plugins/Experimental/EditorDataStorage/Source/TedsCore/Private/TypedElementDatabaseEnvironment.cpp
--- a/Engine/Plugins/Experimental/EditorDataStorage/Source/TedsCore/Private/TypedElementDatabaseEnvironment.cpp
+++ b/Engine/Plugins/Experimental/EditorDataStorage/Source/TedsCore/Private/TypedElementDatabaseEnvironment.cpp
@@ -79,6 +79,16 @@ namespace UE::Editor::DataStorage
 		return MassEntityManager;
 	}
 
+	FTedsHierarchyRegistrar& FEnvironment::GetHierarchyRegistrar()
+	{
+		return HierarchyInterface;
+	}
+
+	const FTedsHierarchyRegistrar& FEnvironment::GetHierarchyRegistrar() const
+	{
+		return HierarchyInterface;
+	}
+
 	FMassArchetypeHandle FEnvironment::LookupMassArchetype(TableHandle Table) const
 	{
 		return DataStorage.LookupArchetype(Table);
diff --git a/Engine/Plugins/Experimental/EditorDataStorage/Source/TedsCore/Private/TypedElementDatabaseEnvironment.h b/Engine/Plugins/Experimental/EditorDataStorage/Source/TedsCore/Private/TypedElementDatabaseEnvironment.h
--- a/Engine/Plugins/Experimental/EditorDataStorage/Source/TedsCore/Private/TypedElementDatabaseEnvironment.h
+++ b/Engine/Plugins/Experimental/EditorDataStorage/Source/TedsCore/Private/TypedElementDatabaseEnvironment.h
@@ -9,6 +9,7 @@
 #include "TypedElementDatabaseCommandBuffer.h"
 #include "TypedElementDatabaseScratchBuffer.h"
 #include "TypedElementDatabaseIndexTable.h"
+#include "Hierarchy/EditorDataStorageHierarchyImplementation.h"
 #include "Memento/TypedElementMementoSystem.h"
 #include "Queries/TypedElementExtendedQueryStore.h"
 
@@ -40,6 +41,9 @@ namespace UE::Editor::DataStorage
 		FMassEntityManager& GetMassEntityManager();
 		const FMassEntityManager& GetMassEntityManager() const;
 
+		FTedsHierarchyRegistrar& GetHierarchyRegistrar();
+		const FTedsHierarchyRegistrar& GetHierarchyRegistrar() const;
+
 		FMassArchetypeHandle LookupMassArchetype(TableHandle Table) const;
 
 		FMassProcessingPhaseManager& GetMassPhaseManager();
@@ -85,6 +89,7 @@ namespace UE::Editor::DataStorage
 		FExtendedQueryStore Queries;
 		FMementoSystem MementoSystem;
 		FValueTagManager ValueTagManager;
+		FTedsHierarchyRegistrar HierarchyInterface;
 
 		FMutex CommandQueueMutex;
 		TArray<FEnvironmentCommand> CommandQueue;
diff --git a/Engine/Plugins/Experimental/EditorDataStorageFeatures/Source/TedsEverythingPicker/Private/Context/TedsPickerContextUtil.cpp b/Engine/Plugins/Experimental/EditorDataStorageFeatures/Source/TedsEverythingPicker/Private/Context/TedsPickerContextUtil.cpp
--- a/Engine/Plugins/Experimental/EditorDataStorageFeatures/Source/TedsEverythingPicker/Private/Context/TedsPickerContextUtil.cpp
+++ b/Engine/Plugins/Experimental/EditorDataStorageFeatures/Source/TedsEverythingPicker/Private/Context/TedsPickerContextUtil.cpp
@@ -44,7 +44,7 @@ namespace UE::Editor::DataStorage::Picker
 		// Using the Teds Outliner Purpose for matching customization for now
 		IUiProvider::FPurposeID PurposeId = IUiProvider::FPurposeInfo("SceneOutliner", "RowLabel", NAME_None).GeneratePurposeID();
 
-		auto ParentExtractionFunction = [](const void* Column, UScriptStruct*)->RowHandle
+		auto ParentExtractionFunction = [](const void* Column, const UScriptStruct*)->RowHandle
 		{
 			const FTableRowParentColumn* ParentColumn = static_cast<const FTableRowParentColumn*>(Column);
 			return ParentColumn->Parent;
diff --git a/Engine/Plugins/Experimental/EditorDataStorageFeatures/Source/TedsQueryStack/Private/TedsHierarchyNode.cpp b/Engine/Plugins/Experimental/EditorDataStorageFeatures/Source/TedsQueryStack/Private/TedsHierarchyNode.cpp
new file mode 100644
--- /dev/null
+++ b/Engine/Plugins/Experimental/EditorDataStorageFeatures/Source/TedsQueryStack/Private/TedsHierarchyNode.cpp
@@ -0,0 +1,92 @@
+﻿// Copyright Epic Games, Inc. All Rights Reserved.
+
+#include "TedsHierarchyNode.h"
+
+#include "Elements/Interfaces/TypedElementDataStorageInterface.h"
+
+
+namespace UE::Editor::DataStorage
+{
+	FHierarchyRowNode::FHierarchyRowNode(
+		ICoreProvider& Storage,
+		FHierarchyHandle InHierarchyHandle,
+		const TSharedPtr<QueryStack::IRowNode>& InTopLevelRows,
+		ESyncFlags InSyncFlags)
+		: TopLevelRows(InTopLevelRows)
+		, Storage(Storage)
+		, HierarchyHandle(InHierarchyHandle)
+		, SyncFlags(InSyncFlags)
+	{
+	}
+
+	QueryStack::INode::RevisionId FHierarchyRowNode::GetRevision() const
+	{
+		return Revision;
+	}
+
+	void FHierarchyRowNode::Update()
+	{
+		
+		if (TopLevelRows->GetRevision() != ParentRevision)
+		{
+			Rows.Empty();
+			RefreshInternal(Storage, Rows);
+			++Revision;
+			ParentRevision = TopLevelRows->GetRevision();
+		}
+		if (EnumHasAllFlags(SyncFlags, ESyncFlags::IncrementWhenDifferent))
+		{
+			FRowHandleArray& TempArray = GetTempRowHandleArray();
+			RefreshInternal(Storage, TempArray);
+			TempArray.Sort();
+			Rows.Sort();
+			if (!TempArray.IsSame(Rows))
+			{
+				Rows = MoveTemp(TempArray);
+				++Revision;
+				ParentRevision = TopLevelRows->GetRevision();
+			}
+		}
+		else
+		{
+			Rows.Empty();
+			RefreshInternal(Storage, Rows);
+			++Revision;
+			ParentRevision = TopLevelRows->GetRevision();
+		}
+	}
+
+	FRowHandleArrayView FHierarchyRowNode::GetRows() const
+	{
+		return Rows.GetRows();
+	}
+
+	FRowHandleArray& FHierarchyRowNode::GetMutableRows()
+	{
+		return Rows;
+	}
+
+	void FHierarchyRowNode::RefreshInternal(ICoreProvider& InCoreProvider, FRowHandleArray& InRows) const
+	{
+		InRows.Empty();
+		
+		auto VisitFn = [&InRows](const ICoreProvider&, RowHandle, RowHandle TargetRow)
+		{
+			InRows.Add(TargetRow);
+		};
+
+		// Top level rows are already going to have the VisitFn run on them
+		FRowHandleArrayView TopLevelRowsView = TopLevelRows->GetRows();
+		for (RowHandle Row : TopLevelRowsView)
+		{
+			// Walk recursively down through the hierarchy and add rows 
+			Storage.WalkDepthFirst(HierarchyHandle, Row, VisitFn);
+		}
+	}
+
+	FRowHandleArray& FHierarchyRowNode::GetTempRowHandleArray()
+	{
+		NewRows.Empty();
+		return NewRows;
+	}
+}
diff --git a/Engine/Plugins/Experimental/EditorDataStorageFeatures/Source/TedsQueryStack/Public/TedsHierarchyNode.h b/Engine/Plugins/Experimental/EditorDataStorageFeatures/Source/TedsQueryStack/Public/TedsHierarchyNode.h
new file mode 100644
--- /dev/null
+++ b/Engine/Plugins/Experimental/EditorDataStorageFeatures/Source/TedsQueryStack/Public/TedsHierarchyNode.h
@@ -0,0 +1,59 @@
+﻿// Copyright Epic Games, Inc. All Rights Reserved.
+
+#pragma once
+#include "TedsRowQueryResultsNode.h"
+
+namespace UE::Editor::DataStorage
+{
+	// Row which enumerates and stores all children of the given TopLevelRows (including the TopLevelRows)
+	//
+	// Warning: The resulting row array is not guaranteed to be sorted in depth order
+	class FHierarchyRowNode final : public QueryStack::IRowNode
+	{
+	public:
+		enum class ESyncFlags
+		{
+			None = 0,
+			/** 
+			 * Update the row list whenever "Update" is called. This can take several milliseconds for large lists. It's
+			 * recommended to use a monitoring node to detect changes for large lists.
+			 */
+			RefreshOnUpdate = 1 << 0,
+			/** Update the row list whenever the parent changes. */
+			RefreshOnParentChange = 1 << 1,
+			/** 
+			 * Compares the previous row list with the current and only updates the revision if there are differences.
+			 */
+			IncrementWhenDifferent = 1 << 2
+		};
+		using EmitRowFn = TFunctionRef<void(TArrayView<const RowHandle>)>;
+		
+		TEDSQUERYSTACK_API FHierarchyRowNode(
+			ICoreProvider& Storage,
+			FHierarchyHandle InHierarchyHandle,
+			const TSharedPtr<QueryStack::IRowNode>& InTopLevelRows,
+			ESyncFlags SyncFlags);
+		
+		TEDSQUERYSTACK_API virtual RevisionId GetRevision() const override;
+		TEDSQUERYSTACK_API virtual void Update() override;
+		TEDSQUERYSTACK_API virtual FRowHandleArrayView GetRows() const override;
+		TEDSQUERYSTACK_API virtual FRowHandleArray& GetMutableRows() override;
+
+	private:
+		void RefreshInternal(ICoreProvider& InCoreProvider, FRowHandleArray& NewRows) const;
+		FRowHandleArray& GetTempRowHandleArray();
+		
+		FRowHandleArray Rows;
+		FRowHandleArray NewRows;
+		const TSharedPtr<QueryStack::IRowNode> TopLevelRows;
+		RevisionId ParentRevision = 0;
+		
+		ICoreProvider& Storage;
+		FHierarchyHandle HierarchyHandle;
+		
+		RevisionId Revision = 0;
+		ESyncFlags SyncFlags;
+	};
+
+	ENUM_CLASS_FLAGS(FHierarchyRowNode::ESyncFlags);
+}
diff --git a/Engine/Plugins/Experimental/EditorDataStorageFeatures/Source/TedsTableViewer/Private/Widgets/STedsHierarchyViewer.cpp b/Engine/Plugins/Experimental/EditorDataStorageFeatures/Source/TedsTableViewer/Private/Widgets/STedsHierarchyViewer.cpp
--- a/Engine/Plugins/Experimental/EditorDataStorageFeatures/Source/TedsTableViewer/Private/Widgets/STedsHierarchyViewer.cpp
+++ b/Engine/Plugins/Experimental/EditorDataStorageFeatures/Source/TedsTableViewer/Private/Widgets/STedsHierarchyViewer.cpp
@@ -20,8 +20,8 @@ namespace UE::Editor::DataStorage
 {
 	void SHierarchyViewer::Construct(
 		const FArguments& InArgs,
-		UScriptStruct* InHierarchyDataColumnType,
-		TFunction<RowHandle(const void*, UScriptStruct*)> InGetRowHandleFn)
+		const UScriptStruct* InHierarchyDataColumnType,
+		TFunction<RowHandle(const void*, const UScriptStruct*)> InGetRowHandleFn)
 	{
 		checkf(InHierarchyDataColumnType != nullptr, TEXT("Nullptr ParentColumnType"));
 		HierarchyColumnType = InHierarchyDataColumnType;
diff --git a/Engine/Plugins/Experimental/EditorDataStorageFeatures/Source/TedsTableViewer/Public/Widgets/STedsHierarchyViewer.h b/Engine/Plugins/Experimental/EditorDataStorageFeatures/Source/TedsTableViewer/Public/Widgets/STedsHierarchyViewer.h
--- a/Engine/Plugins/Experimental/EditorDataStorageFeatures/Source/TedsTableViewer/Public/Widgets/STedsHierarchyViewer.h
+++ b/Engine/Plugins/Experimental/EditorDataStorageFeatures/Source/TedsTableViewer/Public/Widgets/STedsHierarchyViewer.h
@@ -90,8 +90,8 @@ namespace UE::Editor::DataStorage
 		
 		TEDSTABLEVIEWER_API void Construct(
 			const FArguments& InArgs,
-			UScriptStruct* InHierarchyDataColumnType,
-			TFunction<RowHandle(const void*, UScriptStruct*)> InGetRowHandleFn);
+			const UScriptStruct* InHierarchyDataColumnType,
+			TFunction<RowHandle(const void*, const UScriptStruct*)> InGetRowHandleFn);
 
 		// Clear the current list of columns being displayed and set it to the given list
 		TEDSTABLEVIEWER_API void SetColumns(const TArray<TWeakObjectPtr<const UScriptStruct>>& Columns);
@@ -169,10 +169,10 @@ namespace UE::Editor::DataStorage
 		TArray<TableViewerItemPtr> TopLevelRows;
 
 		// Column type that holds hierarchy information for the row
-		UScriptStruct* HierarchyColumnType = nullptr;
+		const UScriptStruct* HierarchyColumnType = nullptr;
 
 		// Function to get the parent from a the HierarchyColumnType of the row
-		TFunction<RowHandle(const void*, UScriptStruct*)> GetParentFromColumnFn;
+		TFunction<RowHandle(const void*, const UScriptStruct*)> GetParentFromColumnFn;
 
 	};
 }
\ No newline at end of file
diff --git a/Engine/Source/Runtime/TypedElementFramework/Private/Elements/Framework/TypedElementQueryBuilder.cpp b/Engine/Source/Runtime/TypedElementFramework/Private/Elements/Framework/TypedElementQueryBuilder.cpp
--- a/Engine/Source/Runtime/TypedElementFramework/Private/Elements/Framework/TypedElementQueryBuilder.cpp
+++ b/Engine/Source/Runtime/TypedElementFramework/Private/Elements/Framework/TypedElementQueryBuilder.cpp
@@ -521,6 +521,12 @@ namespace UE::Editor::DataStorage::Queries
 		return FDependency{ &Query };
 	}
 
+	Select& Select::AccessesHierarchy(const FName& HierachyName)
+	{
+		Query.Hierarchy = HierachyName;
+		return *this;
+	}
+
 	FQueryConditionQuery Select::Where(const Queries::FConditions& Condition)
 	{
 		Query.Conditions.Emplace(Condition);
diff --git a/Engine/Source/Runtime/TypedElementFramework/Private/Elements/Framework/TypedElementRowHandleArray.cpp b/Engine/Source/Runtime/TypedElementFramework/Private/Elements/Framework/TypedElementRowHandleArray.cpp
--- a/Engine/Source/Runtime/TypedElementFramework/Private/Elements/Framework/TypedElementRowHandleArray.cpp
+++ b/Engine/Source/Runtime/TypedElementFramework/Private/Elements/Framework/TypedElementRowHandleArray.cpp
@@ -421,6 +421,64 @@ namespace UE::Editor::DataStorage
 		return Rows.IsEmpty();
 	}
 
+	bool FRowHandleArray::IsSame(const FRowHandleArray& Rhs) const
+	{
+		if (this == &Rhs)
+		{
+			return true;
+		}
+
+		if (Rows.Num() == Rhs.Num())
+		{
+			const bool bRhsIsSorted = Rhs.IsSorted();
+			if (bIsSorted && bRhsIsSorted)
+			{
+				// If both lists are sorted do a fast memory compare.
+				const SIZE_T CompareByteCount = static_cast<SIZE_T>(Rows.NumBytes());
+				int32 Result = FMemory::Memcmp(Rows.GetData(), Rhs.GetRows().GetData(), CompareByteCount); 
+				return Result == 0;
+			}
+			else if (bRhsIsSorted)
+			{
+				// If the right side is sorted (meaning the left side isn't) prefer to search in the right side
+				// as it allows a binary search.
+				for (RowHandle Row : Rows)
+				{
+					if (Algo::BinarySearch(Rhs.Rows, Row) < 0)
+					{
+						return false;
+					}
+				}
+				return true;
+			}
+			else if (bIsSorted)
+			{
+				// Same as the previous but now prefer the left as that's the sorted list.
+				for (RowHandle RhsRow : Rhs.GetRows())
+				{
+					if (Algo::BinarySearch(Rows, RhsRow) < 0)
+					{
+						return false;
+					}
+				}
+				return true;
+			}
+			else
+			{
+				// Both lists are unsorted so do a linear search through one of them.
+				for (RowHandle RhsRow : Rhs.GetRows())
+				{
+					if (Algo::IndexOf(Rows, RhsRow) < 0)
+					{
+						return false;
+					}
+				}
+				return true;
+			}
+		}
+		return false;
+	}
+
 	void FRowHandleArray::SortedMergeInternal(FRowHandleArrayView AddedRows)
 	{
 		// There are two paths, one is an in-place addition if the added rows fit within the existing array
diff --git a/Engine/Source/Runtime/TypedElementFramework/Public/DataStorage/Handles.h b/Engine/Source/Runtime/TypedElementFramework/Public/DataStorage/Handles.h
--- a/Engine/Source/Runtime/TypedElementFramework/Public/DataStorage/Handles.h
+++ b/Engine/Source/Runtime/TypedElementFramework/Public/DataStorage/Handles.h
@@ -17,6 +17,11 @@ namespace UE::Editor::DataStorage
 
 	using QueryHandle = uint64;
 	static constexpr QueryHandle InvalidQueryHandle = TNumericLimits<QueryHandle>::Max();
+
+	struct FHierarchyHandle
+	{
+		uint64 Reserved = 0;
+	};
 } // namespace UE::Editor::DataStorage
 
 /*
diff --git a/Engine/Source/Runtime/TypedElementFramework/Public/DataStorage/Queries/Description.h b/Engine/Source/Runtime/TypedElementFramework/Public/DataStorage/Queries/Description.h
--- a/Engine/Source/Runtime/TypedElementFramework/Public/DataStorage/Queries/Description.h
+++ b/Engine/Source/Runtime/TypedElementFramework/Public/DataStorage/Queries/Description.h
@@ -131,6 +131,8 @@ namespace UE::Editor::DataStorage
 		TArray<QueryHandle> Subqueries;
 		FMetaData MetaData;
 
+		FName Hierarchy;
+
 		EActionType Action;
 		bool bShouldBatchModifications = false;
 	};
diff --git a/Engine/Source/Runtime/TypedElementFramework/Public/DataStorage/Queries/Types.h b/Engine/Source/Runtime/TypedElementFramework/Public/DataStorage/Queries/Types.h
--- a/Engine/Source/Runtime/TypedElementFramework/Public/DataStorage/Queries/Types.h
+++ b/Engine/Source/Runtime/TypedElementFramework/Public/DataStorage/Queries/Types.h
@@ -94,7 +94,9 @@ namespace UE::Editor::DataStorage
 		/** Provide read-only access, but do not include when finding matching tables. */
 		OptionalReadOnly,
 		/** Provide read and write access. */
-		ReadWrite
+		ReadWrite,
+		/**/
+		NUM
 	};
 
 	enum class EQueryDependencyFlags : uint8
diff --git a/Engine/Source/Runtime/TypedElementFramework/Public/Elements/Framework/TypedElementQueryBuilder.h b/Engine/Source/Runtime/TypedElementFramework/Public/Elements/Framework/TypedElementQueryBuilder.h
--- a/Engine/Source/Runtime/TypedElementFramework/Public/Elements/Framework/TypedElementQueryBuilder.h
+++ b/Engine/Source/Runtime/TypedElementFramework/Public/Elements/Framework/TypedElementQueryBuilder.h
@@ -433,6 +433,7 @@ namespace UE::Editor::DataStorage::Queries
 		TYPEDELEMENTFRAMEWORK_API FQueryDescription&& Compile();
 		TYPEDELEMENTFRAMEWORK_API FSimpleQuery Where();
 		TYPEDELEMENTFRAMEWORK_API FDependency DependsOn();
+		TYPEDELEMENTFRAMEWORK_API Select& AccessesHierarchy(const FName& HierarchyName);
 
 	private:
 		FQueryDescription Query;
diff --git a/Engine/Source/Runtime/TypedElementFramework/Public/Elements/Framework/TypedElementRowHandleArray.h b/Engine/Source/Runtime/TypedElementFramework/Public/Elements/Framework/TypedElementRowHandleArray.h
--- a/Engine/Source/Runtime/TypedElementFramework/Public/Elements/Framework/TypedElementRowHandleArray.h
+++ b/Engine/Source/Runtime/TypedElementFramework/Public/Elements/Framework/TypedElementRowHandleArray.h
@@ -138,6 +138,8 @@ namespace UE::Editor::DataStorage
 		TYPEDELEMENTFRAMEWORK_API bool IsUnique() const;
 		/** Whether or not there are any values in the view. */
 		TYPEDELEMENTFRAMEWORK_API bool IsEmpty() const;
+		/** Whether or not the two arrays contain the same values. */
+		TYPEDELEMENTFRAMEWORK_API bool IsSame(const FRowHandleArray& Other) const;
 
 	private:
 		void SortedMergeInternal(FRowHandleArrayView AddedRows);
diff --git a/Engine/Source/Runtime/TypedElementFramework/Public/Elements/Interfaces/TypedElementDataStorageInterface.h b/Engine/Source/Runtime/TypedElementFramework/Public/Elements/Interfaces/TypedElementDataStorageInterface.h
--- a/Engine/Source/Runtime/TypedElementFramework/Public/Elements/Interfaces/TypedElementDataStorageInterface.h
+++ b/Engine/Source/Runtime/TypedElementFramework/Public/Elements/Interfaces/TypedElementDataStorageInterface.h
@@ -51,6 +51,11 @@ struct TTypedElementColumnTypeList
 	operator TConstArrayView<const UScriptStruct*>() const { return ColumnTypes; }
 };
 
+struct FHierarchyRegistrationParams
+{
+	FName Name;
+};
+	
 class ICoreProvider : public IModularFeature
 {
 public:
@@ -357,6 +362,48 @@ public:
 	template<TDynamicColumnTemplate DynamicColumnTemplateType>
 	void ForEachDynamicColumn(TFunctionRef<void(const UScriptStruct& Type)> Callback) const;
 
+	// Hierarchy Interface
+	//-----------------------
+	
+	virtual UE::Editor::DataStorage::FHierarchyHandle RegisterHierarchy(const UE::Editor::DataStorage::FHierarchyRegistrationParams& Params) = 0;
+	virtual UE::Editor::DataStorage::FHierarchyHandle FindHierarchyByName(const FName& Name) const = 0;
+	virtual bool IsValidHierachyHandle(FHierarchyHandle) const = 0;
+
+		// Gets the tag type indicating the row is a child in this hierarchy and thus has a parent
+	virtual const UScriptStruct* GetChildTagType(FHierarchyHandle InHierarchyHandle) const = 0;
+	// Gets the tag type indicating the row is a parent in this hierarchy and thus has a child
+	virtual const UScriptStruct* GetParentTagType(FHierarchyHandle InHierarchyHandle) const = 0;
+	// Gets the column type storing the hierarchy data. Note that the format of this type is opaque
+	// and is used internally.  It is the caller's responsibility to provide at least Read access
+	// to this column type in query requirements in order to use the HierarchyAccessInterface on the row
+	virtual const UScriptStruct* GetHierarchyDataColumnType(FHierarchyHandle InHierarchyHandle) const = 0;
+
+	// CoreProvider Context API
+	// =========================
+
+	// Establishes a parent relationship between the Target and the Parent
+	virtual void SetParentRow(
+		FHierarchyHandle InHierarchyHandle,
+		RowHandle Target,
+		RowHandle Parent) = 0;
+
+	// Gets the parent of a target, if there is one
+	virtual RowHandle GetParentRow(FHierarchyHandle InHierarchyHandle, RowHandle Target) const = 0;
+	
+	// Returns a callable which will extract the parent row from the hierarchy's HierarchyDataColumn
+	// Note that the second parameter must be the same as what is returned by GetHierarchyDataColumn()
+	// and the first parameter should point to a struct of that type
+	virtual TFunction<RowHandle(const void*, const UScriptStruct*)> CreateParentExtractionFunction(FHierarchyHandle InHierarchyHandle) const = 0; 
+	
+	virtual bool HasChildren(FHierarchyHandle InHierarchyHandle, RowHandle Row) const = 0;
+
+	// Iterates depth first from the passed in Row, calling the VisitFn on each row
+	// The Row passed in will have VisitFn called on it
+	virtual void WalkDepthFirst(
+		FHierarchyHandle InHierarchyHandle,
+		RowHandle Row,
+		TFunction<void(const ICoreProvider& Context, RowHandle Owner, RowHandle Target)> VisitFn) const = 0;
+	
 	/**
 	 * Outputs the registered query callbacks to the given output device for debugging purposes.
 	 */
diff --git a/Engine/Source/Runtime/TypedElementFramework/Public/Elements/Interfaces/TypedElementQueryStorageInterfaces.h b/Engine/Source/Runtime/TypedElementFramework/Public/Elements/Interfaces/TypedElementQueryStorageInterfaces.h
--- a/Engine/Source/Runtime/TypedElementFramework/Public/Elements/Interfaces/TypedElementQueryStorageInterfaces.h
+++ b/Engine/Source/Runtime/TypedElementFramework/Public/Elements/Interfaces/TypedElementQueryStorageInterfaces.h
@@ -104,8 +104,50 @@ namespace UE::Editor::DataStorage
 		template<TDynamicColumnTemplate TemplateType>
 		const UScriptStruct* FindDynamicColumnType(const FName& Identifier) const;
 
+		// Sets the Target's Parent if the query has a configured Hierarchy
+		virtual void SetParentRow(
+			RowHandle Target,
+			RowHandle Parent) = 0;
+
+		// Gets the Target's Parent if the query has a configured Hierarchy
+		virtual RowHandle GetParentRow(
+			RowHandle Target) const = 0;
+
 		/** Returns the amount of time now and the last update in seconds. */
 		virtual float GetDeltaTimeSeconds() const = 0;
+
+		/**
+		 * Creates a command that will run at the termination of an execution context.
+		 * Intended to be used for cases where the query callback requires doing some work that needs to be serialized
+		 * with respect to other callback invocations.
+		 * Note that commands will be executed in order with respect to the thread that pushed them and will be executed
+		 * on the query exection thread serially.
+		 *
+		 * Usage:
+		 *  Define a command struct with a mutable operator() overload
+		 *  struct FMyCommand
+		 *  {
+		 *      void operator()() { DoSomethingWithSideEffects(MyActor) };
+		 *      TWeakObjectPtr<AActor> MyActor;
+		 *  };
+		 *
+		 *  Context.PushCommand(FMyCommand{ .MyActor = Actor });
+		 */
+		template<typename T>
+		void PushCommand(T CommandContext);
+		
+		virtual void PushCommand(void (*CommandFunction)(void* /*CommandData*/), void* InCommandData) = 0;
+	protected:
+		struct FEmplaceObjectParams
+		{
+			size_t ObjectSize;
+			size_t Alignment;
+			void (*Construct)(void*, void*);
+			void(*Destroy)(void*);
+			void* SourceObject;
+		};
+		// Required as the interface doesn't have direct access to any temporary handling of memory in the backend.
+		virtual void* EmplaceObjectInScratch(const FEmplaceObjectParams& Params) = 0;
 	};
 
 	struct ICommonQueryWithEnvironmentContext : public ICommonQueryContext
@@ -278,39 +320,6 @@ namespace UE::Editor::DataStorage
 		 * tick group.
 		 */
 		virtual void RemoveColumns(TConstArrayView<RowHandle> Rows, TConstArrayView<const UScriptStruct*> ColumnTypes) = 0;
-		
-		/**
-		 * Creates a command that will run immediately after all processors have completed.
-		 * Intended to be used for cases where the query callback calls into something that cannot run while
-		 * the TEDS processors are running.
-		 * Note that commands will be executed in order with respect to the thread that pushed them and will be executed
-		 * on the game thread.
-		 *
-		 * Usage:
-		 *  Define a command struct with a mutable operator() overload
-		 *  struct FMyCommand
-		 *  {
-		 *      void operator()() { DoSomethingWithSideEffects(MyActor) };
-		 *      TWeakObjectPtr<AActor> MyActor;
-		 *  };
-		 *
-		 *  Context.PushCommand(FMyCommand{ .MyActor = Actor });
-		 */
-		template<typename T>
-		void PushCommand(T CommandContext);
-		
-		virtual void PushCommand(void (*CommandFunction)(void* /*CommandData*/), void* InCommandData) = 0;
-	protected:
-		struct FEmplaceObjectParams
-		{
-			size_t ObjectSize;
-			size_t Alignment;
-			void (*Construct)(void*, void*);
-			void(*Destroy)(void*);
-			void* SourceObject;
-		};
-		// Required as the interface doesn't not have direct access to any temporary handling of memory in the backend.
-		virtual void* EmplaceObjectInScratch(const FEmplaceObjectParams& Params) = 0;
 	};
 
 	enum class EDirectQueryExecutionFlags : uint32
@@ -358,6 +367,7 @@ namespace UE::Editor::DataStorage
 		virtual const UObject* GetDependency(const UClass* DependencyClass) = 0;
 		/** Returns a mutable instance of the requested dependency or a nullptr if not found. */
 		virtual UObject* GetMutableDependency(const UClass* DependencyClass) = 0;
+		
 		/**
 		 * Returns a list of dependencies or nullptrs if a dependency wasn't found. Mutable versions are return and it's up to the
 		 * caller to not change immutable dependencies.
@@ -396,6 +406,8 @@ namespace UE::Editor::DataStorage
 		inline RowHandle FindIndexedRow(IndexHash Index) const;
 		UE_DEPRECATED(5.7, "Use the version of this 'LookupMappedRow' that uses a domain.")
 		inline RowHandle LookupMappedRow(const FMapKeyView& Index) const;*/
+
+		
 	};
 } // namespace UE::Editor::DataStorage
 
@@ -643,7 +655,7 @@ namespace UE::Editor::DataStorage
 	}
 
 	template <typename T>
-	void ICommonQueryWithEnvironmentContext::PushCommand(T CommandContext)
+	void ICommonQueryContext::PushCommand(T CommandContext)
 	{
 		if constexpr(std::is_empty_v<T>)
 		{
