Disable TEDS in commandlets
* Specifically cooking does not run through the Editor Tickables of which TEDS and Mass are.  This can lead to a large backlog of commands in the command queue that are never cleared.

diff --git a/Engine/Plugins/Experimental/EditorDataStorage/Source/TedsCore/Private/TypedElementsDataStorage.cpp b/Engine/Plugins/Experimental/EditorDataStorage/Source/TedsCore/Private/TypedElementsDataStorage.cpp
--- a/Engine/Plugins/Experimental/EditorDataStorage/Source/TedsCore/Private/TypedElementsDataStorage.cpp
+++ b/Engine/Plugins/Experimental/EditorDataStorage/Source/TedsCore/Private/TypedElementsDataStorage.cpp
@@ -76,6 +76,13 @@ void FEditorDataStorageModule::StartupModule()
 	
 	ImpersonateMassTagsAndFragments();
 
+	// TEDS relies on ticking Mass for processors to run and to pump command queues
+	// Commandlets do not run the same loop as the editor and does not set up tickables to execute each frame.
+	if (IsRunningCommandlet())
+	{
+		return;
+	}
+
 	FCoreDelegates::OnAllModuleLoadingPhasesComplete.AddLambda(
 		[this]()
 		{
