EditorDataStorageHierarchy: Add functionality to set unresolved parents for rows. This adds a temporary column to the row that gets resolved when the parent row is reigstered with TEDS.

diff --git a/Engine/Plugins/Experimental/EditorDataStorage/Source/TedsCore/Private/Hierarchy/EditorDataStorageHierarchyColumns.h b/Engine/Plugins/Experimental/EditorDataStorage/Source/TedsCore/Private/Hierarchy/EditorDataStorageHierarchyColumns.h
--- a/Engine/Plugins/Experimental/EditorDataStorage/Source/TedsCore/Private/Hierarchy/EditorDataStorageHierarchyColumns.h
+++ b/Engine/Plugins/Experimental/EditorDataStorage/Source/TedsCore/Private/Hierarchy/EditorDataStorageHierarchyColumns.h
@@ -3,6 +3,7 @@
 #pragma once
 
 #include "DataStorage/CommonTypes.h"
+#include "DataStorage/MapKey.h"
 
 #include "EditorDataStorageHierarchyColumns.generated.h"
 
@@ -30,9 +31,23 @@ struct FEditorDataHierarchyChildTag_Template : public FEditorDataStorageTag
 	GENERATED_BODY()
 };
 
+USTRUCT(meta = (DisplayName = "Unresolved Parent (Template)", EditorDataStorage_DynamicColumnTemplate))
+struct FEditorDataHierarchyUnresolvedParent_Template final : public FEditorDataStorageColumn
+{
+	GENERATED_BODY()
+
+	// The Id used to lookup the parent row every frame
+	UE::Editor::DataStorage::FMapKey ParentId;
+	
+	// The mapping domain the ParentID is looked up in
+	FName MappingDomain;
+	
+};
+
 namespace UE::Editor::DataStorage
 {
 	using FHierarchyData_Template = FEditorDataHierarchyData_Template;
 	using FHierarchyParentTag_Template = FEditorDataHierarchyParentTag_Template;
 	using FHierarchyChildTag_Template = FEditorDataHierarchyChildTag_Template;
+	using FHierarchyUnresolvedParentColumn_Template = FEditorDataHierarchyUnresolvedParent_Template;
 }
diff --git a/Engine/Plugins/Experimental/EditorDataStorage/Source/TedsCore/Private/Hierarchy/EditorDataStorageHierarchyImplementation.cpp b/Engine/Plugins/Experimental/EditorDataStorage/Source/TedsCore/Private/Hierarchy/EditorDataStorageHierarchyImplementation.cpp
--- a/Engine/Plugins/Experimental/EditorDataStorage/Source/TedsCore/Private/Hierarchy/EditorDataStorageHierarchyImplementation.cpp
+++ b/Engine/Plugins/Experimental/EditorDataStorage/Source/TedsCore/Private/Hierarchy/EditorDataStorageHierarchyImplementation.cpp
@@ -79,6 +79,16 @@ namespace UE::Editor::DataStorage
 			};
 			return InProvider->GenerateDynamicColumn(Desc);
 		}();
+
+		const UScriptStruct* UnresolvedParentColumn = [InProvider, &RegistrationParams]()
+		{
+			FDynamicColumnDescription Desc
+			{
+				.TemplateType = FEditorDataHierarchyUnresolvedParent_Template::StaticStruct(),
+				.Identifier = RegistrationParams.Name
+			};
+			return InProvider->GenerateDynamicColumn(Desc);
+		}();
 		
 		TUniquePtr<FTedsHierarchyAccessInterface> AccessInterface = [&]()
 		{
@@ -86,7 +96,8 @@ namespace UE::Editor::DataStorage
 			{
 				.ChildTag = ChildTag,
 				.ParentTag = ParentTag,
-				.HierarchyData = HierarchyData
+				.HierarchyData = HierarchyData,
+				.UnresolvedParentColumn = UnresolvedParentColumn
 			};
 			return MakeUnique<FTedsHierarchyAccessInterface>(Params);
 		}();
@@ -97,6 +108,7 @@ namespace UE::Editor::DataStorage
 			.ChildTag = ChildTag,
 			.ParentTag = ParentTag,
 			.HierarchyData = HierarchyData,
+			.UnresolvedParentColumn = UnresolvedParentColumn,
 			.AccessInterface = MoveTemp(AccessInterface)
 		};
 
@@ -107,6 +119,7 @@ namespace UE::Editor::DataStorage
 		
 		ICoreProvider* DataStorage = GetMutableDataStorageFeature<ICoreProvider>(StorageFeatureName);
 		RegisterObservers(*DataStorage, HierarchyHandle);
+		RegisterProcessors(*DataStorage, HierarchyHandle);
 		
 		return HierarchyHandle;
 	}
@@ -144,7 +157,7 @@ namespace UE::Editor::DataStorage
 		using namespace UE::Editor::DataStorage::Queries;
 
 		const FRegisteredHierarchy* Hierarchy = FindRegisteredHierarchy(Handle);
-		check(Hierarchy != nullptr);
+		checkf(Hierarchy, TEXT("Couldn't find hierarchy that was just registered, observers for the hierarchy could not be registered successfully"));
 
 		QueryHandle Subquery = DataStorage.RegisterQuery(
 				Select()
@@ -249,6 +262,42 @@ namespace UE::Editor::DataStorage
 		.Compile());
 	}
 
+	void FTedsHierarchyRegistrar::RegisterProcessors(ICoreProvider& CoreProvider, const FHierarchyHandle& Handle)
+	{
+		using namespace UE::Editor::DataStorage::Queries;
+
+		const FRegisteredHierarchy* Hierarchy = FindRegisteredHierarchy(Handle);
+		checkf(Hierarchy, TEXT("Couldn't find hierarchy that was just registered, processors for the hierarchy could not be registered successfully"));
+		
+		CoreProvider.RegisterQuery(
+		Select(
+			TEXT("Resolve hierarchy rows"),
+			FProcessor(EQueryTickPhase::PrePhysics, CoreProvider.GetQueryTickGroupName(EQueryTickGroups::Default)),
+			[this, Handle](IQueryContext& Context, RowHandle Row)
+			{
+				const FRegisteredHierarchy* Hierarchy = FindRegisteredHierarchy(Handle);
+				
+				FEditorDataHierarchyUnresolvedParent_Template* UnresolvedParentColumn;
+				{
+					const int RowIndexInChunk = Hack_FindRowIndexInChunk(Context, Row);
+					check(RowIndexInChunk != INDEX_NONE);
+					void* HierarchyDataComponentArrayStart = Context.GetMutableColumn(Hierarchy->UnresolvedParentColumn);
+					uint8* HierarchyDataComponentObjectAddress = static_cast<uint8*>(HierarchyDataComponentArrayStart) + (RowIndexInChunk * Hierarchy->UnresolvedParentColumn->GetStructureSize());
+					UnresolvedParentColumn = reinterpret_cast<FEditorDataHierarchyUnresolvedParent_Template*>(HierarchyDataComponentObjectAddress);
+				}
+				
+				RowHandle ParentRow = Context.LookupMappedRow(UnresolvedParentColumn->MappingDomain, UnresolvedParentColumn->ParentId);
+				if (Context.IsRowAvailable(ParentRow))
+				{
+					Context.SetParentRow(Row, ParentRow);
+					Context.RemoveColumns(Row, {Hierarchy->UnresolvedParentColumn});
+				}
+			})
+		.AccessesHierarchy(Hierarchy->Name)
+		.ReadWrite(Hierarchy->UnresolvedParentColumn)
+		.Compile());
+	}
+
 	const FTedsHierarchyRegistrar::FRegisteredHierarchy* FTedsHierarchyRegistrar::FindRegisteredHierarchy(const FHierarchyHandle& Handle)
 	{
 		int32 Index = GetHandleIndex(Handle);
@@ -263,6 +312,7 @@ namespace UE::Editor::DataStorage
 		: ChildTag(Params.ChildTag)
 		, ParentTag(Params.ParentTag)
 		, HierarchyDataColumnType(Params.HierarchyData)
+		, UnresolvedParentColumnType(Params.UnresolvedParentColumn)
 	{
 		check(HierarchyDataColumnType->IsChildOf(FEditorDataHierarchyData_Template::StaticStruct()));
 	}
@@ -281,7 +331,12 @@ namespace UE::Editor::DataStorage
 	{
 		return HierarchyDataColumnType;
 	}
-	
+
+	const UScriptStruct* FTedsHierarchyAccessInterface::GetUnresolvedParentColumnType() const
+	{
+		return UnresolvedParentColumnType;
+	}
+
 	bool FTedsHierarchyAccessInterface::HasChildren(const ICoreProvider& Context, RowHandle Row) const
 	{
 		const FEditorDataHierarchyData_Template* Hierarchy = static_cast<const FEditorDataHierarchyData_Template*>(Context.GetColumnData(Row, HierarchyDataColumnType));
@@ -317,6 +372,7 @@ namespace UE::Editor::DataStorage
 	void FTedsHierarchyAccessInterface::SetParentRow(ICoreProvider* CoreProvider, RowHandle Target, RowHandle Parent) const
 	{
 		CoreProvider->RemoveColumn(Target, ChildTag);
+		CoreProvider->RemoveColumn(Target, UnresolvedParentColumnType);
 
 		auto AddHierarchyDataColumn = [this](ICoreProvider* CoreProvider, RowHandle Row, RowHandle ParentRow)
 		{
@@ -389,6 +445,27 @@ namespace UE::Editor::DataStorage
 		}
 	}
 
+	void FTedsHierarchyAccessInterface::SetUnresolvedParent(ICoreProvider* CoreProvider, RowHandle Target, FMapKey ParentId, FName MappingDomain) const
+	{
+		const UScriptStruct* UnresolvedParentColumn = GetUnresolvedParentColumnType();
+		
+		CoreProvider->AddColumnData(
+				Target,
+				UnresolvedParentColumn,
+				[ParentId, MappingDomain](void* Dest, const UScriptStruct& ColumnType)
+				{
+					ColumnType.InitializeStruct(Dest);
+					FEditorDataHierarchyUnresolvedParent_Template* UnresolvedParentColumn = static_cast<FEditorDataHierarchyUnresolvedParent_Template*>(Dest);
+					UnresolvedParentColumn->ParentId = ParentId;
+					UnresolvedParentColumn->MappingDomain = MappingDomain;
+				},
+				[](const UScriptStruct& ColumnType, void* Destination, void* Source)
+				{
+					ColumnType.CopyScriptStruct(Destination, Source);
+				}
+			);
+	}
+
 	RowHandle FTedsHierarchyAccessInterface::GetParentRow(const ICoreProvider* CoreProvider, RowHandle Target) const
 	{
 		const UScriptStruct* HierarchyDataColumn = GetHierarchyDataColumnType();
@@ -430,6 +507,21 @@ namespace UE::Editor::DataStorage
 		}
 	};
 
+	// Deferred command to set the unresolved parent of a target row
+	struct FSetUnresolvedParentCommand
+	{
+		RowHandle Target;
+		FMapKey ParentId;
+		FName MappingDomain;
+		const FTedsHierarchyAccessInterface* HierarchyAccessInterface = nullptr;
+		ICoreProvider* CoreProvider = nullptr;
+
+		void operator()() const
+		{
+			HierarchyAccessInterface->SetUnresolvedParent(CoreProvider, Target, ParentId, MappingDomain);
+		}
+	};
+
 	
 	void FTedsHierarchyAccessInterface::SetParent(ICommonQueryContext& Context, RowHandle Target, RowHandle Parent) const
 	{
@@ -440,6 +532,18 @@ namespace UE::Editor::DataStorage
 		Command.CoreProvider = GetMutableDataStorageFeature<ICoreProvider>(StorageFeatureName);
 		Context.PushCommand(MoveTemp(Command));
 	}
+	
+
+	void FTedsHierarchyAccessInterface::SetUnresolvedParent(ICommonQueryContext& Context, RowHandle Target, FMapKey ParentId, FName MappingDomain) const
+	{
+		FSetUnresolvedParentCommand Command;
+		Command.ParentId = ParentId;
+		Command.Target = Target;
+		Command.MappingDomain = MappingDomain;
+		Command.HierarchyAccessInterface = this;
+		Command.CoreProvider = GetMutableDataStorageFeature<ICoreProvider>(StorageFeatureName);
+		Context.PushCommand(MoveTemp(Command));	
+	}
 
 	RowHandle FTedsHierarchyAccessInterface::GetParent(const ICommonQueryContext& Context, RowHandle Target) const
 	{
@@ -471,6 +575,11 @@ namespace UE::Editor::DataStorage
 		SetParent(static_cast<ICommonQueryContext&>(Context), Target, Parent);
 	}
 
+	void FTedsHierarchyAccessInterface::SetUnresolvedParent(IQueryContext& Context, RowHandle Target, FMapKey ParentId, FName MappingDomain) const
+	{
+		SetUnresolvedParent(static_cast<ICommonQueryContext&>(Context), Target, ParentId, MappingDomain);
+	}
+
 	RowHandle FTedsHierarchyAccessInterface::GetParentRow(const IQueryContext& Context, RowHandle Target) const
 	{
 		return GetParent(static_cast<const ICommonQueryContext&>(Context), Target);
@@ -481,6 +590,11 @@ namespace UE::Editor::DataStorage
 		SetParent(static_cast<ICommonQueryContext&>(Context), Target, Parent);
 	}
 
+	void FTedsHierarchyAccessInterface::SetUnresolvedParent(ISubqueryContext& Context, RowHandle Target, FMapKey ParentId, FName MappingDomain) const
+	{
+		SetUnresolvedParent(static_cast<ICommonQueryContext&>(Context), Target, ParentId, MappingDomain);
+	}
+
 	RowHandle FTedsHierarchyAccessInterface::GetParentRow(const ISubqueryContext& Context, RowHandle Target) const
 	{
 		return GetParent(static_cast<const ICommonQueryContext&>(Context), Target);
@@ -491,6 +605,11 @@ namespace UE::Editor::DataStorage
 		SetParent(static_cast<ICommonQueryContext&>(Context), Target, Parent);
 	}
 
+	void FTedsHierarchyAccessInterface::SetUnresolvedParent(IDirectQueryContext& Context, RowHandle Target, FMapKey ParentId, FName MappingDomain) const
+	{
+		SetUnresolvedParent(static_cast<ICommonQueryContext&>(Context), Target, ParentId, MappingDomain);
+	}
+
 	RowHandle FTedsHierarchyAccessInterface::GetParentRow(const IDirectQueryContext& Context, RowHandle Target) const
 	{
 		// NOTE: It is not clear to the user that they will need to have registered the HierarchyDataColumnType as at least ReadOnly
diff --git a/Engine/Plugins/Experimental/EditorDataStorage/Source/TedsCore/Private/Hierarchy/EditorDataStorageHierarchyImplementation.h b/Engine/Plugins/Experimental/EditorDataStorage/Source/TedsCore/Private/Hierarchy/EditorDataStorageHierarchyImplementation.h
--- a/Engine/Plugins/Experimental/EditorDataStorage/Source/TedsCore/Private/Hierarchy/EditorDataStorageHierarchyImplementation.h
+++ b/Engine/Plugins/Experimental/EditorDataStorage/Source/TedsCore/Private/Hierarchy/EditorDataStorageHierarchyImplementation.h
@@ -27,6 +27,7 @@ namespace UE::Editor::DataStorage
 		const FTedsHierarchyAccessInterface* GetAccessInterface(FHierarchyHandle) const;
 	private:
 		void RegisterObservers(ICoreProvider& CoreProvider, const FHierarchyHandle& Handle);
+		void RegisterProcessors(ICoreProvider& CoreProvider, const FHierarchyHandle& Handle);
 		
 		struct FRegisteredHierarchy
 		{
@@ -34,6 +35,7 @@ namespace UE::Editor::DataStorage
 			const UScriptStruct* ChildTag;
 			const UScriptStruct* ParentTag;
 			const UScriptStruct* HierarchyData;
+			const UScriptStruct* UnresolvedParentColumn;
 			TUniquePtr<FTedsHierarchyAccessInterface> AccessInterface;
 		};
 
@@ -51,6 +53,7 @@ namespace UE::Editor::DataStorage
 			const UScriptStruct* ChildTag;
 			const UScriptStruct* ParentTag;
 			const UScriptStruct* HierarchyData;
+			const UScriptStruct* UnresolvedParentColumn;
 		};
 		
 		explicit FTedsHierarchyAccessInterface(const FConstructParams& Params);
@@ -58,8 +61,10 @@ namespace UE::Editor::DataStorage
 		const UScriptStruct* GetChildTagType() const;
 		const UScriptStruct* GetParentTagType() const;
 		const UScriptStruct* GetHierarchyDataColumnType() const;
-		
+		const UScriptStruct* GetUnresolvedParentColumnType() const;
+
 		void SetParentRow(ICoreProvider* CoreProvider, RowHandle Target, RowHandle Parent) const;
+		void SetUnresolvedParent(ICoreProvider* CoreProvider, RowHandle Target, FMapKey ParentId, FName MappingDomain) const;
 		RowHandle GetParentRow(const ICoreProvider* CoreProvider, RowHandle Target) const;
 		
 		bool HasChildren(const ICoreProvider& Context, RowHandle Row) const;
@@ -71,19 +76,24 @@ namespace UE::Editor::DataStorage
 		TFunction<RowHandle(const void*, const UScriptStruct*)> CreateParentExtractionFunction() const;
 		
 		void SetParentRow(IQueryContext& Context, RowHandle Target, RowHandle Parent) const;
+		void SetUnresolvedParent(IQueryContext& Context, RowHandle Target, FMapKey ParentId, FName MappingDomain) const;
 		RowHandle GetParentRow(const IQueryContext& Context, RowHandle Target) const;
 		void SetParentRow(ISubqueryContext& Context, RowHandle Target, RowHandle Parent) const;
+		void SetUnresolvedParent(ISubqueryContext& Context, RowHandle Target, FMapKey ParentId, FName MappingDomain) const;
 		RowHandle GetParentRow(const ISubqueryContext& Context, RowHandle Target) const;
 		void SetParentRow(IDirectQueryContext& Context, RowHandle Target, RowHandle Parent) const;
+		void SetUnresolvedParent(IDirectQueryContext& Context, RowHandle Target, FMapKey ParentId, FName MappingDomain) const;
 		RowHandle GetParentRow(const IDirectQueryContext& Context, RowHandle Target) const;
 
 	private:
 
 		void SetParent(ICommonQueryContext& Context, RowHandle Target, RowHandle Parent) const;
+		void SetUnresolvedParent(ICommonQueryContext& Context, RowHandle Target, FMapKey ParentId, FName MappingDomain) const;
 		RowHandle GetParent(const ICommonQueryContext& Context, RowHandle Target) const;
 		
 		const UScriptStruct* ChildTag;
 		const UScriptStruct* ParentTag;
 		const UScriptStruct* HierarchyDataColumnType;
+		const UScriptStruct* UnresolvedParentColumnType;
 	};
 }
diff --git a/Engine/Plugins/Experimental/EditorDataStorage/Source/TedsCore/Private/Processors/TypedElementProcessorAdaptors.cpp b/Engine/Plugins/Experimental/EditorDataStorage/Source/TedsCore/Private/Processors/TypedElementProcessorAdaptors.cpp
--- a/Engine/Plugins/Experimental/EditorDataStorage/Source/TedsCore/Private/Processors/TypedElementProcessorAdaptors.cpp
+++ b/Engine/Plugins/Experimental/EditorDataStorage/Source/TedsCore/Private/Processors/TypedElementProcessorAdaptors.cpp
@@ -989,6 +989,15 @@ namespace UE::Editor::DataStorage
 				}
 			}
 
+			template<typename ContextType>
+			void SetUnresolvedParent(ContextType& Context, RowHandle Target, FMapKey ParentId, FName MappingDomain) const
+			{
+				if (HierarchyAccessInterface)
+				{
+					HierarchyAccessInterface->SetUnresolvedParent(Context, Target, ParentId, MappingDomain);
+				}
+			}
+
 			template<typename ContextType>
 			RowHandle GetParent(ContextType& Context, RowHandle Target) const
 			{
@@ -1122,6 +1131,13 @@ namespace UE::Editor::DataStorage
 			const UScriptStruct* FindDynamicColumnType(const UE::Editor::DataStorage::FDynamicColumnDescription& Description) const override { return Implementation.FindDynamicColumnType(Description); }
 			float GetDeltaTimeSeconds() const override { return Implementation.GetDeltaTimeSeconds(); }
 			virtual void SetParentRow(RowHandle Target, RowHandle Parent) override { if (HierarchyAccessInterface) { HierarchyAccessInterface->SetParentRow(*static_cast<IDirectQueryContext*>(this), Target, Parent); } }
+			virtual void SetUnresolvedParent(FHierarchyHandle InHierarchyHandle, RowHandle Target, FMapKey ParentId, FName MappingDomain) override
+			{
+				if (HierarchyAccessInterface)
+				{
+					HierarchyAccessInterface->SetUnresolvedParent(*static_cast<IDirectQueryContext*>(this), Target, ParentId, MappingDomain);
+				}
+			}
 			virtual RowHandle GetParentRow(RowHandle Target) const override { return HierarchyAccessInterface ? HierarchyAccessInterface->GetParentRow(*static_cast<const IDirectQueryContext*>(this), Target) : InvalidRowHandle; }
 			virtual void PushCommand(void(* CommandFunction)(void*), void* InCommandData) override  { return Implementation.PushCommand(CommandFunction, InCommandData); }
 		
@@ -1175,6 +1191,7 @@ namespace UE::Editor::DataStorage
 			const UScriptStruct* FindDynamicColumnType(const UE::Editor::DataStorage::FDynamicColumnDescription& Description) const override { return Implementation.FindDynamicColumnType(Description); }
 			float GetDeltaTimeSeconds() const override { return Implementation.GetDeltaTimeSeconds(); }
 			virtual void SetParentRow(RowHandle Target, RowHandle Parent) override { Implementation.SetParent(*static_cast<ISubqueryContext*>(this), Target, Parent); }
+			virtual void SetUnresolvedParent(FHierarchyHandle InHierarchyHandle, RowHandle Target, FMapKey ParentId, FName MappingDomain) override { Implementation.SetUnresolvedParent(*static_cast<ISubqueryContext*>(this), Target, ParentId, MappingDomain); }
 			virtual RowHandle GetParentRow(RowHandle Target) const override { return Implementation.GetParent(*static_cast<const ISubqueryContext*>(this), Target); }
 			void PushCommand(void (*CommandFunction)(void*), void* CommandData) override { return Implementation.PushCommand(CommandFunction, CommandData); }
 		protected:
@@ -1328,9 +1345,9 @@ namespace UE::Editor::DataStorage
 			const UScriptStruct* FindDynamicColumnType(const UE::Editor::DataStorage::FDynamicColumnDescription& Description) const override { return Implementation.FindDynamicColumnType(Description); }
 			float GetDeltaTimeSeconds() const override { return Implementation.GetDeltaTimeSeconds(); }
 			virtual void SetParentRow(RowHandle Target, RowHandle Parent) override { Implementation.SetParent(*static_cast<IQueryContext*>(this), Target, Parent); }
+			virtual void SetUnresolvedParent(FHierarchyHandle InHierarchyHandle, RowHandle Target, FMapKey ParentId, FName MappingDomain) override { Implementation.SetUnresolvedParent(*static_cast<IQueryContext*>(this), Target, ParentId, MappingDomain); }
 			virtual RowHandle GetParentRow(RowHandle Target) const override { return Implementation.GetParent(*static_cast<const IQueryContext*>(this), Target); }
 			void PushCommand(void (*CommandFunction)(void*), void* Context) override { return Implementation.PushCommand(CommandFunction, Context); }
-
 			const UObject* GetDependency(const UClass* DependencyClass) override { return Implementation.GetDependency(DependencyClass); }
 			UObject* GetMutableDependency(const UClass* DependencyClass) override { return Implementation.GetMutableDependency(DependencyClass); }
 			void GetDependencies(TArrayView<UObject*> RetrievedAddresses, TConstArrayView<TWeakObjectPtr<const UClass>> DependencyTypes, TConstArrayView<EQueryAccessType> AccessTypes) override { return Implementation.GetDependencies(RetrievedAddresses, DependencyTypes, AccessTypes); }
diff --git a/Engine/Plugins/Experimental/EditorDataStorage/Source/TedsCore/Private/TypedElementDatabase.cpp b/Engine/Plugins/Experimental/EditorDataStorage/Source/TedsCore/Private/TypedElementDatabase.cpp
--- a/Engine/Plugins/Experimental/EditorDataStorage/Source/TedsCore/Private/TypedElementDatabase.cpp
+++ b/Engine/Plugins/Experimental/EditorDataStorage/Source/TedsCore/Private/TypedElementDatabase.cpp
@@ -1111,6 +1111,15 @@ void UEditorDataStorage::SetParentRow(UE::Editor::DataStorage::FHierarchyHandle
 	}
 }
 
+void UEditorDataStorage::SetUnresolvedParent(UE::Editor::DataStorage::FHierarchyHandle InHierarchyHandle, UE::Editor::DataStorage::RowHandle Target,
+	UE::Editor::DataStorage::FMapKey ParentId, FName MappingDomain)
+{
+	if (const UE::Editor::DataStorage::FTedsHierarchyAccessInterface* AccessInterface = Environment->GetHierarchyRegistrar().GetAccessInterface(InHierarchyHandle))
+	{
+		AccessInterface->SetUnresolvedParent(this, Target, ParentId, MappingDomain);
+	}
+}
+
 UE::Editor::DataStorage::RowHandle UEditorDataStorage::GetParentRow(UE::Editor::DataStorage::FHierarchyHandle InHierarchyHandle,
 	UE::Editor::DataStorage::RowHandle Target) const
 {
diff --git a/Engine/Plugins/Experimental/EditorDataStorage/Source/TedsCore/Private/TypedElementDatabase.h b/Engine/Plugins/Experimental/EditorDataStorage/Source/TedsCore/Private/TypedElementDatabase.h
--- a/Engine/Plugins/Experimental/EditorDataStorage/Source/TedsCore/Private/TypedElementDatabase.h
+++ b/Engine/Plugins/Experimental/EditorDataStorage/Source/TedsCore/Private/TypedElementDatabase.h
@@ -151,6 +151,9 @@ public:
 		UE::Editor::DataStorage::FHierarchyHandle) const override;
 	UE_API virtual void SetParentRow(
 		UE::Editor::DataStorage::FHierarchyHandle, UE::Editor::DataStorage::RowHandle Target,	UE::Editor::DataStorage::RowHandle Parent) override;
+	virtual void SetUnresolvedParent(
+		UE::Editor::DataStorage::FHierarchyHandle InHierarchyHandle, UE::Editor::DataStorage::RowHandle Target, UE::Editor::DataStorage::FMapKey ParentId,
+		FName MappingDomain) override;
 	UE_API virtual UE::Editor::DataStorage::RowHandle GetParentRow(
 		UE::Editor::DataStorage::FHierarchyHandle, UE::Editor::DataStorage::RowHandle Target) const override;
 	UE_API virtual TFunction<UE::Editor::DataStorage::RowHandle(const void*, const UScriptStruct*)> CreateParentExtractionFunction(
diff --git a/Engine/Source/Runtime/TypedElementFramework/Public/Elements/Interfaces/TypedElementDataStorageInterface.h b/Engine/Source/Runtime/TypedElementFramework/Public/Elements/Interfaces/TypedElementDataStorageInterface.h
--- a/Engine/Source/Runtime/TypedElementFramework/Public/Elements/Interfaces/TypedElementDataStorageInterface.h
+++ b/Engine/Source/Runtime/TypedElementFramework/Public/Elements/Interfaces/TypedElementDataStorageInterface.h
@@ -396,6 +396,17 @@ public:
 		RowHandle Target,
 		RowHandle Parent) = 0;
 
+	/**
+	 * Establishes a parent relationship between the Target row and a Parent that is not registered in TEDS yet.
+	 * Every frame, TEDS will attempt to resolve the missing relation by looking up the parent using TEDS Mapping.
+	 * An optional MaxFrameCount can be specified, which will be decremented every frame until the potential relationship is discarded
+	 */
+	virtual void SetUnresolvedParent(
+		FHierarchyHandle InHierarchyHandle,
+		RowHandle Target,
+		FMapKey ParentId,
+		FName MappingDomain) = 0;
+
 	// Gets the parent of a target, if there is one
 	virtual RowHandle GetParentRow(FHierarchyHandle InHierarchyHandle, RowHandle Target) const = 0;
 	
diff --git a/Engine/Source/Runtime/TypedElementFramework/Public/Elements/Interfaces/TypedElementQueryStorageInterfaces.h b/Engine/Source/Runtime/TypedElementFramework/Public/Elements/Interfaces/TypedElementQueryStorageInterfaces.h
--- a/Engine/Source/Runtime/TypedElementFramework/Public/Elements/Interfaces/TypedElementQueryStorageInterfaces.h
+++ b/Engine/Source/Runtime/TypedElementFramework/Public/Elements/Interfaces/TypedElementQueryStorageInterfaces.h
@@ -109,6 +109,17 @@ namespace UE::Editor::DataStorage
 			RowHandle Target,
 			RowHandle Parent) = 0;
 
+		/**
+		 * Establishes a parent relationship between the Target row and a Parent that is not registered in TEDS yet.
+		 * Every frame, TEDS will attempt to resolve the missing relation by looking up the parent using TEDS Mapping.
+		 * An optional MaxFrameCount can be specified, which will be decremented every frame until the potential relationship is discarded
+		 */
+		virtual void SetUnresolvedParent(
+			FHierarchyHandle InHierarchyHandle,
+			RowHandle Target,
+			FMapKey ParentId,
+			FName MappingDomain) = 0;
+
 		// Gets the Target's Parent if the query has a configured Hierarchy
 		virtual RowHandle GetParentRow(
 			RowHandle Target) const = 0;
