Teds query context: Remove unnecessary hierachy handle argument from the context functions, the query contexts already have access to the hierarchy they are operating on.

diff --git a/Engine/Plugins/Experimental/EditorDataStorage/Source/TedsCore/Private/Processors/TypedElementProcessorAdaptors.cpp b/Engine/Plugins/Experimental/EditorDataStorage/Source/TedsCore/Private/Processors/TypedElementProcessorAdaptors.cpp
--- a/Engine/Plugins/Experimental/EditorDataStorage/Source/TedsCore/Private/Processors/TypedElementProcessorAdaptors.cpp
+++ b/Engine/Plugins/Experimental/EditorDataStorage/Source/TedsCore/Private/Processors/TypedElementProcessorAdaptors.cpp
@@ -1131,7 +1131,7 @@ namespace UE::Editor::DataStorage
 			const UScriptStruct* FindDynamicColumnType(const UE::Editor::DataStorage::FDynamicColumnDescription& Description) const override { return Implementation.FindDynamicColumnType(Description); }
 			float GetDeltaTimeSeconds() const override { return Implementation.GetDeltaTimeSeconds(); }
 			virtual void SetParentRow(RowHandle Target, RowHandle Parent) override { if (HierarchyAccessInterface) { HierarchyAccessInterface->SetParentRow(*static_cast<IDirectQueryContext*>(this), Target, Parent); } }
-			virtual void SetUnresolvedParent(FHierarchyHandle InHierarchyHandle, RowHandle Target, FMapKey ParentId, FName MappingDomain) override
+			virtual void SetUnresolvedParent(RowHandle Target, FMapKey ParentId, FName MappingDomain) override
 			{
 				if (HierarchyAccessInterface)
 				{
@@ -1191,7 +1191,7 @@ namespace UE::Editor::DataStorage
 			const UScriptStruct* FindDynamicColumnType(const UE::Editor::DataStorage::FDynamicColumnDescription& Description) const override { return Implementation.FindDynamicColumnType(Description); }
 			float GetDeltaTimeSeconds() const override { return Implementation.GetDeltaTimeSeconds(); }
 			virtual void SetParentRow(RowHandle Target, RowHandle Parent) override { Implementation.SetParent(*static_cast<ISubqueryContext*>(this), Target, Parent); }
-			virtual void SetUnresolvedParent(FHierarchyHandle InHierarchyHandle, RowHandle Target, FMapKey ParentId, FName MappingDomain) override { Implementation.SetUnresolvedParent(*static_cast<ISubqueryContext*>(this), Target, ParentId, MappingDomain); }
+			virtual void SetUnresolvedParent(RowHandle Target, FMapKey ParentId, FName MappingDomain) override { Implementation.SetUnresolvedParent(*static_cast<ISubqueryContext*>(this), Target, ParentId, MappingDomain); }
 			virtual RowHandle GetParentRow(RowHandle Target) const override { return Implementation.GetParent(*static_cast<const ISubqueryContext*>(this), Target); }
 			void PushCommand(void (*CommandFunction)(void*), void* CommandData) override { return Implementation.PushCommand(CommandFunction, CommandData); }
 		protected:
@@ -1345,7 +1345,7 @@ namespace UE::Editor::DataStorage
 			const UScriptStruct* FindDynamicColumnType(const UE::Editor::DataStorage::FDynamicColumnDescription& Description) const override { return Implementation.FindDynamicColumnType(Description); }
 			float GetDeltaTimeSeconds() const override { return Implementation.GetDeltaTimeSeconds(); }
 			virtual void SetParentRow(RowHandle Target, RowHandle Parent) override { Implementation.SetParent(*static_cast<IQueryContext*>(this), Target, Parent); }
-			virtual void SetUnresolvedParent(FHierarchyHandle InHierarchyHandle, RowHandle Target, FMapKey ParentId, FName MappingDomain) override { Implementation.SetUnresolvedParent(*static_cast<IQueryContext*>(this), Target, ParentId, MappingDomain); }
+			virtual void SetUnresolvedParent(RowHandle Target, FMapKey ParentId, FName MappingDomain) override { Implementation.SetUnresolvedParent(*static_cast<IQueryContext*>(this), Target, ParentId, MappingDomain); }
 			virtual RowHandle GetParentRow(RowHandle Target) const override { return Implementation.GetParent(*static_cast<const IQueryContext*>(this), Target); }
 			void PushCommand(void (*CommandFunction)(void*), void* Context) override { return Implementation.PushCommand(CommandFunction, Context); }
 			const UObject* GetDependency(const UClass* DependencyClass) override { return Implementation.GetDependency(DependencyClass); }
diff --git a/Engine/Source/Runtime/TypedElementFramework/Public/Elements/Interfaces/TypedElementQueryStorageInterfaces.h b/Engine/Source/Runtime/TypedElementFramework/Public/Elements/Interfaces/TypedElementQueryStorageInterfaces.h
--- a/Engine/Source/Runtime/TypedElementFramework/Public/Elements/Interfaces/TypedElementQueryStorageInterfaces.h
+++ b/Engine/Source/Runtime/TypedElementFramework/Public/Elements/Interfaces/TypedElementQueryStorageInterfaces.h
@@ -115,7 +115,6 @@ namespace UE::Editor::DataStorage
 		 * An optional MaxFrameCount can be specified, which will be decremented every frame until the potential relationship is discarded
 		 */
 		virtual void SetUnresolvedParent(
-			FHierarchyHandle InHierarchyHandle,
 			RowHandle Target,
 			FMapKey ParentId,
 			FName MappingDomain) = 0;
