TEDS UI: Add functionality to register decorator widgets that are compounded to the actual widget when the widget row contains the column they are registered against.
This can be used to add common decorator functionality to widgets such as UI effects, context menus etc.

- Add API to create custom decorators and register them with IUIProvider
- Deprecate functions to register widget constructor by instance and remove the one use case with CounterWidget
- Convert the universal search custom widget to a TEDS UI decorator as an example

diff --git a/Engine/Plugins/Experimental/EditorDataStorage/Source/TedsCore/Private/TypedElementDatabaseUI.cpp b/Engine/Plugins/Experimental/EditorDataStorage/Source/TedsCore/Private/TypedElementDatabaseUI.cpp
--- a/Engine/Plugins/Experimental/EditorDataStorage/Source/TedsCore/Private/TypedElementDatabaseUI.cpp
+++ b/Engine/Plugins/Experimental/EditorDataStorage/Source/TedsCore/Private/TypedElementDatabaseUI.cpp
@@ -5,12 +5,14 @@
 #include "Algo/BinarySearch.h"
 #include "Algo/Sort.h"
 #include "Algo/Unique.h"
+#include "Elements/Columns/DecoratorWidgetColumns.h"
 #include "Elements/Columns/TypedElementHiearchyColumns.h"
 #include "Elements/Columns/WidgetPurposeColumns.h"
 #include "Elements/Columns/TypedElementMiscColumns.h"
 #include "Elements/Columns/TypedElementSlateWidgetColumns.h"
 #include "Elements/Framework/TypedElementColumnUtils.h"
 #include "Elements/Framework/TypedElementRowHandleArray.h"
+#include "Elements/Interfaces/DecoratorWidgetConstructor.h"
 #include "Elements/Interfaces/TypedElementDataStorageInterface.h"
 #include "Elements/Interfaces/TypedElementDataStorageCompatibilityInterface.h"
 #include "GenericPlatform/GenericPlatformMemory.h"
@@ -69,6 +71,8 @@ void UEditorDataStorageUi::Initialize(
 	StorageCompatibility = StorageCompatibilityInterface;
 	CreateStandardArchetypes();
 	RegisterQueries();
+
+	Storage->OnUpdateCompleted().AddUObject(this, &UEditorDataStorageUi::PostDataStorageUpdate);
 }
 
 void UEditorDataStorageUi::PostInitialize(UE::Editor::DataStorage::ICoreProvider* StorageInterface,
@@ -115,6 +119,7 @@ void UEditorDataStorageUi::PostInitialize(UE::Editor::DataStorage::ICoreProvider
 
 void UEditorDataStorageUi::Deinitialize()
 {
+	Storage->OnUpdateCompleted().RemoveAll(this);
 }
 
 UE::Editor::DataStorage::RowHandle UEditorDataStorageUi::RegisterWidgetPurpose(const FPurposeID& PurposeID, const FPurposeInfo& InPurposeInfo)
@@ -284,41 +289,7 @@ bool UEditorDataStorageUi::RegisterWidgetFactory(UE::Editor::DataStorage::RowHan
 {
 	checkf(Constructor->GetTypeInfo(), TEXT("Widget constructor being registered that doesn't have valid type information."));
 
-	if (FWidgetPurposeColumn* PurposeInfo = Storage->GetColumn<FWidgetPurposeColumn>(PurposeRow))
-	{
-		switch (PurposeInfo->PurposeType)
-		{
-		case EPurposeType::Generic:
-			{
-				UE::Editor::DataStorage::RowHandle FactoryRow = RegisterWidgetFactoryRow(PurposeRow);
-				Storage->AddColumn(FactoryRow, FWidgetFactoryConstructorColumn{.Constructor = MoveTemp(Constructor)});
-				return true;
-			}
-		case EPurposeType::UniqueByName:
-			{
-				UE::Editor::DataStorage::RowHandle FactoryRow = RegisterUniqueWidgetFactoryRow(PurposeRow);
-				Storage->AddColumn(FactoryRow, FWidgetFactoryConstructorColumn{.Constructor = MoveTemp(Constructor)});
-				return true;
-			}
-		case EPurposeType::UniqueByNameAndColumn:
-			{
-				UE_LOG(LogEditorDataStorageUI, Warning,
-				TEXT("Unable to register widget factory '%s' as purpose '%llu' requires at least one column for matching."),
-				*Constructor->GetTypeInfo()->GetName(), PurposeRow);
-				return false;
-			}
-		default:
-			{
-				checkf(false, TEXT("Unexpected UE::Editor::DataStorage::IUiProvider::EPurposeType found provided when registering widget factory."));
-				return false;
-			}
-		}
-	}
-
-	UE_LOG(LogEditorDataStorageUI, Warning,
-		TEXT("Unable to register widget factory '%s' as purpose '%llu' isn't registered."), *Constructor->GetTypeInfo()->GetName(), PurposeRow);
-
-	return false;
+	return RegisterWidgetFactory(PurposeRow, Constructor->GetTypeInfo());
 }
 
 bool UEditorDataStorageUi::RegisterWidgetFactory(UE::Editor::DataStorage::RowHandle PurposeRow,
@@ -326,62 +297,76 @@ bool UEditorDataStorageUi::RegisterWidgetFactory(UE::Editor::DataStorage::RowHan
 {
 	checkf(Constructor->GetTypeInfo(), TEXT("Widget constructor being registered that doesn't have valid type information."));
 
-	if (!Columns.IsEmpty())
+	return RegisterWidgetFactory(PurposeRow, Constructor->GetTypeInfo(), Columns);
+}
+
+bool UEditorDataStorageUi::RegisterDecoratorWidgetFactory(UE::Editor::DataStorage::RowHandle PurposeRow, const UScriptStruct* Constructor,
+	const UScriptStruct* Column)
+{
+	checkf(Constructor->IsChildOf(FTedsDecoratorWidgetConstructor::StaticStruct()),
+		TEXT("Attempting to register a decorator widget constructor '%s' that isn't derived from FTedsDecoratorWidgetConstructor."),
+		*Constructor->GetFullName());
+
+	using namespace UE::Editor::DataStorage::Queries;
+		
+	if (FWidgetPurposeColumn* PurposeInfo = Storage->GetColumn<FWidgetPurposeColumn>(PurposeRow))
 	{
-		if (FWidgetPurposeColumn* PurposeInfo = Storage->GetColumn<FWidgetPurposeColumn>(PurposeRow))
+		constexpr bool bIsDecoratorWidget = true;
+
+		switch (PurposeInfo->PurposeType)
 		{
-			switch (PurposeInfo->PurposeType)
+		case EPurposeType::Generic:
 			{
-			case EPurposeType::Generic:
-				{
-					UE::Editor::DataStorage::RowHandle FactoryRow = RegisterWidgetFactoryRow(PurposeRow);
-					Storage->AddColumn(FactoryRow, FWidgetFactoryConstructorColumn{.Constructor = MoveTemp(Constructor)});
-				}
-				return true;
-			case EPurposeType::UniqueByName:
-				if (!Columns.IsEmpty())
+				RowHandle FactoryRow = RegisterWidgetFactoryRow(PurposeRow, bIsDecoratorWidget);
+				Storage->AddColumn(FactoryRow, FWidgetFactoryConstructorTypeInfoColumn{.Constructor = Constructor});
+			}
+			break;
+		case EPurposeType::UniqueByName:
+			{
+				if (Column)
 				{
-					UE::Editor::DataStorage::RowHandle FactoryRow = RegisterUniqueWidgetFactoryRow(PurposeRow);
-					Storage->AddColumn(FactoryRow, FWidgetFactoryConstructorColumn{.Constructor = MoveTemp(Constructor)});
+					RowHandle FactoryRow = RegisterUniqueWidgetFactoryRow(PurposeRow, bIsDecoratorWidget);
+					Storage->AddColumn(FactoryRow, FWidgetFactoryConstructorTypeInfoColumn{.Constructor = Constructor});
 				}
 				else
 				{
 					return false;
 				}
-			case EPurposeType::UniqueByNameAndColumn:
+			}
+			break;
+		case EPurposeType::UniqueByNameAndColumn:
+			{
+				if (Column)
 				{
-					if (!Columns.IsEmpty())
-					{
-						Columns.Compile(UE::Editor::DataStorage::Queries::FEditorStorageQueryConditionCompileContext(Storage));
-							
-						UE::Editor::DataStorage::RowHandle FactoryRow = RegisterWidgetFactoryRow(PurposeRow);
-						Storage->AddColumn(FactoryRow, FWidgetFactoryConstructorColumn{.Constructor = MoveTemp(Constructor)});
-						Storage->AddColumn(FactoryRow, FWidgetFactoryConditionsColumn{.Conditions = Columns});
-
-						return true;
-					}
-					else
-					{
-						return false;
-					}
+					FConditions ColumnConditions{TColumn(Column)};
+					ColumnConditions.Compile(FEditorStorageQueryConditionCompileContext(Storage));
+						
+					RowHandle FactoryRow = RegisterWidgetFactoryRow(PurposeRow, bIsDecoratorWidget);
+					Storage->AddColumn(FactoryRow, FWidgetFactoryConstructorTypeInfoColumn{.Constructor = Constructor});
+					Storage->AddColumn(FactoryRow, FWidgetFactoryConditionsColumn{.Conditions =ColumnConditions});
 				}
-					
-			default:
-				checkf(false, TEXT("Unexpected UE::Editor::DataStorage::IUiProvider::EPurposeType found provided when registering widget factory."));
-				return false;
+				else
+				{
+					return false;
+				}
+				break;
 			}
+		default:
+			checkf(false, TEXT("Unexpected UE::Editor::DataStorage::IUiProvider::EPurposeType found provided when registering widget factory."));
+			return false;
 		}
-		
-		UE_LOG(LogEditorDataStorageUI, Warning,
-		TEXT("Unable to register widget factory '%s' as purpose '%llu' isn't registered."), *Constructor->GetTypeInfo()->GetName(), PurposeRow);
-
-		return false;
 	}
 	else
 	{
-		return RegisterWidgetFactory(PurposeRow, MoveTemp(Constructor));
+		UE_LOG(LogEditorDataStorageUI, Warning, 
+			TEXT("Unable to register decorator widget factory '%s' as purpose '%llu' isn't registered."), *Constructor->GetName(), PurposeRow);
+		return false;
 	}
+
+	// Register observers for this decorator so any widget rows that add or remove the column post construction are updated
+	RegisterDecoratorWidgetObservers(PurposeRow, Column);
 	
+	return true;
 }
 
 void UEditorDataStorageUi::CreateWidgetConstructors(UE::Editor::DataStorage::RowHandle PurposeRow,
@@ -390,7 +375,7 @@ void UEditorDataStorageUi::CreateWidgetConstructors(UE::Editor::DataStorage::Row
 	while (Storage->HasColumns<FWidgetPurposeColumn>(PurposeRow))
 	{
 		TArray<UE::Editor::DataStorage::RowHandle> Factories;
-		GetFactories(PurposeRow, Factories);
+		GetWidgetFactories(PurposeRow, Factories);
 
 		// If no factories were found for this purpose, move on to the parent purpose
 		if (Factories.IsEmpty())
@@ -434,7 +419,7 @@ void UEditorDataStorageUi::CreateWidgetConstructors(UE::Editor::DataStorage::Row
 	while (Storage->HasColumns<FWidgetPurposeColumn>(PurposeRow))
 	{
 		TArray<UE::Editor::DataStorage::RowHandle> Factories;
-		GetFactories(PurposeRow, Factories);
+		GetWidgetFactories(PurposeRow, Factories);
 
 		if (!Factories.IsEmpty())
 		{
@@ -459,21 +444,27 @@ void UEditorDataStorageUi::CreateWidgetConstructors(UE::Editor::DataStorage::Row
 					return LeftSize > RightSize;
 				});
 
+			FWidgetConstructionMethod ConstructionCallback = [this, Arguments, Callback]
+				(UE::Editor::DataStorage::RowHandle FactoryRow, TArray<TWeakObjectPtr<const UScriptStruct>> MatchedColumns)
+				{
+					return CreateSingleWidgetConstructor(FactoryRow, Arguments, MatchedColumns, Callback);
+				};
+
 			switch (MatchApproach)
 			{
 			case EMatchApproach::LongestMatch:
 
 				// For longest match, we don't want to continue matching with the parent purpose if the user requested us to stop
-				if (!CreateWidgetConstructors_LongestMatch(Factories, Columns, Arguments, Callback))
+				if (!CreateWidgetConstructors_LongestMatch(Factories, Columns, ConstructionCallback))
 				{
 					return;
 				}
 				break;
 			case EMatchApproach::ExactMatch:
-				CreateWidgetConstructors_ExactMatch(Factories, Columns, Arguments, Callback);
+				CreateWidgetConstructors_ExactMatch(Factories, Columns, ConstructionCallback);
 				break;
 			case EMatchApproach::SingleMatch:
-				CreateWidgetConstructors_SingleMatch(Factories, Columns, Arguments, Callback);
+				CreateWidgetConstructors_SingleMatch(Factories, Columns, ConstructionCallback);
 				break;
 			default:
 				checkf(false, TEXT("Unsupported match type (%i) for CreateWidgetConstructors."), 
@@ -507,7 +498,7 @@ void UEditorDataStorageUi::ConstructWidgets(UE::Editor::DataStorage::RowHandle P
 
 	while (Storage->HasColumns<FWidgetPurposeColumn>(PurposeRow) && Factories.IsEmpty())
 	{
-		GetFactories(PurposeRow, Factories);
+		GetWidgetFactories(PurposeRow, Factories);
 		
 		// If no factories were found for this purpose, move on to the parent purpose
 		if (FTableRowParentColumn* ParentColumn = Storage->GetColumn<FTableRowParentColumn>(PurposeRow))
@@ -522,11 +513,7 @@ void UEditorDataStorageUi::ConstructWidgets(UE::Editor::DataStorage::RowHandle P
 	
 	for (UE::Editor::DataStorage::RowHandle FactoryRow : Factories)
 	{
-		if (FWidgetFactoryConstructorColumn* ConstructorColumn = Storage->GetColumn<FWidgetFactoryConstructorColumn>(FactoryRow))
-		{
-			CreateWidgetInstance(*ConstructorColumn->Constructor, Arguments, ConstructionCallback);
-		}
-		else if (FWidgetFactoryConstructorTypeInfoColumn* ConstructorTypeInfoColumn = Storage->GetColumn<FWidgetFactoryConstructorTypeInfoColumn>(FactoryRow))
+		if (FWidgetFactoryConstructorTypeInfoColumn* ConstructorTypeInfoColumn = Storage->GetColumn<FWidgetFactoryConstructorTypeInfoColumn>(FactoryRow))
 		{
 			if (const UScriptStruct* ConstructorType = ConstructorTypeInfoColumn->Constructor.Get())
 			{
@@ -628,22 +615,7 @@ bool UEditorDataStorageUi::CreateSingleWidgetConstructor(
 	TArray<TWeakObjectPtr<const UScriptStruct>> MatchedColumnTypes,
 	const WidgetConstructorCallback& Callback)
 {
-	if (FWidgetFactoryConstructorColumn* ConstructorColumn = Storage->GetColumn<FWidgetFactoryConstructorColumn>(FactoryRow))
-	{
-		const UScriptStruct* TargetType = ConstructorColumn->Constructor->GetTypeInfo();
-		checkf(TargetType, TEXT("Expected valid type information from a widget constructor."));
-		TUniquePtr<FTypedElementWidgetConstructor> Result(reinterpret_cast<FTypedElementWidgetConstructor*>(
-			FMemory::Malloc(TargetType->GetStructureSize(), TargetType->GetMinAlignment())));
-		if (Result)
-		{
-			TargetType->InitializeStruct(Result.Get());
-			TargetType->CopyScriptStruct(Result.Get(), ConstructorColumn->Constructor.Get());
-			Result->Initialize(Arguments, MoveTemp(MatchedColumnTypes), FactoryRow);
-			return Callback(MoveTemp(Result), Result->GetMatchedColumns());
-		}
-		return true;
-	}
-	else if (FWidgetFactoryConstructorTypeInfoColumn* ConstructorTypeInfoColumn = Storage->GetColumn<FWidgetFactoryConstructorTypeInfoColumn>(FactoryRow))
+	if (FWidgetFactoryConstructorTypeInfoColumn* ConstructorTypeInfoColumn = Storage->GetColumn<FWidgetFactoryConstructorTypeInfoColumn>(FactoryRow))
 	{
 		if (const UScriptStruct* Target = ConstructorTypeInfoColumn->Constructor.Get())
 		{
@@ -662,15 +634,33 @@ bool UEditorDataStorageUi::CreateSingleWidgetConstructor(
 	return false;
 }
 
+TUniquePtr<FTedsDecoratorWidgetConstructor> UEditorDataStorageUi::CreateSingleDecoratorWidget(UE::Editor::DataStorage::RowHandle FactoryRow,
+	const UE::Editor::DataStorage::FMetaDataView& Arguments, TArray<TWeakObjectPtr<const UScriptStruct>> MatchedColumnTypes)
+{
+	if (FWidgetFactoryConstructorTypeInfoColumn* ConstructorTypeInfoColumn = Storage->GetColumn<FWidgetFactoryConstructorTypeInfoColumn>(FactoryRow))
+	{
+		if (const UScriptStruct* Target = ConstructorTypeInfoColumn->Constructor.Get())
+		{
+			TUniquePtr<FTedsDecoratorWidgetConstructor> Result(reinterpret_cast<FTedsDecoratorWidgetConstructor*>(
+				FMemory::Malloc(Target->GetStructureSize(), Target->GetMinAlignment())));
+			if (Result)
+			{
+				Target->InitializeStruct(Result.Get());
+				Result->Initialize(Arguments, MoveTemp(MatchedColumnTypes), FactoryRow);
+				return MoveTemp(Result);
+			}
+		}
+	}
+	return nullptr;
+}
+
 void UEditorDataStorageUi::CreateWidgetInstance(
 	FTypedElementWidgetConstructor& Constructor, 
 	const UE::Editor::DataStorage::FMetaDataView& Arguments,
 	const WidgetCreatedCallback& ConstructionCallback)
 {
 	UE::Editor::DataStorage::RowHandle Row = Storage->AddRow(WidgetTable);
-	Storage->AddColumns(Row, Constructor.GetAdditionalColumnsList());
-	TSharedPtr<SWidget> Widget = Constructor.ConstructFinalWidget(Row, Storage, this, Arguments);
-	if (Widget)
+	if (TSharedPtr<SWidget> Widget = ConstructWidget(Row, Constructor, Arguments))
 	{
 		ConstructionCallback(Widget.ToSharedRef(), Row);
 	}
@@ -700,7 +690,28 @@ void UEditorDataStorageUi::CreateWidgetInstance(const UScriptStruct* Constructor
 TSharedPtr<SWidget> UEditorDataStorageUi::ConstructWidget(UE::Editor::DataStorage::RowHandle Row, FTypedElementWidgetConstructor& Constructor,
 	const UE::Editor::DataStorage::FMetaDataView& Arguments)
 {
-	return Constructor.ConstructFinalWidget(Row, Storage, this, Arguments);
+	TSharedPtr<SWidget> InternalWidget = ConstructInternalWidget(Row, Constructor, Arguments);
+
+	// Create a container widget to hold the content
+	if (TSharedPtr<UE::Editor::DataStorage::ITedsWidget> ContainerWidget = CreateContainerTedsWidget(Row))
+	{
+		ContainerWidget->SetContent(InternalWidget.ToSharedRef());
+		Storage->GetColumn<FTypedElementSlateWidgetReferenceColumn>(Row)->TedsWidget = ContainerWidget;
+		return ContainerWidget->AsWidget();
+	}
+
+	return InternalWidget;
+}
+
+TSharedPtr<SWidget> UEditorDataStorageUi::ConstructInternalWidget(UE::Editor::DataStorage::RowHandle Row, FTypedElementWidgetConstructor& Constructor,
+	const UE::Editor::DataStorage::FMetaDataView& Arguments)
+{
+	TSharedPtr<SWidget> InternalWidget = Constructor.ConstructFinalWidget(Row, Storage, this, Arguments);
+
+	// Create any decorator widgets
+	InternalWidget = CreateDecoratorWidgets(Row, InternalWidget, Constructor, Arguments);
+
+	return InternalWidget;
 }
 
 void UEditorDataStorageUi::ListWidgetPurposes(const WidgetPurposeCallback& Callback) const
@@ -798,6 +809,8 @@ UE::Editor::DataStorage::RowHandle UEditorDataStorageUi::FindPurpose(const FPurp
 
 void UEditorDataStorageUi::CreateStandardArchetypes()
 {
+	using namespace UE::Editor::DataStorage::Ui;
+	
 	WidgetTable = Storage->RegisterTable(MakeArrayView(
 		{
 			FTypedElementSlateWidgetReferenceColumn::StaticStruct(),
@@ -805,8 +818,15 @@ void UEditorDataStorageUi::CreateStandardArchetypes()
 			FSlateColorColumn::StaticStruct()
 		}), FName(TEXT("Editor_WidgetTable")));
 
-	WidgetPurposeTable = Storage->RegisterTable(UE::Editor::DataStorage::TTypedElementColumnTypeList<FWidgetPurposeColumn, FWidgetPurposeNameColumn>(), FName("Editor_WidgetPurposeTable"));
-	WidgetFactoryTable = Storage->RegisterTable(UE::Editor::DataStorage::TTypedElementColumnTypeList<FWidgetFactoryColumn>(), FName("Editor_WidgetFactoryTable"));
+	WidgetPurposeTable = Storage->RegisterTable(UE::Editor::DataStorage::TTypedElementColumnTypeList<FWidgetPurposeColumn,
+		FWidgetPurposeNameColumn>(), FName("Editor_WidgetPurposeTable"));
+	
+	WidgetFactoryTable = Storage->RegisterTable(UE::Editor::DataStorage::TTypedElementColumnTypeList<FWidgetFactoryColumn>(),
+		FName("Editor_WidgetFactoryTable"));
+	
+	DecoratorWidgetFactoryTable = Storage->RegisterTable(
+		UE::Editor::DataStorage::TTypedElementColumnTypeList<FWidgetFactoryColumn, FDecoratorWidgetFactoryTag>(),
+		FName("Editor_DecoratorWidgetFactoryTable"));
 }
 
 void UEditorDataStorageUi::RegisterQueries()
@@ -824,23 +844,34 @@ void UEditorDataStorageUi::RegisterQueries()
 		.Where()
 			.All<FWidgetFactoryColumn>()
 		.Compile());
-	
-	Storage->RegisterQuery(
-		Select(
-			TEXT("Add the name to widget factory with constructor type info"),
-			FObserver::OnAdd<FWidgetFactoryConstructorColumn>().SetExecutionMode(EExecutionMode::GameThread),
-			[this](IQueryContext& Context, RowHandle Row, const FWidgetFactoryConstructorColumn& ConstructorColumn)
+}
+
+void UEditorDataStorageUi::PostDataStorageUpdate()
+{
+	// For any widgets that need their decorators updated, re-create decorators based on the columns on the widget row
+	for (UE::Editor::DataStorage::RowHandle WidgetRow : WidgetsToUpdateDecorators)
+	{
+		if (FTypedElementSlateWidgetReferenceColumn* WidgetReferenceColumn = Storage->GetColumn<FTypedElementSlateWidgetReferenceColumn>(WidgetRow))
+		{
+			if (WidgetReferenceColumn->WidgetConstructor.IsValid())
 			{
-				Context.AddColumn(Row, FWidgetConstructorNameColumn{ .Name = ConstructorColumn.Constructor->GetTypeInfo()->GetFName() });
-			})
-		.Where()
-			.All<FWidgetFactoryColumn>()
-		.Compile());
+				if (TSharedPtr<UE::Editor::DataStorage::ITedsWidget> ContainerWidget = WidgetReferenceColumn->TedsWidget.Pin())
+				{
+					TSharedPtr<SWidget> InternalWidget = CreateDecoratorWidgets(WidgetRow, WidgetReferenceColumn->Widget.Pin(),
+						*WidgetReferenceColumn->WidgetConstructor.Pin(), UE::Editor::DataStorage::FMetaDataView());
+
+					// Add the final composite widget that contains all the decorators and the actual widget to the container ITedsWidget
+					ContainerWidget->SetContent(InternalWidget.ToSharedRef());
+				}
+			}
+		}
+	}
+
+	WidgetsToUpdateDecorators.Empty();
 }
 
 bool UEditorDataStorageUi::CreateWidgetConstructors_LongestMatch(const TArray<UE::Editor::DataStorage::RowHandle>& WidgetFactories,
-	TArray<TWeakObjectPtr<const UScriptStruct>>& Columns, const UE::Editor::DataStorage::FMetaDataView& Arguments,
-	const WidgetConstructorCallback& Callback)
+	TArray<TWeakObjectPtr<const UScriptStruct>>& Columns, FWidgetConstructionMethod ConstructionFunction)
 {
 	TArray<TWeakObjectPtr<const UScriptStruct>> MatchedColumns;
 	for (auto FactoryIt = WidgetFactories.CreateConstIterator(); FactoryIt && !Columns.IsEmpty();)
@@ -899,7 +930,7 @@ bool UEditorDataStorageUi::CreateWidgetConstructors_LongestMatch(const TArray<UE
 				--ColumnsEnd;
 			}
 			
-			if (!CreateSingleWidgetConstructor(*FactoryIt, Arguments, MoveTemp(RequestedColumnsThatMatched), Callback))
+			if (!ConstructionFunction(*FactoryIt, MoveTemp(RequestedColumnsThatMatched)))
 			{
 				return false;
 			}
@@ -917,8 +948,7 @@ bool UEditorDataStorageUi::CreateWidgetConstructors_LongestMatch(const TArray<UE
 }
 
 void UEditorDataStorageUi::CreateWidgetConstructors_ExactMatch(const TArray<UE::Editor::DataStorage::RowHandle>& WidgetFactories,
-	TArray<TWeakObjectPtr<const UScriptStruct>>& Columns, const UE::Editor::DataStorage::FMetaDataView& Arguments,
-	const WidgetConstructorCallback& Callback)
+	TArray<TWeakObjectPtr<const UScriptStruct>>& Columns, FWidgetConstructionMethod ConstructionFunction)
 {
 	int32 ColumnCount = Columns.Num();
 	TArray<TWeakObjectPtr<const UScriptStruct>> MatchedColumns;
@@ -948,7 +978,7 @@ void UEditorDataStorageUi::CreateWidgetConstructors_ExactMatch(const TArray<UE::
 			if (MatchedColumns.Num() == Columns.Num())
 			{
 				Columns.Reset();
-				CreateSingleWidgetConstructor(FactoryRow, Arguments, MoveTemp(MatchedColumns), Callback);
+				ConstructionFunction(FactoryRow, MoveTemp(MatchedColumns));
 				return;
 			}
 		}
@@ -956,15 +986,14 @@ void UEditorDataStorageUi::CreateWidgetConstructors_ExactMatch(const TArray<UE::
 }
 
 void UEditorDataStorageUi::CreateWidgetConstructors_SingleMatch(const TArray<UE::Editor::DataStorage::RowHandle>& WidgetFactories,
-	TArray<TWeakObjectPtr<const UScriptStruct>>& Columns, const UE::Editor::DataStorage::FMetaDataView& Arguments,
-	const WidgetConstructorCallback& Callback)
+	TArray<TWeakObjectPtr<const UScriptStruct>>& Columns, FWidgetConstructionMethod ConstructionFunction)
 {
-	auto FactoryIt = WidgetFactories.rbegin();
-	auto FactoryEnd = WidgetFactories.rend();
-
 	// Start from the back as the widgets with lower counts will be last.
 	for (int32 ColumnIndex = Columns.Num() - 1; ColumnIndex >= 0; --ColumnIndex)
 	{
+		auto FactoryIt = WidgetFactories.rbegin();
+		auto FactoryEnd = WidgetFactories.rend();
+
 		for (; FactoryIt != FactoryEnd; ++FactoryIt)
 		{
 			const UE::Editor::DataStorage::Queries::FConditions& Conditions = GetFactoryConditions(*FactoryIt);
@@ -987,9 +1016,7 @@ void UEditorDataStorageUi::CreateWidgetConstructors_SingleMatch(const TArray<UE:
 				// dynamic column and the widget constructor wants the latter
 				TWeakObjectPtr<const UScriptStruct> RequestedColumn = Columns[ColumnIndex];
 				Columns.RemoveAt(ColumnIndex);
-				CreateSingleWidgetConstructor(*FactoryIt, Arguments, {RequestedColumn}, Callback);
-				
-				
+				ConstructionFunction(*FactoryIt, {RequestedColumn});
 				// Match was found so move on to the next column in the column.
 				break;
 			} 
@@ -997,17 +1024,17 @@ void UEditorDataStorageUi::CreateWidgetConstructors_SingleMatch(const TArray<UE:
 	}
 }
 
-UE::Editor::DataStorage::RowHandle UEditorDataStorageUi::RegisterWidgetFactoryRow(UE::Editor::DataStorage::RowHandle PurposeRowHandle) const
+UE::Editor::DataStorage::RowHandle UEditorDataStorageUi::RegisterWidgetFactoryRow(UE::Editor::DataStorage::RowHandle PurposeRowHandle, bool bDecoratorWidgetFactory) const
 {
-	UE::Editor::DataStorage::RowHandle FactoryRowHandle = Storage->AddRow(WidgetFactoryTable);
+	UE::Editor::DataStorage::RowHandle FactoryRowHandle = bDecoratorWidgetFactory ? Storage->AddRow(DecoratorWidgetFactoryTable) : Storage->AddRow(WidgetFactoryTable);
 	Storage->GetColumn<FWidgetFactoryColumn>(FactoryRowHandle)->PurposeRowHandle = PurposeRowHandle;
-
 	return FactoryRowHandle;
 }
 
-UE::Editor::DataStorage::RowHandle UEditorDataStorageUi::RegisterUniqueWidgetFactoryRow(UE::Editor::DataStorage::RowHandle InPurposeRowHandle) const
+UE::Editor::DataStorage::RowHandle UEditorDataStorageUi::RegisterUniqueWidgetFactoryRow(UE::Editor::DataStorage::RowHandle InPurposeRowHandle, bool bDecoratorWidgetFactory) const
 {
 	using namespace UE::Editor::DataStorage::Queries;
+	using namespace UE::Editor::DataStorage::Ui;
 	
 	static QueryHandle FactoryQueryHandle = Storage->RegisterQuery(Select().ReadOnly<FWidgetFactoryColumn>().Compile());
 
@@ -1026,23 +1053,57 @@ UE::Editor::DataStorage::RowHandle UEditorDataStorageUi::RegisterUniqueWidgetFac
 	// If there was a factory already registered for this purpose, overwrite its information
 	if (Storage->IsRowAvailable(FactoryRowHandle))
 	{
-		Storage->RemoveColumns<FWidgetFactoryConstructorColumn, FWidgetFactoryConstructorTypeInfoColumn>(FactoryRowHandle);
+		Storage->RemoveColumns<FWidgetFactoryConstructorTypeInfoColumn>(FactoryRowHandle);
 		Storage->GetColumn<FWidgetFactoryColumn>(FactoryRowHandle)->PurposeRowHandle = InPurposeRowHandle;
+
+		// If we are registering a decorator widget factory, ensure the correct tag is added (and remove if not). 
+		bDecoratorWidgetFactory ? Storage->AddColumn<FDecoratorWidgetFactoryTag>(FactoryRowHandle) : Storage->RemoveColumn<FDecoratorWidgetFactoryTag>(FactoryRowHandle);
 		return FactoryRowHandle;
 	}
 	// Otherwise just register the factory row as usual
 	else
 	{
-		return RegisterWidgetFactoryRow(InPurposeRowHandle);
+		return RegisterWidgetFactoryRow(InPurposeRowHandle, bDecoratorWidgetFactory);
 	}
 }
 
-void UEditorDataStorageUi::GetFactories(UE::Editor::DataStorage::RowHandle PurposeRowHandle,
+void UEditorDataStorageUi::GetWidgetFactories(UE::Editor::DataStorage::RowHandle PurposeRowHandle,
 	TArray<UE::Editor::DataStorage::RowHandle>& OutFactories) const
 {
 	using namespace UE::Editor::DataStorage::Queries;
+	using namespace UE::Editor::DataStorage::Ui;
 	
-	static QueryHandle FactoryQueryHandle = Storage->RegisterQuery(Select().ReadOnly<FWidgetFactoryColumn>().Compile());
+	static QueryHandle FactoryQueryHandle =
+		Storage->RegisterQuery(
+			Select()
+				.ReadOnly<FWidgetFactoryColumn>()
+			.Where()
+				.None<FDecoratorWidgetFactoryTag>()
+			.Compile());
+	
+	Storage->RunQuery(FactoryQueryHandle, CreateDirectQueryCallbackBinding(
+		[PurposeRowHandle, &OutFactories](RowHandle RowHandle, const FWidgetFactoryColumn& PurposeReferenceColumn)
+		{
+			if (PurposeReferenceColumn.PurposeRowHandle == PurposeRowHandle)
+			{
+				OutFactories.Add(RowHandle);
+			}
+		}));
+}
+
+void UEditorDataStorageUi::GetDecoratorWidgetFactories(UE::Editor::DataStorage::RowHandle PurposeRowHandle,
+	TArray<UE::Editor::DataStorage::RowHandle>& OutFactories) const
+{
+	using namespace UE::Editor::DataStorage::Queries;
+	using namespace UE::Editor::DataStorage::Ui;
+	
+	static QueryHandle FactoryQueryHandle =
+		Storage->RegisterQuery(
+			Select()
+				.ReadOnly<FWidgetFactoryColumn>()
+			.Where()
+				.All<FDecoratorWidgetFactoryTag>()
+			.Compile());
 	
 	Storage->RunQuery(FactoryQueryHandle, CreateDirectQueryCallbackBinding(
 		[PurposeRowHandle, &OutFactories](RowHandle RowHandle, const FWidgetFactoryColumn& PurposeReferenceColumn)
@@ -1073,4 +1134,112 @@ const UE::Editor::DataStorage::Queries::FConditions& UEditorDataStorageUi::GetFa
 	
 }
 
+TSharedPtr<SWidget> UEditorDataStorageUi::CreateDecoratorWidgets(UE::Editor::DataStorage::RowHandle WidgetRowHandle, TSharedPtr<SWidget> InternalWidget, FTypedElementWidgetConstructor& Constructor,
+	const UE::Editor::DataStorage::FMetaDataView& Arguments)
+{
+	using namespace UE::Editor::DataStorage;
+	RowHandle FactoryRow = Constructor.GetWidgetFactoryRow();
+
+	if (FWidgetFactoryColumn* FactoryColumn = Storage->GetColumn<FWidgetFactoryColumn>(FactoryRow))
+	{
+		RowHandle PurposeRow = FactoryColumn->PurposeRowHandle;
+
+		if (Storage->IsRowAvailable(PurposeRow))
+		{
+			// First get all the decorator widgets for this purpose and all the columns on the UI row so we can figure out which decorators are
+			// actually needed
+			TArray<RowHandle> Factories;
+			GetDecoratorWidgetFactories(PurposeRow, Factories);
+			
+			TArray<TWeakObjectPtr<const UScriptStruct>> Columns;
+			Storage->ListColumns(WidgetRowHandle, [&Columns](const UScriptStruct& ColumnType)
+			{
+				Columns.Emplace(&ColumnType);
+			});
+
+			Columns.Sort(
+				[](const TWeakObjectPtr<const UScriptStruct>& Lhs, const TWeakObjectPtr<const UScriptStruct>& Rhs)
+				{
+					return Lhs->GetName() < Rhs->GetName();
+				});
+
+			TArray<TUniquePtr<FTedsDecoratorWidgetConstructor>> Decorators;
+
+			FWidgetConstructionMethod ConstructionCallback = [this, Arguments, &Decorators]
+				(UE::Editor::DataStorage::RowHandle FactoryRow, TArray<TWeakObjectPtr<const UScriptStruct>> MatchedColumns)
+			{
+				TUniquePtr<FTedsDecoratorWidgetConstructor> Constructor = CreateSingleDecoratorWidget(FactoryRow, Arguments, MatchedColumns);
+				Decorators.Add(MoveTemp(Constructor));
+				return true;
+			};
+
+			// Decorator widgets can only match against a single column, so we can re-use the single match algorithm to figure out which ones
+			// match the columns on the widget row
+			CreateWidgetConstructors_SingleMatch(Factories, Columns, ConstructionCallback);
+
+			// If this widget was created for a specific data row, also let the decorator know about that
+			RowHandle TargetRow = InvalidRowHandle;
+			if (FTypedElementRowReferenceColumn* RowReferenceColumn = Storage->GetColumn<FTypedElementRowReferenceColumn>(WidgetRowHandle))
+			{
+				TargetRow = RowReferenceColumn->Row;
+			}
+
+			// Create each decorator widget, compounding the result widget of the previous operation onto the current widget
+			for (TUniquePtr<FTedsDecoratorWidgetConstructor>& ConstructorIt : Decorators)
+			{
+				InternalWidget = ConstructorIt->CreateDecoratorWidget(InternalWidget, Storage, this, TargetRow, WidgetRowHandle, Arguments);
+			}
+		}
+	}
+
+	return InternalWidget;
+}
+
+void UEditorDataStorageUi::RegisterDecoratorWidgetObservers(UE::Editor::DataStorage::RowHandle PurposeRowHandle, const UScriptStruct* DecoratorColumn)
+{
+	using namespace UE::Editor::DataStorage::Queries;
+
+	FString OnAddName = TEXT("TEDS UI Decorator Monitor: OnAdd - ");
+	DecoratorColumn->GetFName().AppendString(OnAddName);
+
+	FQueryDescription OnAddObserver =
+		Select(
+			FName(OnAddName),
+			FObserver(FObserver::EEvent::Add, DecoratorColumn).SetExecutionMode(EExecutionMode::GameThread),
+			[PurposeRowHandle, this](IQueryContext& Context, RowHandle Row, const FWidgetPurposeReferenceColumn& PurposeColumn)
+			{
+				if (PurposeColumn.PurposeRowHandle == PurposeRowHandle)
+				{
+					// Defer the actual decorator update since widget construction could be accessing the data storage directly for data
+					WidgetsToUpdateDecorators.Add(Row);
+				}
+			})
+		.Where()
+			.All<FTypedElementSlateWidgetReferenceColumn>()
+		.Compile();
+
+	DecoratorObservers.Add(Storage->RegisterQuery(MoveTemp(OnAddObserver)));
+
+	FString OnRemoveName = TEXT("TEDS UI Decorator Monitor: OnRemove - ");
+	DecoratorColumn->GetFName().AppendString(OnRemoveName);
+
+	FQueryDescription OnRemoveObserver =
+		Select(
+			FName(OnRemoveName),
+			FObserver(FObserver::EEvent::Remove, DecoratorColumn).SetExecutionMode(EExecutionMode::GameThread),
+			[PurposeRowHandle, this](IQueryContext& Context, RowHandle Row, const FWidgetPurposeReferenceColumn& PurposeColumn)
+			{
+				if (PurposeColumn.PurposeRowHandle == PurposeRowHandle)
+				{
+					// Defer the actual decorator update since widget construction could be accessing the data storage directly for data
+					WidgetsToUpdateDecorators.Add(Row);
+				}
+			})
+		.Where()
+			.All<FTypedElementSlateWidgetReferenceColumn>()
+		.Compile();
+
+	DecoratorObservers.Add(Storage->RegisterQuery(MoveTemp(OnRemoveObserver)));
+}
+
 #undef LOCTEXT_NAMESPACE // "TypedElementDatabaseUI"
diff --git a/Engine/Plugins/Experimental/EditorDataStorage/Source/TedsCore/Private/TypedElementDatabaseUI.h b/Engine/Plugins/Experimental/EditorDataStorage/Source/TedsCore/Private/TypedElementDatabaseUI.h
--- a/Engine/Plugins/Experimental/EditorDataStorage/Source/TedsCore/Private/TypedElementDatabaseUI.h
+++ b/Engine/Plugins/Experimental/EditorDataStorage/Source/TedsCore/Private/TypedElementDatabaseUI.h
@@ -5,6 +5,7 @@
 #include <type_traits>
 #include "Containers/Map.h"
 #include "DataStorage/Queries/Conditions.h"
+#include "Elements/Interfaces/DecoratorWidgetConstructor.h"
 #include "Elements/Interfaces/TypedElementDataStorageUiInterface.h"
 #include "Logging/LogMacros.h"
 #include "UObject/ObjectMacros.h"
@@ -46,10 +47,8 @@ public:
 	UE_API virtual bool RegisterWidgetFactory(UE::Editor::DataStorage::RowHandle PurposeRow, const UScriptStruct* Constructor) override;
 	UE_API virtual bool RegisterWidgetFactory(UE::Editor::DataStorage::RowHandle PurposeRow, const UScriptStruct* Constructor,
 		UE::Editor::DataStorage::Queries::FConditions Columns) override;
-	UE_API virtual bool RegisterWidgetFactory(UE::Editor::DataStorage::RowHandle PurposeRow,
-		TUniquePtr<FTypedElementWidgetConstructor>&& Constructor) override;
-	UE_API virtual bool RegisterWidgetFactory(UE::Editor::DataStorage::RowHandle PurposeRow, TUniquePtr<FTypedElementWidgetConstructor>&& Constructor, 
-		UE::Editor::DataStorage::Queries::FConditions Columns) override;
+	UE_API virtual bool RegisterDecoratorWidgetFactory(UE::Editor::DataStorage::RowHandle PurposeRow, const UScriptStruct* Constructor,
+		const UScriptStruct* Column) override;
 	
 	UE_API virtual void CreateWidgetConstructors(UE::Editor::DataStorage::RowHandle PurposeRow, 
 		const UE::Editor::DataStorage::FMetaDataView& Arguments, const WidgetConstructorCallback& Callback) override;
@@ -58,7 +57,11 @@ public:
 	UE_API virtual void ConstructWidgets(UE::Editor::DataStorage::RowHandle PurposeRow, const UE::Editor::DataStorage::FMetaDataView& Arguments,
 			const WidgetCreatedCallback& ConstructionCallback) override;
 
-	// Deprecated FName overloads
+	// Deprecated functions
+	UE_API virtual bool RegisterWidgetFactory(UE::Editor::DataStorage::RowHandle PurposeRow,
+	TUniquePtr<FTypedElementWidgetConstructor>&& Constructor) override;
+	UE_API virtual bool RegisterWidgetFactory(UE::Editor::DataStorage::RowHandle PurposeRow, TUniquePtr<FTypedElementWidgetConstructor>&& Constructor, 
+		UE::Editor::DataStorage::Queries::FConditions Columns) override;
 	UE_API void RegisterWidgetPurpose(FName Purpose, EPurposeType Type, FText Description) override;
 	UE_API bool RegisterWidgetFactory(FName Purpose, const UScriptStruct* Constructor) override;
 	UE_API bool RegisterWidgetFactory(FName Purpose, const UScriptStruct* Constructor,
@@ -72,11 +75,14 @@ public:
 		const UE::Editor::DataStorage::FMetaDataView& Arguments, const WidgetConstructorCallback& Callback) override;
 	UE_API void ConstructWidgets(FName Purpose, const UE::Editor::DataStorage::FMetaDataView& Arguments,
 		const WidgetCreatedCallback& ConstructionCallback) override;
-	// ~FName overloads
+	// ~Deprecated functions
 	
 	UE_API TSharedPtr<SWidget> ConstructWidget(UE::Editor::DataStorage::RowHandle Row, FTypedElementWidgetConstructor& Constructor,
 		const UE::Editor::DataStorage::FMetaDataView& Arguments) override;
 
+	virtual TSharedPtr<SWidget> ConstructInternalWidget(UE::Editor::DataStorage::RowHandle Row, FTypedElementWidgetConstructor& Constructor,
+		const UE::Editor::DataStorage::FMetaDataView& Arguments) override;
+	
 	UE_API void ListWidgetPurposes(const WidgetPurposeCallback& Callback) const override;
 
 	UE_API bool SupportsExtension(FName Extension) const override;
@@ -108,6 +114,9 @@ public:
 private:
 	UE_API void CreateStandardArchetypes();
 	UE_API void RegisterQueries();
+	void PostDataStorageUpdate();
+
+	using FWidgetConstructionMethod = TFunction<bool(UE::Editor::DataStorage::RowHandle, TArray<TWeakObjectPtr<const UScriptStruct>>)>;
 	
 	UE_API bool CreateSingleWidgetConstructor(
 		UE::Editor::DataStorage::RowHandle FactoryRow,
@@ -127,40 +136,59 @@ private:
 
 	UE_API bool CreateWidgetConstructors_LongestMatch(
 		const TArray<UE::Editor::DataStorage::RowHandle>& WidgetFactories, 
-		TArray<TWeakObjectPtr<const UScriptStruct>>& Columns, 
-		const UE::Editor::DataStorage::FMetaDataView& Arguments,
-		const WidgetConstructorCallback& Callback);
+		TArray<TWeakObjectPtr<const UScriptStruct>>& Columns,
+		FWidgetConstructionMethod ConstructionFunction);
+	
 	UE_API void CreateWidgetConstructors_ExactMatch(
 		const TArray<UE::Editor::DataStorage::RowHandle>& WidgetFactories,
 		TArray<TWeakObjectPtr<const UScriptStruct>>& Columns,
-		const UE::Editor::DataStorage::FMetaDataView& Arguments,
-		const WidgetConstructorCallback& Callback);
+		FWidgetConstructionMethod ConstructionFunction);
+	
 	UE_API void CreateWidgetConstructors_SingleMatch(
 		const TArray<UE::Editor::DataStorage::RowHandle>& WidgetFactories,
 		TArray<TWeakObjectPtr<const UScriptStruct>>& Columns,
+		FWidgetConstructionMethod ConstructionFunction);
+
+	UE_API TUniquePtr<FTedsDecoratorWidgetConstructor> CreateSingleDecoratorWidget(
+		UE::Editor::DataStorage::RowHandle FactoryRow,
 		const UE::Editor::DataStorage::FMetaDataView& Arguments,
-		const WidgetConstructorCallback& Callback);
+		TArray<TWeakObjectPtr<const UScriptStruct>> MatchedColumnTypes);
 
 	// Register a widget factory for the provided purpose
-	UE_API UE::Editor::DataStorage::RowHandle RegisterWidgetFactoryRow(UE::Editor::DataStorage::RowHandle PurposeRowHandle) const;
+	UE_API UE::Editor::DataStorage::RowHandle RegisterWidgetFactoryRow(UE::Editor::DataStorage::RowHandle PurposeRowHandle, bool bDecoratorWidgetFactory = false) const;
 
 	// Register a unique factory for the provided purpose, clearing the info if there was any factory previously registered for the purpose
-	UE_API UE::Editor::DataStorage::RowHandle RegisterUniqueWidgetFactoryRow(UE::Editor::DataStorage::RowHandle PurposeRowHandle) const;
+	UE_API UE::Editor::DataStorage::RowHandle RegisterUniqueWidgetFactoryRow(UE::Editor::DataStorage::RowHandle PurposeRowHandle, bool bDecoratorWidgetFactory = false) const;
+
+	// Get all (non-decorator) factories for the provided purpose
+	UE_API void GetWidgetFactories(UE::Editor::DataStorage::RowHandle PurposeRowHandle, TArray<UE::Editor::DataStorage::RowHandle>& OutFactories) const;
 
-	// Get all factories for the provided purpose
-	UE_API void GetFactories(UE::Editor::DataStorage::RowHandle PurposeRowHandle, TArray<UE::Editor::DataStorage::RowHandle>& OutFactories) const;
+	// Get all decorator factories for the provided purpose
+	UE_API void GetDecoratorWidgetFactories(UE::Editor::DataStorage::RowHandle PurposeRowHandle, TArray<UE::Editor::DataStorage::RowHandle>& OutFactories) const;
 
 	// The query conditions for the provided factory
 	UE_API const UE::Editor::DataStorage::Queries::FConditions& GetFactoryConditions(UE::Editor::DataStorage::RowHandle FactoryRow) const;
 
+	// Create decorator widgets for the given TEDS UI widget
+	UE_API TSharedPtr<SWidget> CreateDecoratorWidgets(UE::Editor::DataStorage::RowHandle WidgetRowHandle, TSharedPtr<SWidget> InternalWidget, FTypedElementWidgetConstructor& Constructor, const UE::Editor::DataStorage::FMetaDataView& Arguments);
+
+	// Register observers for the given decorator + purpose combo to automatically add/remove the decorator widget when the column is added/removed from widget rows
+	UE_API void RegisterDecoratorWidgetObservers(UE::Editor::DataStorage::RowHandle PurposeRowHandle, const UScriptStruct* DecoratorColumn);
+
 	TMap<const FFieldClass*, PropertySorterConstructorCallback> PropertySorterConstructors;
 
 	UE::Editor::DataStorage::TableHandle WidgetTable{ UE::Editor::DataStorage::InvalidTableHandle };
 	UE::Editor::DataStorage::TableHandle WidgetPurposeTable{ UE::Editor::DataStorage::InvalidTableHandle };
 	UE::Editor::DataStorage::TableHandle WidgetFactoryTable{ UE::Editor::DataStorage::InvalidTableHandle };
-	
+	UE::Editor::DataStorage::TableHandle DecoratorWidgetFactoryTable{ UE::Editor::DataStorage::InvalidTableHandle };
+
 	UE::Editor::DataStorage::ICoreProvider* Storage{ nullptr };
 	UE::Editor::DataStorage::ICompatibilityProvider* StorageCompatibility{ nullptr };
+
+	TArray<UE::Editor::DataStorage::QueryHandle> DecoratorObservers;
+
+	// Widget rows that need their decorators updated
+	TSet<UE::Editor::DataStorage::RowHandle> WidgetsToUpdateDecorators;
 };
 
 #undef UE_API
diff --git a/Engine/Plugins/Experimental/EditorDataStorage/Source/TedsUI/Private/Widgets/CounterWidget.cpp b/Engine/Plugins/Experimental/EditorDataStorage/Source/TedsUI/Private/Widgets/CounterWidget.cpp
deleted file mode 100644
diff --git a/Engine/Plugins/Experimental/EditorDataStorage/Source/TedsUI/Public/Widgets/CounterWidget.h b/Engine/Plugins/Experimental/EditorDataStorage/Source/TedsUI/Public/Widgets/CounterWidget.h
deleted file mode 100644
diff --git a/Engine/Plugins/Experimental/EditorDataStorage/Source/TedsUI/Public/Widgets/ExportedTextWidget.h b/Engine/Plugins/Experimental/EditorDataStorage/Source/TedsUI/Public/Widgets/ExportedTextWidget.h
--- a/Engine/Plugins/Experimental/EditorDataStorage/Source/TedsUI/Public/Widgets/ExportedTextWidget.h
+++ b/Engine/Plugins/Experimental/EditorDataStorage/Source/TedsUI/Public/Widgets/ExportedTextWidget.h
@@ -38,6 +38,9 @@ public:
 	UE_API FExportedTextWidgetConstructor();
 	virtual ~FExportedTextWidgetConstructor() override = default;
 
+	// To avoid name hiding since we are only overriding one of the GetQueryConditions overload
+	using FTypedElementWidgetConstructor::GetQueryConditions;
+	
 	UE_API virtual TConstArrayView<const UScriptStruct*> GetAdditionalColumnsList() const override;
 	UE_API virtual const UE::Editor::DataStorage::Queries::FConditions* GetQueryConditions(const UE::Editor::DataStorage::ICoreProvider* Storage) const override;
 
@@ -54,6 +57,7 @@ public:
 protected:
 	// The column this exported text widget is operating on
 	UE::Editor::DataStorage::Queries::FConditions MatchedColumn;
+
 };
 
 USTRUCT(meta = (DisplayName = "Exported text widget"))
diff --git a/Engine/Plugins/Experimental/EditorDataStorageFeatures/Source/TedsTableViewer/Private/TedsTableViewerColumn.cpp b/Engine/Plugins/Experimental/EditorDataStorageFeatures/Source/TedsTableViewer/Private/TedsTableViewerColumn.cpp
--- a/Engine/Plugins/Experimental/EditorDataStorageFeatures/Source/TedsTableViewer/Private/TedsTableViewerColumn.cpp
+++ b/Engine/Plugins/Experimental/EditorDataStorageFeatures/Source/TedsTableViewer/Private/TedsTableViewerColumn.cpp
@@ -2,6 +2,7 @@
 
 #include "TedsTableViewerColumn.h"
 
+#include "EditorAnalytics.h"
 #include "Columns/SlateHeaderColumns.h"
 #include "DataStorage/Features.h"
 #include "Elements/Columns/TypedElementMiscColumns.h"
@@ -375,11 +376,9 @@ namespace UE::Editor::DataStorage
 			}
 
 			// If a column was added and we are here, we need to re-create the widget
-			// TEDS-Outliner TODO: Do we need to create the widget only if it doesn't exist? Or should we also update it to automatically respond
-			// to column changes even if it was already created
 			if(*bColumnAddedPtr)
 			{
-				if(const TSharedPtr<SWidget> RowWidget = CellWidgetConstructor->Construct(Row, Storage, StorageUi, WidgetMetaData))
+				if(const TSharedPtr<SWidget> RowWidget = StorageUi->ConstructInternalWidget(Row, *CellWidgetConstructor, WidgetMetaData))
 				{
 					TedsWidget->SetContent(RowWidget.ToSharedRef());
 				}
diff --git a/Engine/Source/Runtime/TypedElementFramework/Private/Elements/Interfaces/DecoratorWidgetConstructor.cpp b/Engine/Source/Runtime/TypedElementFramework/Private/Elements/Interfaces/DecoratorWidgetConstructor.cpp
new file mode 100644
--- /dev/null
+++ b/Engine/Source/Runtime/TypedElementFramework/Private/Elements/Interfaces/DecoratorWidgetConstructor.cpp
@@ -0,0 +1,16 @@
+// Copyright Epic Games, Inc. All Rights Reserved.
+
+#include "Elements/Interfaces/DecoratorWidgetConstructor.h"
+
+FTedsDecoratorWidgetConstructor::FTedsDecoratorWidgetConstructor(const UScriptStruct* InTypeInfo)
+	: FTedsWidgetConstructorBase(InTypeInfo)
+{
+}
+
+TSharedPtr<SWidget> FTedsDecoratorWidgetConstructor::CreateDecoratorWidget(TSharedPtr<SWidget> InChildWidget,
+	UE::Editor::DataStorage::ICoreProvider* DataStorage, UE::Editor::DataStorage::IUiProvider* DataStorageUi,
+	UE::Editor::DataStorage::RowHandle DataRow, UE::Editor::DataStorage::RowHandle WidgetRow,
+	const UE::Editor::DataStorage::FMetaDataView& Arguments)
+{
+	return nullptr;
+}
diff --git a/Engine/Source/Runtime/TypedElementFramework/Private/Elements/Interfaces/TypedElementDataStorageUiInterface.cpp b/Engine/Source/Runtime/TypedElementFramework/Private/Elements/Interfaces/TypedElementDataStorageUiInterface.cpp
--- a/Engine/Source/Runtime/TypedElementFramework/Private/Elements/Interfaces/TypedElementDataStorageUiInterface.cpp
+++ b/Engine/Source/Runtime/TypedElementFramework/Private/Elements/Interfaces/TypedElementDataStorageUiInterface.cpp
@@ -17,30 +17,22 @@
 #define LOCTEXT_NAMESPACE "FTypedElementWidgetConstructor"
 
 //
-// FTypedElementWidgetConstructor
+// FTedsWidgetConstructorBase
 //
 
-FTypedElementWidgetConstructor::FTypedElementWidgetConstructor(const UScriptStruct* InTypeInfo)
+FTedsWidgetConstructorBase::FTedsWidgetConstructorBase(const UScriptStruct* InTypeInfo)
 	: TypeInfo(InTypeInfo)
 {
 }
 
-bool FTypedElementWidgetConstructor::Initialize(const UE::Editor::DataStorage::FMetaDataView& InArguments,
-	TArray<TWeakObjectPtr<const UScriptStruct>> InMatchedColumnTypes, const UE::Editor::DataStorage::Queries::FConditions& InQueryConditions)
-{
-	
-	QueryConditions = &InQueryConditions;
-	return Initialize_Internal(InArguments, MoveTemp(InMatchedColumnTypes));
-}
-
-bool FTypedElementWidgetConstructor::Initialize(const UE::Editor::DataStorage::FMetaDataView& InArguments,
+bool FTedsWidgetConstructorBase::Initialize(const UE::Editor::DataStorage::FMetaDataView& InArguments,
 	TArray<TWeakObjectPtr<const UScriptStruct>> InMatchedColumnTypes, UE::Editor::DataStorage::RowHandle InFactoryRowHandle)
 {
 	WidgetFactoryRow = InFactoryRowHandle;
 	return Initialize_Internal(InArguments, MoveTemp(InMatchedColumnTypes));
 }
 
-bool FTypedElementWidgetConstructor::Initialize_Internal(const UE::Editor::DataStorage::FMetaDataView& InArguments,
+bool FTedsWidgetConstructorBase::Initialize_Internal(const UE::Editor::DataStorage::FMetaDataView& InArguments,
 	TArray<TWeakObjectPtr<const UScriptStruct>> InMatchedColumnTypes)
 {
 	MatchedColumnTypes = MoveTemp(InMatchedColumnTypes);
@@ -58,24 +50,17 @@ bool FTypedElementWidgetConstructor::Initialize_Internal(const UE::Editor::DataS
 	return true;
 }
 
-const UScriptStruct* FTypedElementWidgetConstructor::GetTypeInfo() const
+const UScriptStruct* FTedsWidgetConstructorBase::GetTypeInfo() const
 {
 	return TypeInfo;
 }
 
-const TArray<TWeakObjectPtr<const UScriptStruct>>& FTypedElementWidgetConstructor::GetMatchedColumns() const
+const TArray<TWeakObjectPtr<const UScriptStruct>>& FTedsWidgetConstructorBase::GetMatchedColumns() const
 {
 	return MatchedColumnTypes;
 }
 
-const UE::Editor::DataStorage::Queries::FConditions* FTypedElementWidgetConstructor::GetQueryConditions() const
-{
-	const UE::Editor::DataStorage::ICoreProvider* Storage = UE::Editor::DataStorage::GetDataStorageFeature<UE::Editor::DataStorage::ICoreProvider>(UE::Editor::DataStorage::StorageFeatureName);
-
-	return GetQueryConditions(Storage);
-}
-
-const UE::Editor::DataStorage::Queries::FConditions* FTypedElementWidgetConstructor::GetQueryConditions(const UE::Editor::DataStorage::ICoreProvider* Storage) const
+const UE::Editor::DataStorage::Queries::FConditions* FTedsWidgetConstructorBase::GetQueryConditions(const UE::Editor::DataStorage::ICoreProvider* Storage) const
 {
 	if (const FWidgetFactoryConditionsColumn* ConditionsColumn = Storage->GetColumn<FWidgetFactoryConditionsColumn>(WidgetFactoryRow))
 	{
@@ -87,6 +72,15 @@ const UE::Editor::DataStorage::Queries::FConditions* FTypedElementWidgetConstruc
 	}
 }
 
+//
+// FTypedElementWidgetConstructor
+//
+
+FTypedElementWidgetConstructor::FTypedElementWidgetConstructor(const UScriptStruct* InTypeInfo)
+	: FTedsWidgetConstructorBase(InTypeInfo)
+{
+}
+
 TConstArrayView<const UScriptStruct*> FTypedElementWidgetConstructor::GetAdditionalColumnsList() const
 {
 	return {};
@@ -182,15 +176,6 @@ TSharedPtr<SWidget> FTypedElementWidgetConstructor::ConstructFinalWidget(
 		DataStorage->AddColumns(Row, GetAdditionalColumnsList());
 		Widget = Construct(Row, DataStorage, DataStorageUi, Arguments);
 	}
-
-	// Create a container widget to hold the content (even if it doesn't exist yet)
-	if (TSharedPtr<UE::Editor::DataStorage::ITedsWidget> ContainerWidget = DataStorageUi->CreateContainerTedsWidget(Row))
-	{
-		ContainerWidget->SetContent(Widget.ToSharedRef());
-		DataStorage->GetColumn<FTypedElementSlateWidgetReferenceColumn>(Row)->TedsWidget = ContainerWidget;
-
-		return ContainerWidget->AsWidget();
-	}
 	
 	return Widget;
 }
@@ -285,6 +270,15 @@ void FTypedElementWidgetConstructor::AddDefaultWidgetColumns(UE::Editor::DataSto
 	const FString WidgetLabel(CreateWidgetDisplayNameText(DataStorage, Row).ToString());
 	DataStorage->AddColumn(Row, FTypedElementLabelColumn{.Label = WidgetLabel} );
 
+	UE::Editor::DataStorage::RowHandle FactoryRow = GetWidgetFactoryRow();
+
+	// Add a reference to the widget purpose this widget is being created for
+	if (FWidgetFactoryColumn* FactoryColumn = DataStorage->GetColumn<FWidgetFactoryColumn>(FactoryRow))
+	{
+		UE::Editor::DataStorage::RowHandle PurposeRow = FactoryColumn->PurposeRowHandle;
+		DataStorage->AddColumn(Row, FWidgetPurposeReferenceColumn{.PurposeRowHandle = PurposeRow});
+	}
+
 	// We don't want to display any second level widgets (widgets for widgets and so on...) in UI because they will cause the table viewer to
 	// infinitely grow as you keep scrolling (which creates new widgets)
 	if(DataStorage->HasColumns<FTypedElementSlateWidgetReferenceColumn>(Row))
@@ -313,6 +307,23 @@ UE::Editor::DataStorage::RowHandle FTypedElementWidgetConstructor::GetTargetRow(
 	return TargetRow;
 }
 
+UE::Editor::DataStorage::RowHandle FTypedElementWidgetConstructor::GetWidgetFactoryRow() const
+{
+	return WidgetFactoryRow;
+}
+
+bool FTypedElementWidgetConstructor::Initialize(const UE::Editor::DataStorage::FMetaDataView& InArguments,
+	TArray<TWeakObjectPtr<const UScriptStruct>> InMatchedColumnTypes, const UE::Editor::DataStorage::Queries::FConditions& InQueryConditions)
+{
+	QueryConditions = &InQueryConditions;
+	return Initialize_Internal(InArguments, MoveTemp(InMatchedColumnTypes));
+}
+
+const UE::Editor::DataStorage::Queries::FConditions* FTypedElementWidgetConstructor::GetQueryConditions() const
+{
+	const UE::Editor::DataStorage::ICoreProvider* Storage = UE::Editor::DataStorage::GetDataStorageFeature<UE::Editor::DataStorage::ICoreProvider>(UE::Editor::DataStorage::StorageFeatureName);
+	return FTedsWidgetConstructorBase::GetQueryConditions(Storage);
+}
 
 
 // FSimpleWidgetConstructor
diff --git a/Engine/Source/Runtime/TypedElementFramework/Public/Elements/Columns/DecoratorWidgetColumns.h b/Engine/Source/Runtime/TypedElementFramework/Public/Elements/Columns/DecoratorWidgetColumns.h
new file mode 100644
--- /dev/null
+++ b/Engine/Source/Runtime/TypedElementFramework/Public/Elements/Columns/DecoratorWidgetColumns.h
@@ -0,0 +1,21 @@
+// Copyright Epic Games, Inc. All Rights Reserved.
+
+#pragma once
+
+#include "Elements/Interfaces/TypedElementDataStorageUiInterface.h"
+
+#include "DecoratorWidgetColumns.generated.h"
+
+/**
+ * Tag used to uniquely identify widget factories for decorator widgets
+ */
+USTRUCT(meta = (DisplayName = "Decorator Widget Factory"))
+struct FEditorDataStorageUiDecoratorWidgetFactoryTag : public FEditorDataStorageTag
+{
+	GENERATED_BODY()
+};
+
+namespace UE::Editor::DataStorage::Ui
+{
+	using FDecoratorWidgetFactoryTag = FEditorDataStorageUiDecoratorWidgetFactoryTag;
+}
\ No newline at end of file
diff --git a/Engine/Source/Runtime/TypedElementFramework/Public/Elements/Columns/TypedElementSlateWidgetColumns.h b/Engine/Source/Runtime/TypedElementFramework/Public/Elements/Columns/TypedElementSlateWidgetColumns.h
--- a/Engine/Source/Runtime/TypedElementFramework/Public/Elements/Columns/TypedElementSlateWidgetColumns.h
+++ b/Engine/Source/Runtime/TypedElementFramework/Public/Elements/Columns/TypedElementSlateWidgetColumns.h
@@ -97,3 +97,14 @@ struct FIsInEditingModeTag final : public FEditorDataStorageTag
 {
 	GENERATED_BODY()
 };
+
+/**
+ * A reference to the widget purpose this widget was created for
+ */
+USTRUCT(meta = (DisplayName = "Widget Purpose Reference"))
+struct FWidgetPurposeReferenceColumn final : public FEditorDataStorageColumn
+{
+	GENERATED_BODY()
+
+	UE::Editor::DataStorage::RowHandle PurposeRowHandle;
+};
diff --git a/Engine/Source/Runtime/TypedElementFramework/Public/Elements/Columns/WidgetPurposeColumns.h b/Engine/Source/Runtime/TypedElementFramework/Public/Elements/Columns/WidgetPurposeColumns.h
--- a/Engine/Source/Runtime/TypedElementFramework/Public/Elements/Columns/WidgetPurposeColumns.h
+++ b/Engine/Source/Runtime/TypedElementFramework/Public/Elements/Columns/WidgetPurposeColumns.h
@@ -45,9 +45,7 @@ struct FWidgetFactoryColumn : public FEditorDataStorageColumn
 	UE::Editor::DataStorage::RowHandle PurposeRowHandle;
 };
 
-/** Column used to store the widget constructor used by a widget factory by its TypeInfo
- *  A factory can either have FWidgetFactoryConstructorTypeInfoColumn or FWidgetFactoryConstructorColumn
- */
+/** Column used to store the widget constructor used by a widget factory by its TypeInfo */
 USTRUCT()
 struct FWidgetFactoryConstructorTypeInfoColumn : public FEditorDataStorageColumn
 {
@@ -56,15 +54,6 @@ struct FWidgetFactoryConstructorTypeInfoColumn : public FEditorDataStorageColumn
 	TWeakObjectPtr<const UScriptStruct> Constructor;
 };
 
-/** Column used to store the widget constructor used by a widget factory */
-USTRUCT()
-struct FWidgetFactoryConstructorColumn : public FEditorDataStorageColumn
-{
-	GENERATED_BODY()
-
-	TUniquePtr<FTypedElementWidgetConstructor> Constructor;
-};
-
 /** Column used to store the name of the widget constructor */
 USTRUCT()
 struct FWidgetConstructorNameColumn : public FEditorDataStorageColumn
@@ -75,15 +64,6 @@ struct FWidgetConstructorNameColumn : public FEditorDataStorageColumn
 	FName Name;
 };
 
-template<>
-struct TStructOpsTypeTraits<FWidgetFactoryConstructorColumn> : public TStructOpsTypeTraitsBase2<FWidgetFactoryConstructorColumn>
-{
-	enum
-	{
-		WithCopy = false
-	};
-};
-
 /** Column used to store the query conditions used by a widget factory
  *  If the column is not present the factory doesn't match against any conditions (i.e is a general purpose factory)
  */
@@ -94,3 +74,34 @@ struct FWidgetFactoryConditionsColumn : public FEditorDataStorageColumn
 
 	UE::Editor::DataStorage::Queries::FConditions Conditions;
 };
+
+// Deprecated columns
+
+USTRUCT()
+struct UE_DEPRECATED(5.7, "Use of the FWidgetFactoryConstructorColumn has been deprecated, use FWidgetFactoryConstructorTypeInfoColumn to store the constructor by type instead.")
+FWidgetFactoryConstructorColumn : public FEditorDataStorageColumn
+{
+	GENERATED_BODY()
+
+	TUniquePtr<FTypedElementWidgetConstructor> Constructor;
+
+	PRAGMA_DISABLE_DEPRECATION_WARNINGS
+	FWidgetFactoryConstructorColumn() = default;
+	~FWidgetFactoryConstructorColumn() = default;
+	FWidgetFactoryConstructorColumn(const FWidgetFactoryConstructorColumn&) = delete;
+	FWidgetFactoryConstructorColumn(FWidgetFactoryConstructorColumn&&) = default;
+	FWidgetFactoryConstructorColumn& operator=(const FWidgetFactoryConstructorColumn&) = delete;
+	FWidgetFactoryConstructorColumn& operator=(FWidgetFactoryConstructorColumn&&) = default;
+	PRAGMA_ENABLE_DEPRECATION_WARNINGS
+};
+
+PRAGMA_DISABLE_DEPRECATION_WARNINGS
+template<>
+struct TStructOpsTypeTraits<FWidgetFactoryConstructorColumn> : public TStructOpsTypeTraitsBase2<FWidgetFactoryConstructorColumn>
+{
+	enum
+	{
+		WithCopy = false
+	};
+};
+PRAGMA_ENABLE_DEPRECATION_WARNINGS
\ No newline at end of file
diff --git a/Engine/Source/Runtime/TypedElementFramework/Public/Elements/Interfaces/DecoratorWidgetConstructor.h b/Engine/Source/Runtime/TypedElementFramework/Public/Elements/Interfaces/DecoratorWidgetConstructor.h
new file mode 100644
--- /dev/null
+++ b/Engine/Source/Runtime/TypedElementFramework/Public/Elements/Interfaces/DecoratorWidgetConstructor.h
@@ -0,0 +1,60 @@
+// Copyright Epic Games, Inc. All Rights Reserved.
+
+#pragma once
+
+#include "DataStorage/Handles.h"
+#include "TypedElementDataStorageInterface.h"
+#include "TypedElementDataStorageUiInterface.h"
+#include "Widgets/SWidget.h"
+
+#include "DecoratorWidgetConstructor.generated.h"
+
+/**
+ * Base class used to for decorator widgets in TEDS UI.
+ * A decorator widget is appended to regular TEDS UI widgets to provide additional functionality such as UI effects, context menus etc
+ * A decorator widget constructor can be registered against a specific purpose and column using IUiProvider.
+ * They will then be automatically added to any widgets where the widget row matches the purpose and column
+ * (both on creation and dynamically post creation)
+ */
+USTRUCT()
+struct FTedsDecoratorWidgetConstructor : public FTedsWidgetConstructorBase
+{
+	GENERATED_BODY()
+
+public:
+	TYPEDELEMENTFRAMEWORK_API explicit FTedsDecoratorWidgetConstructor(const UScriptStruct* InTypeInfo);
+	explicit FTedsDecoratorWidgetConstructor(EForceInit) : FTedsWidgetConstructorBase(ForceInit) {} //< For compatibility and shouldn't be directly used.
+
+	virtual ~FTedsDecoratorWidgetConstructor() = default;
+
+	/**
+	 * Construct a decorator widget that wraps InChildWidget
+	 * @param InChildWidget The actual internal widget, it is expected that the decorator widget you create contains this widget as a child
+	 * @param DataStorage Pointer to the TEDS interface
+	 * @param DataStorageUi Pointer to the TEDS UI interface
+	 * @param DataRow The row containing the actual data InChildWidget was created to observe (e.g Actor row) 
+	 * @param WidgetRow The row containing UI specific data about InChildWidget (e.g widget color)
+	 * @param Arguments Any metadata arguments specified when requesting this decorator
+	 */
+	virtual TSharedPtr<SWidget> CreateDecoratorWidget(
+		TSharedPtr<SWidget> InChildWidget,
+		UE::Editor::DataStorage::ICoreProvider* DataStorage,
+		UE::Editor::DataStorage::IUiProvider* DataStorageUi,
+		UE::Editor::DataStorage::RowHandle DataRow,
+		UE::Editor::DataStorage::RowHandle WidgetRow,
+		const UE::Editor::DataStorage::FMetaDataView& Arguments);
+
+protected:
+	// The constructor's typeinfo
+	TWeakObjectPtr<const UScriptStruct> TypeInfo = nullptr;
+};
+
+template<>
+struct TStructOpsTypeTraits<FTedsDecoratorWidgetConstructor> : public TStructOpsTypeTraitsBase2<FTedsDecoratorWidgetConstructor>
+{
+	enum
+	{
+		WithNoInitConstructor = true,
+		WithPureVirtual = true,
+	};
+};
\ No newline at end of file
diff --git a/Engine/Source/Runtime/TypedElementFramework/Public/Elements/Interfaces/TypedElementDataStorageUiInterface.h b/Engine/Source/Runtime/TypedElementFramework/Public/Elements/Interfaces/TypedElementDataStorageUiInterface.h
--- a/Engine/Source/Runtime/TypedElementFramework/Public/Elements/Interfaces/TypedElementDataStorageUiInterface.h
+++ b/Engine/Source/Runtime/TypedElementFramework/Public/Elements/Interfaces/TypedElementDataStorageUiInterface.h
@@ -28,20 +28,19 @@
 class SWidget;
 
 /**
- * Base class used to construct Typed Element widgets with.
- * See below for the options to register a constructor with the Data Storage.
- * In most cases you want to inherit from FSimpleWidgetConstructor instead which has a simpler pipeline to create widgets
+ * Base class for Widget Constructors in TEDS.
+ * Contains common functionality needed by the different types of widget constructors that TEDS UI supports.
  */
 USTRUCT()
-struct FTypedElementWidgetConstructor
+struct FTedsWidgetConstructorBase
 {
-	GENERATED_BODY()
-
 public:
-	TYPEDELEMENTFRAMEWORK_API explicit FTypedElementWidgetConstructor(const UScriptStruct* InTypeInfo);
-	explicit FTypedElementWidgetConstructor(EForceInit) {} //< For compatibility and shouldn't be directly used.
 
-	virtual ~FTypedElementWidgetConstructor() = default;
+	GENERATED_BODY()
+
+	TYPEDELEMENTFRAMEWORK_API explicit FTedsWidgetConstructorBase(const UScriptStruct* InTypeInfo);
+	explicit FTedsWidgetConstructorBase(EForceInit) {} //< For compatibility and shouldn't be directly used.
+	virtual ~FTedsWidgetConstructorBase() = default;
 	
 	/** Initializes a new constructor based on the provided arguments. */
 	TYPEDELEMENTFRAMEWORK_API virtual bool Initialize(const UE::Editor::DataStorage::FMetaDataView& InArguments,
@@ -54,7 +53,54 @@
 	
 	/** Retrieves the query conditions that need to match for this widget constructor to produce a widget. */
 	TYPEDELEMENTFRAMEWORK_API virtual const UE::Editor::DataStorage::Queries::FConditions* GetQueryConditions(const UE::Editor::DataStorage::ICoreProvider* Storage) const;
+protected:
+	
+	bool Initialize_Internal(const UE::Editor::DataStorage::FMetaDataView& InArguments,
+		TArray<TWeakObjectPtr<const UScriptStruct>> InMatchedColumnTypes);
 
+protected:
+	// The columns this widget constructor matched against
+	TArray<TWeakObjectPtr<const UScriptStruct>> MatchedColumnTypes;
+
+	// If the constructor matched against any dynamic columns, map from the base template for the column -> FName identifier
+	TMap<TWeakObjectPtr<const UScriptStruct>, FName> MatchedDynamicTemplates;
+
+	// The row containing the widget factory this constructor was factories from
+	UE::Editor::DataStorage::RowHandle WidgetFactoryRow = UE::Editor::DataStorage::InvalidRowHandle;
+
+	// The query conditions required by this constructor
+	const UE::Editor::DataStorage::Queries::FConditions* QueryConditions = nullptr;
+
+	// The constructor's typeinfo
+	const UScriptStruct* TypeInfo = nullptr;
+};
+
+template<>
+struct TStructOpsTypeTraits<FTedsWidgetConstructorBase> : public TStructOpsTypeTraitsBase2<FTedsWidgetConstructorBase>
+{
+	enum
+	{
+		WithNoInitConstructor = true,
+		WithPureVirtual = true,
+	};
+};
+
+/**
+ * Base class used to construct TEDS UI widgets with.
+ * See below for the options to register a constructor with the Data Storage.
+ * In most cases you want to inherit from FSimpleWidgetConstructor instead which has a simpler pipeline to create widgets
+ */
+USTRUCT()
+struct FTypedElementWidgetConstructor : public FTedsWidgetConstructorBase
+{
+	GENERATED_BODY()
+
+public:
+	TYPEDELEMENTFRAMEWORK_API explicit FTypedElementWidgetConstructor(const UScriptStruct* InTypeInfo);
+	explicit FTypedElementWidgetConstructor(EForceInit) : FTedsWidgetConstructorBase(ForceInit) {} //< For compatibility and shouldn't be directly used.
+
+	virtual ~FTypedElementWidgetConstructor() override = default;
+	
 	/** Returns a list of additional columns the widget requires to be added to its rows. */
 	TYPEDELEMENTFRAMEWORK_API virtual TConstArrayView<const UScriptStruct*> GetAdditionalColumnsList() const;
 	
@@ -108,15 +154,15 @@
 		UE::Editor::DataStorage::IUiProvider* DataStorageUi,
 		const UE::Editor::DataStorage::FMetaDataView& Arguments);
 
-	/** Initializes a new constructor based on the provided arguments.. */
-	UE_DEPRECATED(5.6, "Use Initialize that takes in the factory row handle instead")
-	TYPEDELEMENTFRAMEWORK_API virtual bool Initialize(const UE::Editor::DataStorage::FMetaDataView& InArguments,
-		TArray<TWeakObjectPtr<const UScriptStruct>> InMatchedColumnTypes, const UE::Editor::DataStorage::Queries::FConditions& InQueryConditions);
+	/**
+	 * Helper function to get the row containing the widget factory this constructor was factoried from
+	 */
+	TYPEDELEMENTFRAMEWORK_API UE::Editor::DataStorage::RowHandle GetWidgetFactoryRow() const;
 
-	/** Retrieves the query conditions that need to match for this widget constructor to produce a widget. */
-	UE_DEPRECATED(5.6, "Use GetQueryConditions that takes in ICoreProvider* instead")
-	TYPEDELEMENTFRAMEWORK_API virtual const UE::Editor::DataStorage::Queries::FConditions* GetQueryConditions() const final;
-	
+	/*
+	 * Deprecated public functions
+	 */
+
 	/**
 	 * Returns a friendly name for the data the created widget represents.
 	 * By default the associated column is used. If there are multiple columns associated with the constructor
@@ -127,6 +173,19 @@
 	TYPEDELEMENTFRAMEWORK_API virtual FString CreateWidgetDisplayName(
 		UE::Editor::DataStorage::ICoreProvider* DataStorage, UE::Editor::DataStorage::RowHandle Row) const final;
 
+	// To get around name hiding since the deprecated functions overload methods from the base class
+	using FTedsWidgetConstructorBase::Initialize;
+	using FTedsWidgetConstructorBase::GetQueryConditions;
+
+	/** Initializes a new constructor based on the provided arguments.. */
+	UE_DEPRECATED(5.6, "Use Initialize that takes in the factory row handle instead")
+	TYPEDELEMENTFRAMEWORK_API virtual bool Initialize(const UE::Editor::DataStorage::FMetaDataView& InArguments,
+		TArray<TWeakObjectPtr<const UScriptStruct>> InMatchedColumnTypes, const UE::Editor::DataStorage::Queries::FConditions& InQueryConditions);
+
+	/** Retrieves the query conditions that need to match for this widget constructor to produce a widget. */
+	UE_DEPRECATED(5.6, "Use GetQueryConditions that takes in ICoreProvider* instead")
+	TYPEDELEMENTFRAMEWORK_API virtual const UE::Editor::DataStorage::Queries::FConditions* GetQueryConditions() const final;
+
 protected:
 	/** Create a new instance of the target widget. This is a required function. */
 	TYPEDELEMENTFRAMEWORK_API virtual TSharedPtr<SWidget> CreateWidget(const UE::Editor::DataStorage::FMetaDataView& Arguments);
@@ -162,27 +221,6 @@
 	 */
 	TYPEDELEMENTFRAMEWORK_API UE::Editor::DataStorage::RowHandle GetTargetRow(
 		UE::Editor::DataStorage::ICoreProvider* DataStorage, UE::Editor::DataStorage::RowHandle WidgetRow) const;
-	
-protected:
-
-	// The columns this widget constructor matched against
-	TArray<TWeakObjectPtr<const UScriptStruct>> MatchedColumnTypes;
-
-	// If the constructor matched against any dynamic columns, map from the base template for the column -> FName identifier
-	TMap<TWeakObjectPtr<const UScriptStruct>, FName> MatchedDynamicTemplates;
-
-	// The row containing the widget factory this constructor was factories from
-	UE::Editor::DataStorage::RowHandle WidgetFactoryRow = UE::Editor::DataStorage::InvalidRowHandle;
-
-	// The query conditions required by this constructor
-	const UE::Editor::DataStorage::Queries::FConditions* QueryConditions = nullptr;
-
-	// The constructor's typeinfo
-	const UScriptStruct* TypeInfo = nullptr;
-
-private:
-	bool Initialize_Internal(const UE::Editor::DataStorage::FMetaDataView& InArguments,
-	TArray<TWeakObjectPtr<const UScriptStruct>> InMatchedColumnTypes);
 };
 
 /**
@@ -431,24 +469,22 @@
 		bool RegisterWidgetFactory(RowHandle PurposeRow, Queries::FConditions Columns);
 
 		/**
-		 * Registers a widget factory that will be called when the purpose it's registered under is requested.
-		 * This version uses a previously created instance of the Constructor. The benefit of this is that it store
-		 * configuration options. The downside is that it takes up more memory and requires copying when it's
-		 * used. If the purpose has not been registered the factory will not be recorded and a warning will be printed.
-		 * If registration is successful true will be returned otherwise false.
+		 * Registers a decorator widget factory against a specific column
+		 * Decorator widgets are added onto TEDS UI widgets for the specified purpose when the widget row
+		 * has the column/tag they are registered against.
+		 * @see FTedsDecoratorWidgetConstructor
 		 */
-		virtual bool RegisterWidgetFactory(RowHandle PurposeRow, TUniquePtr<FTypedElementWidgetConstructor>&& Constructor) = 0;
+		template<TColumnType ColumnType>
+		bool RegisterDecoratorWidgetFactory(RowHandle PurposeRow, const UScriptStruct* Constructor);
+
+		/**
+		 * Registers a decorator widget factory against a specific column.
+		 * Decorator widgets are added onto TEDS UI widgets for the specified purpose when the widget row
+		 * has the column/tag they are registered against.
+		 * @see FTedsDecoratorWidgetConstructor
+		 */
+		virtual bool RegisterDecoratorWidgetFactory(RowHandle PurposeRow, const UScriptStruct* Constructor,
+			const UScriptStruct* Column) = 0;
 	
-		/**
-		 * Registers a widget factory that will be called when the purpose it's registered under is requested.
-		 * This version uses a previously created instance of the Constructor. The benefit of this is that it store
-		 * configuration options. The downside is that it takes up more memory and requires copying when it's
-		 * used. If the purpose has not been registered the factory will not be recorded and a warning will be printed.
-		 * The provided columns will be used when matching the factory during widget construction.
-		 * If registration is successful true will be returned otherwise false.
-		 */
-		virtual bool RegisterWidgetFactory(RowHandle PurposeRow, TUniquePtr<FTypedElementWidgetConstructor>&& Constructor, 
-			Queries::FConditions Columns) = 0;
-	
 		/** 
 		 * Creates widget constructors for the requested purpose.
@@ -490,6 +526,14 @@
 		virtual TSharedPtr<SWidget> ConstructWidget(RowHandle Row, FTypedElementWidgetConstructor& Constructor,
 			const FMetaDataView& Arguments) = 0;
 
+		/** 
+		 * Creates a single widget using the provided constructor. Alternative to ConstructWidget that doesn't create the ContainerTedsWidget that
+		 * wraps the actual widget. Only recommended for advanced use cases, in most situations you want to use ConstructWidget
+		 * NOTE: This will include any decorator widgets that match with the given row
+		 */
+		virtual TSharedPtr<SWidget> ConstructInternalWidget(RowHandle Row, FTypedElementWidgetConstructor& Constructor,
+			const FMetaDataView& Arguments) = 0;
+
 		/** Create the container widget that every TEDS UI widget is stored in */
 		virtual TSharedPtr<ITedsWidget> CreateContainerTedsWidget(RowHandle UiRowHandle) const = 0;
 
@@ -538,6 +582,13 @@
 		 * @section Deprecated
 		 * Please replace use of the functions below with the recommended replacements.
 		 */
+
+		UE_DEPRECATED(5.7, "Registering Widget Factories by instance is not supported anymore, register them by type instead")
+		virtual bool RegisterWidgetFactory(RowHandle PurposeRow, TUniquePtr<FTypedElementWidgetConstructor>&& Constructor) = 0;
+	
+		UE_DEPRECATED(5.7, "Registering Widget Factories by instance is not supported anymore, register them by type instead")
+		virtual bool RegisterWidgetFactory(RowHandle PurposeRow, TUniquePtr<FTypedElementWidgetConstructor>&& Constructor, 
+			Queries::FConditions Columns) = 0;
 		
 		UE_DEPRECATED(5.6, "Use the version that takes in an FPurposeInfo struct instead")
 		virtual void RegisterWidgetPurpose(FName Purpose, EPurposeType Type, FText Description) = 0;
@@ -610,6 +661,12 @@
 		return this->RegisterWidgetFactory(PurposeRow, ConstructorType::StaticStruct(), Columns);
 	}
 
+	template<TColumnType ColumnType>
+	bool IUiProvider::RegisterDecoratorWidgetFactory(RowHandle PurposeRow, const UScriptStruct* Constructor)
+	{
+		return this->RegisterDecoratorWidgetFactory(PurposeRow, Constructor, ColumnType::StaticStruct());
+	}
+
 	template<typename PropertyClassType> requires(std::is_base_of_v<FProperty, PropertyClassType>)
 		void IUiProvider::RegisterSorterGeneratorForProperty(PropertySorterConstructorCallback PropertySorterConstructor)
 	{
