Harden TEDS Compatibility against Async Loaded Object
* The primary goal of TEDS is that if an object is visible to TEDS, it is valid to be used.  A challenge is that with async loading enabled, objects may be initialized on either the async loading thread or the game thread but are not fully loaded.  This means that determining if PostInitProperties is on the GameThread is not a sufficient guarantee that the object is visible to be used in these sitatuations.
* This CL modifies the TEDS Compatibility layer to determine if an object should be visible to TEDS by checking whether an object has the AsyncLoading flags.  If these are present, then that processing of that object will be deferred until the next TEDS Compatibility tick.
* A 'default' strategy was also added which disables this deferral but instead adds an ensure to check that objects are not being added to compaibility for async loading.  Enable this strategy to discover callstacks where objects tagged for async loading are created by the game thread.
* Also adjusted the strategy that ActorComponents use to set up hierarchies.  The processor to set up the hierarchy relies on the ActorComponent in TEDS and will effectively synchronize the hierarchy to the ActorComponent's Owner.  This does rely on the ActorComponent being given a sync tag if the owner changes - Modify() does this and would be required to opt the operation in  into undo-redo so should be a reasonable thing to rely upon.
* Removed actors from adding/removing ActorComponents as compatible objects on add//remove since ActorComponents use the same hooks that Actors do to register with compatibility layer in PostInitiProperties (if GameThread AND, with this change, !AsyncLoading).  In PostLoad (in the case of loading the ActorComponents) and in ULevel::AddLoadedActors.  As stated above, determining the hierarchy is based on ActorComponent's Owner when ActorComponent is synced/modified.

diff --git a/Engine/Plugins/Experimental/EditorDataStorage/Source/TedsCore/Private/TypedElementDatabaseCompatibility.cpp b/Engine/Plugins/Experimental/EditorDataStorage/Source/TedsCore/Private/TypedElementDatabaseCompatibility.cpp
--- a/Engine/Plugins/Experimental/EditorDataStorage/Source/TedsCore/Private/TypedElementDatabaseCompatibility.cpp
+++ b/Engine/Plugins/Experimental/EditorDataStorage/Source/TedsCore/Private/TypedElementDatabaseCompatibility.cpp
@@ -61,6 +61,36 @@ namespace UE::Editor::DataStorage::Private
 			"1 - summarize number of nops\n"
 			"2 - include nops"));
 
+	enum class EAsyncUObjectLoadStrategy : int32
+	{
+		// assert if object that is being asynchronously loaded is added
+		// This strategy will cause assertions to occur with -asyncloadingthread enabled
+		// as UObjects are added from the GameThread but are set with the AsyncLoading flags
+		// indicating that the object is not supposed to be visible for use
+		// Strategy has been provided to make it possible to better understand these situations
+		Assert = 0,
+		// defer registration of async loaded objects into TEDS
+		// This strategy will check UObjects for the presence of the AsyncLoading flags
+		// and will only make the UObject visible to TEDS if these flags are not present.
+		// Registration will be attempted at the next Compatibility tick
+		DeferRegistration = 1,
+		END
+	};
+	
+	int AsyncLoadedUObjectStrategy = 1;
+	FAutoConsoleVariableRef CVarAsyncLoadedObjectStrategy(
+		TEXT("TEDS.Feature.AsyncLoadedUObjectStrategy"),
+		AsyncLoadedUObjectStrategy,
+		TEXT("Strategy used to handle asynchronously loaded objects.\n"
+			 "0 - assert if object that is being asynchronously loaded is added\n"
+			 "1 - defer registration of async loaded objects into TEDS\n"));
+
+	static EAsyncUObjectLoadStrategy GetAsyncLoadStrategy()
+	{
+		checkf(AsyncLoadedUObjectStrategy >= 0 && AsyncLoadedUObjectStrategy < static_cast<int32>(EAsyncUObjectLoadStrategy::END), TEXT("Invalid value for TEDS.Feature.AsyncLoadedUObjectStrategy [%d]"), AsyncLoadedUObjectStrategy);
+		return static_cast<EAsyncUObjectLoadStrategy>(AsyncLoadedUObjectStrategy);
+	}
+
 	static const FName IntegrateWithGCName(TEXT("IntegrateWithGC"));
 	static const FName CompatibilityUsesCommandBufferExtensionName(TEXT("CompatiblityUsesCommandBuffer"));
 } // namespace UE::Editor::DataStorage::Private
@@ -427,6 +457,15 @@ UE::Editor::DataStorage::RowHandle UEditorDataStorageCompatibility::AddCompatibl
 	RowHandle Result = FindRowWithCompatibleObjectExplicit(Object);
 	if (!Storage->IsRowAvailable(Result))
 	{
+		if (Private::GetAsyncLoadStrategy() == Private::EAsyncUObjectLoadStrategy::Assert)
+		{
+			EInternalObjectFlags Flags = Object->GetInternalFlags();
+			ensureMsgf(
+				!(Flags & EInternalObjectFlags_AsyncLoading),
+				TEXT("Cannot add object that is being asynchronously loaded to TEDS"));
+			return InvalidRowHandle;
+		}
+		
 		Result = Storage->ReserveRow();
 		Storage->MapRow(ICompatibilityProvider::ObjectMappingDomain, FMapKey(Object), Result);
 		if (Private::bUseCommandBuffer)
@@ -725,6 +764,15 @@ void UEditorDataStorageCompatibility::PendingRegistration<AddressType>::ProcessE
 		if constexpr (std::is_same_v<AddressType, TWeakObjectPtr<UObject>>)
 		{
 			bIsValid = bIsValid && It->Address.IsValid();
+			if (bIsValid && Private::GetAsyncLoadStrategy() == Private::EAsyncUObjectLoadStrategy::Assert)
+			{
+				// Technically should never ensure due to this check being made when objects are added to the list
+				// but left here for completeness
+				ensureMsgf(
+					!It->Address->HasAnyInternalFlags(EInternalObjectFlags_AsyncLoading),
+					TEXT("UObject that is currently being asynchronously loaded has been queued to be registered with TEDS. Cannot process this entry"));
+				bIsValid = false;
+			}
 		}
 		else if constexpr (std::is_same_v<AddressType, ExternalObjectRegistration>)
 		{
@@ -750,12 +798,39 @@ void UEditorDataStorageCompatibility::PendingRegistration<AddressType>::ProcessE
 		{
 			if constexpr (std::is_same_v<AddressType, TWeakObjectPtr<UObject>>)
 			{
-				Entry.Table = Compatibility.FindBestMatchingTable(Entry.Address->GetClass());
-				checkf(Entry.Table != InvalidTableHandle, 
-					TEXT("The data storage could not find any matching tables for object of type '%s'. "
-					"This can mean that the object doesn't derive from UObject or that a table for UObject is no longer registered."), 
-					*Entry.Address->GetClass()->GetFName().ToString());
-
+				switch (Private::GetAsyncLoadStrategy())
+				{
+					case Private::EAsyncUObjectLoadStrategy::Assert:
+					{
+						Entry.Table = Compatibility.FindBestMatchingTable(Entry.Address->GetClass());
+						checkf(Entry.Table != InvalidTableHandle, 
+							TEXT("The data storage could not find any matching tables for object of type '%s'. "
+							"This can mean that the object doesn't derive from UObject or that a table for UObject is no longer registered."), 
+							*Entry.Address->GetClass()->GetFName().ToString());
+						break;
+					}
+					case Private::EAsyncUObjectLoadStrategy::DeferRegistration:
+					{
+						if (Entry.Address->HasAnyInternalFlags(EInternalObjectFlags_AsyncLoading))
+						{
+							// Do nothing, we will check the entry again next frame
+							Entry.Table = InvalidTableHandle;
+						}
+						else
+						{
+							Entry.Table = Compatibility.FindBestMatchingTable(Entry.Address->GetClass());
+							checkf(Entry.Table != InvalidTableHandle, 
+								TEXT("The data storage could not find any matching tables for object of type '%s'. "
+								"This can mean that the object doesn't derive from UObject or that a table for UObject is no longer registered."), 
+								*Entry.Address->GetClass()->GetFName().ToString());
+						}
+						break;
+					}
+					default:
+					{
+						checkf(false, TEXT("Unhandled async uobject load strategy"));
+					}
+				}
 			}
 			else if constexpr (std::is_same_v<AddressType, ExternalObjectRegistration>)
 			{
@@ -791,12 +866,27 @@ void UEditorDataStorageCompatibility::PendingRegistration<AddressType>::ProcessE
 		FEntry* End = Current + Entries.Num();
 		
 		FEntry* TableFront = Current;
+
+		// It is expected that InvalidTableHandle should always be sorted at the end
+		// This assertion ensures that we catch a change in this behaviour should the constant
+		// be modified. 
+		static_assert(InvalidTableHandle == TNumericLimits<TableHandle>::Max()); 
+		if (Entries[0].Table == InvalidTableHandle)
+		{
+			return;
+		}
 		TableHandle CurrentTable = Entries[0].Table;
+
+		Compatibility.RowScratchBuffer.Reset();
 		
 		for (; Current != End; ++Current)
 		{
 			if (Current->Table != CurrentTable)
 			{
+				if (CurrentTable == InvalidTableHandle)
+				{
+					break;
+				}
 				StorageInterface.BatchAddRow(CurrentTable, Compatibility.RowScratchBuffer,
 					[&SetupRowCallback, &TableFront](RowHandle Row)
 					{
@@ -809,23 +899,36 @@ void UEditorDataStorageCompatibility::PendingRegistration<AddressType>::ProcessE
 			}
 			Compatibility.RowScratchBuffer.Add(Current->Row);
 		}
-		StorageInterface.BatchAddRow(CurrentTable, Compatibility.RowScratchBuffer,
-			[&SetupRowCallback, &TableFront](RowHandle Row)
+		if (CurrentTable != InvalidTableHandle)
+		{
+			StorageInterface.BatchAddRow(CurrentTable, Compatibility.RowScratchBuffer,
+				[&SetupRowCallback, &TableFront](RowHandle Row)
+				{
+					SetupRowCallback(Row, TableFront->Address);
+					++TableFront;
+				});
+			// The last entry was able to be added to TEDS
+			// therefore none remain in the list.
+			Entries.Reset();
+		}
+		else
+		{
+			// There are some objects that couldn't be added this frame
+			// Return how many remain in the list, we will try again next frame
+			const int32 RemainingEntries = static_cast<int32>(End - Current);
+			// Shuffle all the remaining entries to the front of the list and then
+			// truncate the list
+			const int32 ProcessedEntries = Entries.Num() - RemainingEntries;
+			for (int32 Index = 0; Index < RemainingEntries; ++Index)
 			{
-				SetupRowCallback(Row, TableFront->Address);
-				++TableFront;
-			});
+				Entries[Index] = Entries[ProcessedEntries + Index];
+			}
+			Entries.SetNum(RemainingEntries, EAllowShrinking::No);
+		}
 		Compatibility.RowScratchBuffer.Reset();
 	}
 }
 
-template<typename AddressType>
-void UEditorDataStorageCompatibility::PendingRegistration<AddressType>::Reset()
-{
-	// Thread-safe as it's only called from functions that already lock using.
-	Entries.Reset();
-}
-
 void UEditorDataStorageCompatibility::TickPendingCommands()
 {
 	using namespace UE::Editor::DataStorage;
@@ -897,8 +1000,6 @@ void UEditorDataStorageCompatibility::TickPendingUObjectRegistration()
 				Interface->AddColumn<FTypedElementSyncFromWorldTag>(Row);
 				TriggerOnObjectAdded(Object.Get(), Object->GetClass(), Row);
 			});
-
-		UObjectsPendingRegistration.Reset();
 	}
 }
 
@@ -918,8 +1019,6 @@ void UEditorDataStorageCompatibility::TickPendingExternalObjectRegistration()
 				Interface->AddColumn<FTypedElementSyncFromWorldTag>(Row);
 				TriggerOnObjectAdded(Object.Object, Object.TypeInfo.Get(), Row);
 			});
-
-		ExternalObjectsPendingRegistration.Reset();
 	}
 }
 
diff --git a/Engine/Plugins/Experimental/EditorDataStorage/Source/TedsCore/Private/TypedElementDatabaseCompatibility.h b/Engine/Plugins/Experimental/EditorDataStorage/Source/TedsCore/Private/TypedElementDatabaseCompatibility.h
--- a/Engine/Plugins/Experimental/EditorDataStorage/Source/TedsCore/Private/TypedElementDatabaseCompatibility.h
+++ b/Engine/Plugins/Experimental/EditorDataStorage/Source/TedsCore/Private/TypedElementDatabaseCompatibility.h
@@ -182,9 +182,10 @@ private:
 		int32 Num() const;
 		
 		void ForEachAddress(const TFunctionRef<void(AddressType&)>& Callback);
+		// Processes the entries and return the entries remaining in the list. They will be at the back of the list
+		// Once this returns, the list will contain only entries that were not processed
 		void ProcessEntries(UE::Editor::DataStorage::ICoreProvider& Storage, UEditorDataStorageCompatibility& Compatibility,
 			const TFunctionRef<void(UE::Editor::DataStorage::RowHandle, const AddressType&)>& SetupRowCallback);
-		void Reset();
 	};
 
 	UE::Editor::DataStorage::CompatibilityCommandBuffer QueuedCommands;
diff --git a/Engine/Plugins/Experimental/EditorDataStorageFeatures/Source/TedsActorCompatibility/Private/TedsActorCompatibilityModule.cpp b/Engine/Plugins/Experimental/EditorDataStorageFeatures/Source/TedsActorCompatibility/Private/TedsActorCompatibilityModule.cpp
--- a/Engine/Plugins/Experimental/EditorDataStorageFeatures/Source/TedsActorCompatibility/Private/TedsActorCompatibilityModule.cpp
+++ b/Engine/Plugins/Experimental/EditorDataStorageFeatures/Source/TedsActorCompatibility/Private/TedsActorCompatibilityModule.cpp
@@ -76,7 +76,6 @@ namespace UE::Editor::DataStorage
 
 				// 1. For actors that need syncing, add any components they have to the compatibility layer
 				// This solves a lot of issues around trying to intercept all the places components are added and removed
-				// In a perfect world, this may not be required
 				struct FRegisterActorComponentsWithCompatibility
 				{
 					TWeakObjectPtr<AActor> WkActor;
@@ -149,12 +148,9 @@ namespace UE::Editor::DataStorage
 						{
 							if (UActorComponent* ActorComponent = Cast<UActorComponent>(ObjectColumn.Object.Get()))
 							{
-								// Confusingly, an ActorComponent could indicate that it has an owner while its
-								// owner does not consider it owned.
-								// Defer to the owner's source of truth for this
 								const AActor* OwnerActor = ActorComponent->GetOwner();
 								RowHandle ActorRow;
-								if (OwnerActor && OwnerActor->GetComponents().Contains(ActorComponent))
+								if (OwnerActor)
 								{
 									ActorRow = Context.LookupMappedRow(ICompatibilityProvider::ObjectMappingDomain, FMapKeyView(static_cast<const UObject*>(OwnerActor)));
 								}
diff --git a/Engine/Source/Runtime/Engine/Private/Actor.cpp b/Engine/Source/Runtime/Engine/Private/Actor.cpp
--- a/Engine/Source/Runtime/Engine/Private/Actor.cpp
+++ b/Engine/Source/Runtime/Engine/Private/Actor.cpp
@@ -534,19 +534,6 @@ void AActor::ResetOwnedComponents()
 			}
 		}
 	}, true, RF_NoFlags, EInternalObjectFlags::Garbage);
-
-#if WITH_EDITOR
-	{
-		using namespace UE::Editor::DataStorage;
-		if (ICompatibilityProvider* Storage = UE::Editor::DataStorage::GetMutableDataStorageFeature<ICompatibilityProvider>(CompatibilityFeatureName))
-		{
-			for (UActorComponent* Component : OwnedComponents)
-			{
-				Storage->AddCompatibleObject(const_cast<UActorComponent*>(Component));
-			}
-		}
-	}
-#endif
 }
 
 void AActor::PostInitProperties()
@@ -3812,7 +3799,6 @@ void AActor::AddOwnedComponent(UActorComponent* Component)
 
 			AddComponentForReplication(Component);
 		}
-		
 
 		if (Component->IsCreatedByConstructionScript())
 		{
@@ -3847,19 +3833,6 @@ void AActor::RemoveOwnedComponent(UActorComponent* Component)
 		}
 
 		
-#if WITH_EDITOR
-		{
-			using namespace UE::Editor::DataStorage;
-			if (!HasAnyFlags(RF_BeginDestroyed)) // No need to sync if this component is being destroyed
-			{
-				if (ICoreProvider* Storage = UE::Editor::DataStorage::GetMutableDataStorageFeature<ICoreProvider>(StorageFeatureName))
-				{
-					RowHandle ComponentRow = Storage->LookupMappedRow(ICompatibilityProvider::ObjectMappingDomain, FMapKeyView(Component));
-					Storage->AddColumns<FTypedElementSyncFromWorldTag>(ComponentRow);
-				}
-			}
-		}
-#endif
 	}
 }
 
diff --git a/Engine/Source/Runtime/Engine/Private/Components/ActorComponent.cpp b/Engine/Source/Runtime/Engine/Private/Components/ActorComponent.cpp
--- a/Engine/Source/Runtime/Engine/Private/Components/ActorComponent.cpp
+++ b/Engine/Source/Runtime/Engine/Private/Components/ActorComponent.cpp
@@ -595,11 +595,18 @@ void UActorComponent::PostInitProperties()
 	Super::PostInitProperties();
 
 #if WITH_EDITOR
-	UEngineElementsLibrary::CreateEditorComponentElement(this);
-	using namespace UE::Editor::DataStorage;
-	if (ICompatibilityProvider* Storage = UE::Editor::DataStorage::GetMutableDataStorageFeature<ICompatibilityProvider>(CompatibilityFeatureName))
+	if (IsInGameThread())
+	{
+		UEngineElementsLibrary::CreateEditorComponentElement(this);
+		using namespace UE::Editor::DataStorage;
+		if (ICompatibilityProvider* Storage = UE::Editor::DataStorage::GetMutableDataStorageFeature<ICompatibilityProvider>(CompatibilityFeatureName))
+		{
+			Storage->AddCompatibleObject(this);
+		}
+	}
+	else
 	{
-		Storage->AddCompatibleObject(this);
+		check(IsInAsyncLoadingThread());
 	}
 #endif	// WITH_EDITOR
 
