Stopped TFunctionRefBase::CheckCallable being instantiated differently for different templates.
Deprecated member CheckCallable which was never intended to be public.

diff --git a/Engine/Plugins/Developer/Concert/ConcertUI/ConcertSharedSlate/Source/ConcertSharedSlate/Private/Session/Browser/SNewSessionRow.cpp b/Engine/Plugins/Developer/Concert/ConcertUI/ConcertSharedSlate/Source/ConcertSharedSlate/Private/Session/Browser/SNewSessionRow.cpp
--- a/Engine/Plugins/Developer/Concert/ConcertUI/ConcertSharedSlate/Source/ConcertSharedSlate/Private/Session/Browser/SNewSessionRow.cpp
+++ b/Engine/Plugins/Developer/Concert/ConcertUI/ConcertSharedSlate/Source/ConcertSharedSlate/Private/Session/Browser/SNewSessionRow.cpp
@@ -23,9 +23,9 @@ void SNewSessionRow::Construct(const FArguments& InArgs, TSharedPtr<FConcertSess
 	HighlightText = InArgs._HighlightText;
 	DefaultServerURL = InArgs._DefaultServerURL;
 
-	GetServersFunc.CheckCallable();
-	AcceptFunc.CheckCallable();
-	DeclineFunc.CheckCallable();
+	check(GetServersFunc);
+	check(AcceptFunc);
+	check(DeclineFunc);
 
 	// Construct base class
 	SMultiColumnTableRow<TSharedPtr<FConcertSessionTreeItem>>::Construct(FSuperRowType::FArguments().Padding(InArgs._Padding), InOwnerTableView);
diff --git a/Engine/Plugins/Developer/Concert/ConcertUI/ConcertSharedSlate/Source/ConcertSharedSlate/Private/Session/Browser/SSaveRestoreSessionRow.cpp b/Engine/Plugins/Developer/Concert/ConcertUI/ConcertSharedSlate/Source/ConcertSharedSlate/Private/Session/Browser/SSaveRestoreSessionRow.cpp
--- a/Engine/Plugins/Developer/Concert/ConcertUI/ConcertSharedSlate/Source/ConcertSharedSlate/Private/Session/Browser/SSaveRestoreSessionRow.cpp
+++ b/Engine/Plugins/Developer/Concert/ConcertUI/ConcertSharedSlate/Source/ConcertSharedSlate/Private/Session/Browser/SSaveRestoreSessionRow.cpp
@@ -22,8 +22,8 @@ void SSaveRestoreSessionRow::Construct(const FArguments& InArgs, TSharedPtr<FCon
 	DeclineFunc = InArgs._OnDeclineFunc;
 	HighlightText = InArgs._HighlightText;
 
-	AcceptFunc.CheckCallable();
-	DeclineFunc.CheckCallable();
+	check(AcceptFunc);
+	check(DeclineFunc);
 
 	// Construct base class
 	SMultiColumnTableRow<TSharedPtr<FConcertSessionTreeItem>>::Construct(FSuperRowType::FArguments().Padding(InArgs._Padding), InOwnerTableView);
diff --git a/Engine/Plugins/Developer/Concert/ConcertUI/ConcertSharedSlate/Source/ConcertSharedSlate/Private/Session/Browser/SSessionRow.cpp b/Engine/Plugins/Developer/Concert/ConcertUI/ConcertSharedSlate/Source/ConcertSharedSlate/Private/Session/Browser/SSessionRow.cpp
--- a/Engine/Plugins/Developer/Concert/ConcertUI/ConcertSharedSlate/Source/ConcertSharedSlate/Private/Session/Browser/SSessionRow.cpp
+++ b/Engine/Plugins/Developer/Concert/ConcertUI/ConcertSharedSlate/Source/ConcertSharedSlate/Private/Session/Browser/SSessionRow.cpp
@@ -20,9 +20,9 @@ void SSessionRow::Construct(const FArguments& InArgs, TSharedPtr<FConcertSession
 	HighlightText = InArgs._HighlightText;
 	IsSelected = InArgs._IsSelected;
 
-	DoubleClickFunc.CheckCallable();
-	RenameFunc.CheckCallable();
-	IsDefaultSession.CheckCallable();
+	check(DoubleClickFunc);
+	check(RenameFunc);
+	check(IsDefaultSession);
 
 	// Construct base class
 	SMultiColumnTableRow<TSharedPtr<FConcertSessionTreeItem>>::Construct(FSuperRowType::FArguments().Padding(InArgs._Padding), InOwnerTableView);
diff --git a/Engine/Plugins/Editor/WaveformEditor/Source/WaveformEditorWidgets/Private/WaveformTransformationRenderLayerFactory.cpp b/Engine/Plugins/Editor/WaveformEditor/Source/WaveformEditorWidgets/Private/WaveformTransformationRenderLayerFactory.cpp
--- a/Engine/Plugins/Editor/WaveformEditor/Source/WaveformEditorWidgets/Private/WaveformTransformationRenderLayerFactory.cpp
+++ b/Engine/Plugins/Editor/WaveformEditor/Source/WaveformEditorWidgets/Private/WaveformTransformationRenderLayerFactory.cpp
@@ -15,7 +15,6 @@ TSharedPtr<IWaveformTransformationRenderer> FWaveformTransformationRenderLayerFa
 
 	if (Instantiator)
 	{
-		Instantiator->CheckCallable();
 		OutRenderer = (*Instantiator)();
 
 		if (OutRenderer)
diff --git a/Engine/Plugins/Experimental/EditorDataStorage/Source/TedsCore/Private/TypedElementDatabase.cpp b/Engine/Plugins/Experimental/EditorDataStorage/Source/TedsCore/Private/TypedElementDatabase.cpp
--- a/Engine/Plugins/Experimental/EditorDataStorage/Source/TedsCore/Private/TypedElementDatabase.cpp
+++ b/Engine/Plugins/Experimental/EditorDataStorage/Source/TedsCore/Private/TypedElementDatabase.cpp
@@ -379,7 +379,6 @@ UE::Editor::DataStorage::RowHandle UEditorDataStorage::AddRow(UE::Editor::DataSt
 {
 	using namespace UE::Editor::DataStorage;
 
-	OnCreated.CheckCallable();
 	if (ActiveEditorEntityManager)
 	{
 		const int32 TableHandleAsIndex = Private::ConvertTableHandleToIndex(Table);
@@ -421,7 +420,6 @@ bool UEditorDataStorage::AddRow(UE::Editor::DataStorage::RowHandle ReservedRow,
 {
 	using namespace UE::Editor::DataStorage;
 
-	OnCreated.CheckCallable();
 	if (ActiveEditorEntityManager)
 	{
 		const int32 TableHandleAsIndex = Private::ConvertTableHandleToIndex(Table);
@@ -441,7 +439,6 @@ bool UEditorDataStorage::BatchAddRow(
 {
 	using namespace UE::Editor::DataStorage;
 
-	OnCreated.CheckCallable();
 	if (ActiveEditorEntityManager)
 	{
 		const int32 TableHandleAsIndex = Private::ConvertTableHandleToIndex(Table);
@@ -468,7 +465,6 @@ bool UEditorDataStorage::BatchAddRow(UE::Editor::DataStorage::TableHandle Table,
 	using namespace UE::Editor;
 	using namespace UE::Editor::DataStorage;
 
-	OnCreated.CheckCallable();
 	if (ActiveEditorEntityManager)
 	{
 		const int32 TableHandleAsIndex = Private::ConvertTableHandleToIndex(Table);
diff --git a/Engine/Source/Programs/AutoRTFMTests/Private/FunctionPointerTests.cpp b/Engine/Source/Programs/AutoRTFMTests/Private/FunctionPointerTests.cpp
--- a/Engine/Source/Programs/AutoRTFMTests/Private/FunctionPointerTests.cpp
+++ b/Engine/Source/Programs/AutoRTFMTests/Private/FunctionPointerTests.cpp
@@ -86,7 +86,7 @@ TEST_CASE("FunctionPointer.TFunction")
 					MyFunc();
 				}
 
-				MyFunc.CheckCallable();
+				check(MyFunc);
 
 				MyFunc.Reset();
 			});
@@ -109,7 +109,7 @@ TEST_CASE("FunctionPointer.TFunction")
 					MyFunc();
 				}
 
-				MyFunc.CheckCallable();
+				check(MyFunc);
 
 				MyFunc.Reset();
 			});
@@ -135,7 +135,7 @@ TEST_CASE("FunctionPointer.TUniqueFunction")
 					MyFunc();
 				}
 
-				MyFunc.CheckCallable();
+				check(MyFunc);
 
 				MyFunc.Reset();
 			});
@@ -158,7 +158,7 @@ TEST_CASE("FunctionPointer.TUniqueFunction")
 					MyFunc();
 				}
 
-				MyFunc.CheckCallable();
+				check(MyFunc);
 
 				MyFunc.Reset();
 			});
diff --git a/Engine/Source/Runtime/Core/Private/Templates/Function.cpp b/Engine/Source/Runtime/Core/Private/Templates/Function.cpp
new file mode 100644
--- /dev/null
+++ b/Engine/Source/Runtime/Core/Private/Templates/Function.cpp
@@ -0,0 +1,11 @@
+// Copyright Epic Games, Inc. All Rights Reserved.
+
+#include "Templates/Function.h"
+
+namespace UE::Core::Private::Function
+{
+	FORCENOINLINE void CheckCallable(void* Callable)
+	{
+		checkf(Callable, TEXT("Attempting to call an unbound TFunction!"));
+	}
+}
diff --git a/Engine/Source/Runtime/Core/Public/Templates/Function.h b/Engine/Source/Runtime/Core/Public/Templates/Function.h
--- a/Engine/Source/Runtime/Core/Public/Templates/Function.h
+++ b/Engine/Source/Runtime/Core/Public/Templates/Function.h
@@ -298,6 +298,13 @@ namespace UE::Core::Private::Function
 		}
 	};
 
+#if DO_CHECK
+	// Move all of the assert code out of line.
+	// This takes a void* to avoid multiple template instantiations - the caller should
+	// explicitly cast the Callable function pointer to void* to call this function.
+	CORE_API void CheckCallable(void* Callable);
+#endif
+
 	/**
 	 * A class which defines an operator() which will own storage of a callable and invoke the TFunctionRefCaller::Call function on it.
 	 */
@@ -390,7 +397,7 @@ namespace UE::Core::Private::Function
 		TFunctionRefBase& operator=(TFunctionRefBase&&) = delete;
 		TFunctionRefBase& operator=(const TFunctionRefBase&) = delete;
 
-		// Move all of the assert code out of line
+		UE_DEPRECATED(5.7, "CheckCallable should not be called - use check(Func) instead, and for a TFunctionRef there's no need to check.")
 		FORCENOINLINE void CheckCallable() const
 		{
 			checkf(Callable, TEXT("Attempting to call an unbound TFunction!"));
@@ -401,7 +408,7 @@ namespace UE::Core::Private::Function
 #if DO_CHECK
 			if constexpr (StorageType::bCanBeNull)
 			{
-				CheckCallable();
+				UE::Core::Private::Function::CheckCallable((void*)Callable);
 			}
 #endif
 			return Callable(Storage.GetPtr(), Params...);
