[TEDS] Sped up legacy command buffer.

Removed the stable sort used to order commands to the same rows together. With a large number of commands, the sorting was more expensive that any gains from more aligned memory access.
The test case was a file for CVD that created 185k commands in one frame:
- With StableSort processing the commands took 10 seconds.
- With RadixSort64 processing the commands to 92ms.
- With RadixSort64 and swapping the index and generation around processing the commands took 79ms.
- Without any sorting prepass processing the commands took 39ms.
Since this was only ever tested on smaller sets, where it makes a marginal difference, the preference is to optimize for large batches. The new command buffer will need to be updated to reflect this as well.

[RN]

diff --git a/Engine/Plugins/Experimental/EditorDataStorage/Source/TedsCore/Private/TypedElementDatabaseCommandBuffer.cpp b/Engine/Plugins/Experimental/EditorDataStorage/Source/TedsCore/Private/TypedElementDatabaseCommandBuffer.cpp
--- a/Engine/Plugins/Experimental/EditorDataStorage/Source/TedsCore/Private/TypedElementDatabaseCommandBuffer.cpp
+++ b/Engine/Plugins/Experimental/EditorDataStorage/Source/TedsCore/Private/TypedElementDatabaseCommandBuffer.cpp
@@ -320,12 +320,6 @@ namespace UE::Editor::DataStorage::Legacy
 
 	void FCommandBuffer::ProcessCommands()
 	{
-		Commands.StableSort(
-			[](const FCommand& Lhs, const FCommand& Rhs)
-			{
-				return Lhs.Row < Rhs.Row;
-			});
-
 		struct FProcessor
 		{
 			FMassEntityManager& EntityManager;
@@ -337,10 +331,11 @@ namespace UE::Editor::DataStorage::Legacy
 			void operator()(const FRemoveColumnsCommand& Command) { Execute_RemoveColumnsCommand(EntityManager, Row, Command.FragmentsToRemove, Command.TagsToRemove); }
 		};
 
-		FProcessor Processor{ .EntityManager = Environment.GetMassEntityManager() };
+		FMassEntityManager& EntityManager = Environment.GetMassEntityManager();
+		FProcessor Processor{ .EntityManager = EntityManager };
 		for (FCommand& Command : Commands)
 		{
-			if (Execute_IsRowAssigned(Environment.GetMassEntityManager(), Command.Row))
+			if (Execute_IsRowAssigned(EntityManager, Command.Row))
 			{
 				Processor.Row = Command.Row;
 				Visit(Processor, Command.Data);
