TEDS: Add dynamic column template metadata to the value tag template struct

diff --git a/Engine/Plugins/Experimental/EditorDataStorage/Source/TedsCore/Private/DynamicColumnGenerator.h b/Engine/Plugins/Experimental/EditorDataStorage/Source/TedsCore/Private/DynamicColumnGenerator.h
--- a/Engine/Plugins/Experimental/EditorDataStorage/Source/TedsCore/Private/DynamicColumnGenerator.h
+++ b/Engine/Plugins/Experimental/EditorDataStorage/Source/TedsCore/Private/DynamicColumnGenerator.h
@@ -16,7 +16,7 @@ class FName;
 class UScriptStruct;
 
 // The template struct that is used to generate the ValueTag column
-USTRUCT()
+USTRUCT(meta=(EditorDataStorage_DynamicColumnTemplate))
 struct FTedsValueTagColumn : public FTedsSharedColumn
 {
 	GENERATED_BODY()
