TEDS Hierarchies: Add support for an optional hierarchy registration param that adds a column to rows when their parent changes. This column behaves similar to the sync tags and is removed from the row at the end of the frame, so it can be used for change detection.

diff --git a/Engine/Plugins/Experimental/EditorDataStorage/Source/TedsCore/Private/Hierarchy/EditorDataStorageHierarchyImplementation.cpp b/Engine/Plugins/Experimental/EditorDataStorage/Source/TedsCore/Private/Hierarchy/EditorDataStorageHierarchyImplementation.cpp
--- a/Engine/Plugins/Experimental/EditorDataStorage/Source/TedsCore/Private/Hierarchy/EditorDataStorageHierarchyImplementation.cpp
+++ b/Engine/Plugins/Experimental/EditorDataStorage/Source/TedsCore/Private/Hierarchy/EditorDataStorageHierarchyImplementation.cpp
@@ -97,7 +97,8 @@ namespace UE::Editor::DataStorage
 				.ChildTag = ChildTag,
 				.ParentTag = ParentTag,
 				.HierarchyData = HierarchyData,
-				.UnresolvedParentColumn = UnresolvedParentColumn
+				.UnresolvedParentColumn = UnresolvedParentColumn,
+				.ParentChangedColumn = RegistrationParams.ParentChangedColumn
 			};
 			return MakeUnique<FTedsHierarchyAccessInterface>(Params);
 		}();
@@ -109,6 +110,7 @@ namespace UE::Editor::DataStorage
 			.ParentTag = ParentTag,
 			.HierarchyData = HierarchyData,
 			.UnresolvedParentColumn = UnresolvedParentColumn,
+			.ParentChangedColumn = RegistrationParams.ParentChangedColumn,
 			.AccessInterface = MoveTemp(AccessInterface)
 		};
 
@@ -260,6 +262,25 @@ namespace UE::Editor::DataStorage
 		.ReadWrite(Hierarchy->HierarchyData)
 		.DependsOn().SubQuery(Subquery)
 		.Compile());
+
+		if (Hierarchy->ParentChangedColumn)
+		{
+			DataStorage.RegisterQuery(
+				Select(
+					TEXT("Remove Parent Changed Tag at the end of the frame"),
+					FPhaseAmble(FPhaseAmble::ELocation::Postamble, EQueryTickPhase::FrameEnd),
+					[this, Handle](IQueryContext& Context, const RowHandle* Rows)
+					{
+						if (const FRegisteredHierarchy* Hierarchy = FindRegisteredHierarchy(Handle))
+						{
+							Context.RemoveColumns(TConstArrayView<RowHandle>(Rows, Context.GetRowCount()), {Hierarchy->ParentChangedColumn});
+						}
+					}
+				)
+				.Where()
+					.All(Hierarchy->ParentChangedColumn)
+				.Compile());
+		}
 	}
 
 	void FTedsHierarchyRegistrar::RegisterProcessors(ICoreProvider& CoreProvider, const FHierarchyHandle& Handle)
@@ -313,6 +334,7 @@ namespace UE::Editor::DataStorage
 		, ParentTag(Params.ParentTag)
 		, HierarchyDataColumnType(Params.HierarchyData)
 		, UnresolvedParentColumnType(Params.UnresolvedParentColumn)
+		, ParentChangedColumn(Params.ParentChangedColumn)
 	{
 		check(HierarchyDataColumnType->IsChildOf(FEditorDataHierarchyData_Template::StaticStruct()));
 	}
@@ -408,7 +430,8 @@ namespace UE::Editor::DataStorage
 		{
 			HierarchyDataColumn_Parent = static_cast<FEditorDataHierarchyData_Template*>(ParentHierarchyDataColumnAddress);
 		}
-		
+
+		RowHandle PreviousParent = InvalidRowHandle;
 		FEditorDataHierarchyData_Template* HierarchyDataColumn_Target;
 		if (void* TargetHierarchyDataColumnAddress = CoreProvider->GetColumnData(Target, HierarchyDataColumnType); TargetHierarchyDataColumnAddress == nullptr)
 		{
@@ -420,6 +443,7 @@ namespace UE::Editor::DataStorage
 		else
 		{
 			HierarchyDataColumn_Target = static_cast<FEditorDataHierarchyData_Template*>(TargetHierarchyDataColumnAddress);
+			PreviousParent = HierarchyDataColumn_Target->Parent;
 		}
 		
 		if (Parent != InvalidRowHandle)
@@ -443,6 +467,11 @@ namespace UE::Editor::DataStorage
 			// 	CoreProvider->RemoveColumns(Parent, {ParentTag});
 			// }
 		}
+
+		if (ParentChangedColumn && Parent != PreviousParent)
+		{
+			CoreProvider->AddColumn(Target, ParentChangedColumn);
+		}
 	}
 
 	void FTedsHierarchyAccessInterface::SetUnresolvedParent(ICoreProvider* CoreProvider, RowHandle Target, FMapKey ParentId, FName MappingDomain) const
diff --git a/Engine/Plugins/Experimental/EditorDataStorage/Source/TedsCore/Private/Hierarchy/EditorDataStorageHierarchyImplementation.h b/Engine/Plugins/Experimental/EditorDataStorage/Source/TedsCore/Private/Hierarchy/EditorDataStorageHierarchyImplementation.h
--- a/Engine/Plugins/Experimental/EditorDataStorage/Source/TedsCore/Private/Hierarchy/EditorDataStorageHierarchyImplementation.h
+++ b/Engine/Plugins/Experimental/EditorDataStorage/Source/TedsCore/Private/Hierarchy/EditorDataStorageHierarchyImplementation.h
@@ -36,6 +36,7 @@ namespace UE::Editor::DataStorage
 			const UScriptStruct* ParentTag;
 			const UScriptStruct* HierarchyData;
 			const UScriptStruct* UnresolvedParentColumn;
+			const UScriptStruct* ParentChangedColumn;
 			TUniquePtr<FTedsHierarchyAccessInterface> AccessInterface;
 		};
 
@@ -54,6 +55,7 @@ namespace UE::Editor::DataStorage
 			const UScriptStruct* ParentTag;
 			const UScriptStruct* HierarchyData;
 			const UScriptStruct* UnresolvedParentColumn;
+			const UScriptStruct* ParentChangedColumn;
 		};
 		
 		explicit FTedsHierarchyAccessInterface(const FConstructParams& Params);
@@ -95,5 +97,6 @@ namespace UE::Editor::DataStorage
 		const UScriptStruct* ParentTag;
 		const UScriptStruct* HierarchyDataColumnType;
 		const UScriptStruct* UnresolvedParentColumnType;
+		const UScriptStruct* ParentChangedColumn;
 	};
 }
diff --git a/Engine/Source/Runtime/TypedElementFramework/Public/Elements/Interfaces/TypedElementDataStorageInterface.h b/Engine/Source/Runtime/TypedElementFramework/Public/Elements/Interfaces/TypedElementDataStorageInterface.h
--- a/Engine/Source/Runtime/TypedElementFramework/Public/Elements/Interfaces/TypedElementDataStorageInterface.h
+++ b/Engine/Source/Runtime/TypedElementFramework/Public/Elements/Interfaces/TypedElementDataStorageInterface.h
@@ -54,6 +54,10 @@ struct TTypedElementColumnTypeList
 struct FHierarchyRegistrationParams
 {
 	FName Name;
+
+	// Optional column that is added to a row when the parent changes (useful for change detection).
+	// This behaves similar to the sync tags and is removed at the end of the frame
+	const UScriptStruct* ParentChangedColumn = nullptr;
 };
 	
 class ICoreProvider : public IModularFeature
