TEDS: Fix incorrect IsEntityValid() check instead of the intended IsEntityBuilt() check when deferring commands

diff --git a/Engine/Plugins/Experimental/EditorDataStorage/Source/TedsCore/Private/Processors/TypedElementProcessorAdaptors.cpp b/Engine/Plugins/Experimental/EditorDataStorage/Source/TedsCore/Private/Processors/TypedElementProcessorAdaptors.cpp
--- a/Engine/Plugins/Experimental/EditorDataStorage/Source/TedsCore/Private/Processors/TypedElementProcessorAdaptors.cpp
+++ b/Engine/Plugins/Experimental/EditorDataStorage/Source/TedsCore/Private/Processors/TypedElementProcessorAdaptors.cpp
@@ -826,7 +826,7 @@ namespace UE::Editor::DataStorage
 						this->Context.Defer().template PushCommand<FMassDeferredAddCommand>(
 							[RemovedColumns](FMassEntityManager& System)
 							{
-								if (System.IsEntityValid(RemovedColumns->Entity))
+								if (System.IsEntityBuilt(RemovedColumns->Entity))
 								{
 									System.RemoveCompositionFromEntity(RemovedColumns->Entity, RemovedColumns->RemoveDescriptor);
 								}
