TEDS: Check if a row is valid before checking if it is built, this is because internally IsEntityBuild asserts if the entity is not valid

diff --git a/Engine/Plugins/Experimental/EditorDataStorage/Source/TedsCore/Private/Processors/TypedElementProcessorAdaptors.cpp b/Engine/Plugins/Experimental/EditorDataStorage/Source/TedsCore/Private/Processors/TypedElementProcessorAdaptors.cpp
--- a/Engine/Plugins/Experimental/EditorDataStorage/Source/TedsCore/Private/Processors/TypedElementProcessorAdaptors.cpp
+++ b/Engine/Plugins/Experimental/EditorDataStorage/Source/TedsCore/Private/Processors/TypedElementProcessorAdaptors.cpp
@@ -826,7 +826,7 @@ namespace UE::Editor::DataStorage
 						this->Context.Defer().template PushCommand<FMassDeferredAddCommand>(
 							[RemovedColumns](FMassEntityManager& System)
 							{
-								if (System.IsEntityValid(RemovedColumns->Entity))
+								if (System.IsEntityValid(RemovedColumns->Entity) && System.IsEntityBuilt(RemovedColumns->Entity))
 								{
 									System.RemoveCompositionFromEntity(RemovedColumns->Entity, RemovedColumns->RemoveDescriptor);
 								}
