TEDS Hierarchies: When SetParent() causes a hierarchy change on a row (i.e valid  old parent -> valid new parent), immediately update the information on the old parent. This is because the observer is run later and uses the new parent data, so the old parent is never updated.

Also only do work if the new parent is actually different than the old parent, no-op otherwise

diff --git a/Engine/Plugins/Experimental/EditorDataStorage/Source/TedsCore/Private/Hierarchy/EditorDataStorageHierarchyImplementation.cpp b/Engine/Plugins/Experimental/EditorDataStorage/Source/TedsCore/Private/Hierarchy/EditorDataStorageHierarchyImplementation.cpp
--- a/Engine/Plugins/Experimental/EditorDataStorage/Source/TedsCore/Private/Hierarchy/EditorDataStorageHierarchyImplementation.cpp
+++ b/Engine/Plugins/Experimental/EditorDataStorage/Source/TedsCore/Private/Hierarchy/EditorDataStorageHierarchyImplementation.cpp
@@ -455,16 +455,40 @@ namespace UE::Editor::DataStorage
 		
 		if (Parent != InvalidRowHandle)
 		{
-			HierarchyDataColumn_Target->Parent = Parent;
-			HierarchyDataColumn_Parent->Children.Add(FTedsRowHandle(Target));
-			CoreProvider->AddColumn(Target, ChildTag);
-			CoreProvider->AddColumn(Parent, ParentTag);
+			// No updates needed if the parent is the same as before
+			if (Parent != PreviousParent)
+			{
+				FEditorDataHierarchyData_Template* HierarchyDataColumn_PreviousParent = nullptr;
+				if (void* PreviousParentHierarchyDataColumnAddress = CoreProvider->GetColumnData(PreviousParent, HierarchyDataColumnType))
+				{
+					HierarchyDataColumn_PreviousParent = static_cast<FEditorDataHierarchyData_Template*>(PreviousParentHierarchyDataColumnAddress);
+				}
+				
+				// In the case the parent changed, we have to also update the previous parent right now since any observers will act on the new parent
+				if (HierarchyDataColumn_PreviousParent)
+				{
+					HierarchyDataColumn_PreviousParent->Children.RemoveSwap(FTedsRowHandle(Target));
+
+					if (HierarchyDataColumn_PreviousParent->Children.IsEmpty())
+					{
+						CoreProvider->RemoveColumn(PreviousParent, ParentTag);
+					}
+				}
+
+				// Update the target with the new information
+				HierarchyDataColumn_Target->Parent = Parent;
+				CoreProvider->AddColumn(Target, ChildTag);
+
+				// Update the new parent with the new information
+				HierarchyDataColumn_Parent->Children.Add(FTedsRowHandle(Target));
+				CoreProvider->AddColumn(Parent, ParentTag);
+			}
 		}
 		else
 		{
 			CoreProvider->RemoveColumns(Target, {ChildTag});
 
-			// All this work is handled by observer...
+			// When the child is removed from a parent (not including when the parent changes to another row), an observer handles all this
 			// Commented out just to make it clear what should happen
 			
 			// HierarchyDataColumn_Target->Parent = InvalidRowHandle;
