[TEDS] Added basic support for collaborative time slicing.

It's now possible to register a callback with TEDS that will be periodically called. TEDS will monitor for the amount of time left at the end of the frame and then pick tasks from the queue to run for the remaining time. Priorities can be set to high, medium and low, to allow some tasks to run more or less frequently on average. The first use of this is the index table.

diff --git a/Engine/Plugins/Experimental/EditorDataStorage/Source/TedsCore/Private/TypedElementDatabase.cpp b/Engine/Plugins/Experimental/EditorDataStorage/Source/TedsCore/Private/TypedElementDatabase.cpp
--- a/Engine/Plugins/Experimental/EditorDataStorage/Source/TedsCore/Private/TypedElementDatabase.cpp
+++ b/Engine/Plugins/Experimental/EditorDataStorage/Source/TedsCore/Private/TypedElementDatabase.cpp
@@ -3,12 +3,14 @@
 #include "TypedElementDatabase.h"
 
 #include "DataStorage/Features.h"
+#include "DataStorage/Debug/Log.h"
 #include "Editor.h"
 #include "EditorDataStorageSettings.h"
 #include "Elements/Interfaces/TypedElementDataStorageFactory.h"
 #include "Elements/Interfaces/TypedElementDataStorageInterface.h"
 #include "Engine/World.h"
 #include "GlobalLock.h"
+#include "Math/UnrealMathUtility.h"
 #include "MassEntityEditorSubsystem.h"
 #include "MassEntityTypes.h"
 #include "MassProcessor.h"
@@ -19,6 +21,7 @@
 #include "Stats/Stats.h"
 #include "TickTaskManagerInterface.h"
 #include "TypedElementDatabaseEnvironment.h"
+#include "TypedElementDataStorageProfilingMacros.h"
 #include "TypedElementUtils.h"
 #include "UObject/UObjectIterator.h"
 
@@ -73,6 +76,19 @@ FAutoConsoleCommandWithOutputDevice PrintSupportedColumnsConsoleCommand(
 			Output.Log(TEXT("End of Typed Elements Data Storage supported column list."));
 		}));
 
+static TAutoConsoleVariable<int32> CooperativeTasks_TargetFrameRate(TEXT("TEDS.Features.CooperativeTasks.TargetFrameRate"), 60,
+	TEXT("TEDS runs additional cooperative background tasks with the time that is left at the end of a frame. Adjust this value to set "
+		"how much TEDS considers as available time. Clamped between 16 and 120 fps."));
+
+static TAutoConsoleVariable<float> CooperativeTasks_MinTimePerFrameMs(TEXT("TEDS.Features.CooperativeTasks.MinTimePerFrameMs"), 2.0f,
+	TEXT("The minimal amount of time in milliseconds that TEDS will spend guaranteed on cooperative tasks."));
+
+static TAutoConsoleVariable<float> CooperativeTasks_MinTimePerTaskMs(TEXT("TEDS.Features.CooperativeTasks.MinTimePerTaskMs"), 1.0f,
+	TEXT("The minimal amount of time in milliseconds that needs to remain for TEDS to consider running another cooperative task."));
+
+static TAutoConsoleVariable<float> CooperativeTasks_ReportThresholdMs(TEXT("TEDS.Features.CooperativeTasks.ReportThresholdMs"), 1.0f,
+	TEXT("The maximum amount of time in milliseconds that cooperative tasks are allowed to go over their allotted time before they're reported."));
+
 namespace UE::Editor::DataStorage::Private
 {
 	struct ColumnsToBitSetsResult
@@ -270,6 +286,8 @@ void UEditorDataStorage::OnPostMassTick(float DeltaTime)
 	
 	Environment->NextUpdateCycle();
 	OnUpdateCompletedDelegate.Broadcast();
+
+	TickCooperativeTasks();
 }
 
 TSharedPtr<FMassEntityManager> UEditorDataStorage::GetActiveMutableEditorEntityManager()
@@ -1352,6 +1370,151 @@ UE::Editor::DataStorage::FTypedElementOnDataStorageUpdate& UEditorDataStorage::O
 	return OnUpdateCompletedDelegate;
 }
 
+void UEditorDataStorage::RegisterCooperativeUpdate(const FName& TaskName, ECooperativeTaskPriority Priority, FOnCooperativeUpdate Callback)
+{
+	checkf(!bRunningCooperativeUpdate, TEXT("Cooperative tasks can't be added to TEDS during the cooperative task update."));
+
+	uint8 OrderValue;
+	// This distribution means that if there's one high, one medium and one low task the high task would be called twice as often as the
+	// medium task and four times as often the low priority task. These numbers are skewed if there's a different number of task per
+	// priority, but no matter how many tasks there are, on average high will be more frequent than medium and medium more frequently
+	// called than low.
+	// As an example, if there's 2 high, 5 medium and 3 low tasks than a 1000 iterations will be distributed as:
+	//		2 High:   total 247 calls, avg. 123.5 calls
+	//		5 Medium: total 529 calls, avg. 105.8 calls
+	//		3 Low :	  total 224 calls, avg.  74.7 calls
+	// So while there are 529 medium tasks executed compared to only 247 high tasks, because there are 5 medium and only 2 high task, the
+	// high task occur with higher frequency than medium tasks.
+	switch (Priority)
+	{
+	case ECooperativeTaskPriority::High:
+		// Don't go below 1 because that practically means it's always set to 0. If Medium is set to 2 it means they'll effectively have 
+		// the same priority because they'll both constantly be ping-pong-ing around 0 and 1.
+		OrderValue = 2;
+		break;
+	case ECooperativeTaskPriority::Medium:
+		OrderValue = 4;
+		break;
+	case ECooperativeTaskPriority::Low:
+		[[fallthrough]];
+	default:
+		OrderValue = 8;
+		break;
+	}
+
+	// Just add to the back of the queue as the task will eventually bubble up and 
+	// land in its correct position.
+	CooperativeTasks.Add(FCooperativeTask
+		{
+			.Callback = MoveTemp(Callback),
+			.Name = TaskName,
+			.OrderResetValue = OrderValue,
+			.Order = OrderValue
+		});
+}
+
+void UEditorDataStorage::UnregisterCooperativeUpdate(const FName& TaskName)
+{
+	checkf(!bRunningCooperativeUpdate, TEXT("Cooperative tasks can't be removed from TEDS during the cooperative task update."));
+
+	if (uint32 Index = CooperativeTasks.IndexOfByPredicate([&TaskName](const FCooperativeTask& Task) { return Task.Name == TaskName; });
+		Index != INDEX_NONE)
+	{
+		// Maintain order as the order needed to make sure tasks are ordered in the correct order.
+		CooperativeTasks.RemoveAt(Index);
+	}
+}
+
+void UEditorDataStorage::TickCooperativeTasks()
+{
+	// The minimum amount of time of time that will be spend each tick on cooperative tasks. A cap is set on the minimum
+	// time spend processing tasks each frame to avoid situations where the editor goes over frame budget for prolonged
+	// periods, causing the tasks to not be run. This can be detrimental if any of the tasks are needed to speed up
+	// the editor like garbage collection tasks.
+	const FTimespan MinTimePerTick = FTimespan::FromMilliseconds(CooperativeTasks_MinTimePerFrameMs.GetValueOnGameThread());
+	// Minimum amount of time that needs to be left in order for a task to run. This is to prevent there being a few
+	// nanoseconds left that are used by a task that takes several milliseconds. Without continuous analysis this
+	// value is a guess
+	const FTimespan MinTimePerTask = FTimespan::FromMilliseconds(CooperativeTasks_MinTimePerTaskMs.GetValueOnGameThread());
+	// The expected duration per frame.
+	int64 TargetFrameRate = FMath::Clamp(CooperativeTasks_TargetFrameRate.GetValueOnGameThread(), 16, 120);
+	const FTimespan TargetFrameDuration = FTimespan::FromMilliseconds(1000.0 / static_cast<double>(TargetFrameRate));
+	// The amount of time a task has to go over the alloted time before it's reported.
+	const FTimespan OverBudgetReportThreshold = FTimespan::FromMilliseconds(CooperativeTasks_ReportThresholdMs.GetValueOnGameThread());
+
+	if (!CooperativeTasks.IsEmpty())
+	{
+		TEDS_EVENT_SCOPE("Tick time sliced tasks");
+
+		bRunningCooperativeUpdate = true;
+
+		// Update the priorities. This causes all tasks to eventually float to the top so even low priority tasks
+		// get to run on occasion.
+		for (FCooperativeTask& Task : CooperativeTasks)
+		{
+			Task.Order = Task.Order > 0 ? Task.Order - 1 : Task.Order;
+			Task.bHasRun = false;
+		}
+
+		FTimespan LastFrameDuration(FTimespan::FromSeconds(FPlatformTime::ToSeconds64(FPlatformTime::Cycles64() - CooperativeTickStartTime)));
+		FTimespan RemainingTime = LastFrameDuration < (TargetFrameDuration - MinTimePerTick)
+			? TargetFrameDuration - LastFrameDuration
+			: MinTimePerTick;
+		
+		int32 CurrentTaskIndex = 0;
+		while (RemainingTime > MinTimePerTask && CurrentTaskIndex < CooperativeTasks.Num())
+		{
+			FCooperativeTask& CurrentTask = CooperativeTasks[CurrentTaskIndex];
+			// Make sure that tasks aren't run multiple times per frame and eat up cycles from lower priority tasks
+			if (!CurrentTask.bHasRun)
+			{
+				uint64 TaskStartTime = FPlatformTime::Cycles64();
+
+				TRACE_CPUPROFILER_EVENT_SCOPE_TEXT(
+					*FString::Printf(TEXT("[TEDS] Tick time sliced tasks: '%s'"), *CurrentTask.Name.ToString()));
+
+				// Run the task.
+				CurrentTask.Callback(RemainingTime);
+				CurrentTask.Order = CurrentTask.OrderResetValue;
+				CurrentTask.bHasRun = true;
+
+				// Find the new slot to put the processed task in by bubbling it up the queue.
+				for (int32 Index = CurrentTaskIndex + 1; Index < CooperativeTasks.Num(); Index++)
+				{
+					if (CooperativeTasks[Index - 1].Order >= CooperativeTasks[Index].Order)
+					{
+						Swap(CooperativeTasks[Index - 1], CooperativeTasks[Index]);
+					}
+					else
+					{
+						break;
+					}
+				}
+				
+				FTimespan TaskDuration(FTimespan::FromSeconds(FPlatformTime::ToSeconds64(FPlatformTime::Cycles64() - TaskStartTime)));
+				if (TaskDuration > RemainingTime + OverBudgetReportThreshold)
+				{
+					// Report this as verbose and not a warning because it's expected that when the frame behaves erratic during loading or
+					// under heavy load this can trigger for even well behaving tasks. It's also possible that the OS takes control away from
+					// the editor for a while to give other programs a chance to run which can also cause spikes. This becomes useful if there's
+					// a pattern of frequent runs going over budget.
+					UE_LOG(LogEditorDataStorage, Verbose, TEXT("Time sliced task: '%s' took %.2fms, but was alloted %.2fms"),
+						*CurrentTask.Name.ToString(), TaskDuration.GetTotalMilliseconds(), RemainingTime.GetTotalMilliseconds());
+				}
+				RemainingTime -= TaskDuration;
+			}
+			else
+			{
+				CurrentTaskIndex++;
+			}
+		}
+
+		bRunningCooperativeUpdate = false;
+
+		CooperativeTickStartTime = FPlatformTime::Cycles64();
+	}
+}
+
 bool UEditorDataStorage::IsAvailable() const
 {
 	return bool(ActiveEditorEntityManager);
diff --git a/Engine/Plugins/Experimental/EditorDataStorage/Source/TedsCore/Private/TypedElementDatabase.h b/Engine/Plugins/Experimental/EditorDataStorage/Source/TedsCore/Private/TypedElementDatabase.h
--- a/Engine/Plugins/Experimental/EditorDataStorage/Source/TedsCore/Private/TypedElementDatabase.h
+++ b/Engine/Plugins/Experimental/EditorDataStorage/Source/TedsCore/Private/TypedElementDatabase.h
@@ -194,6 +194,8 @@ public:
 
 	UE_API virtual UE::Editor::DataStorage::FTypedElementOnDataStorageUpdate& OnUpdate() override;
 	UE_API virtual UE::Editor::DataStorage::FTypedElementOnDataStorageUpdate& OnUpdateCompleted() override;
+	UE_API virtual void RegisterCooperativeUpdate(const FName& TaskName, ECooperativeTaskPriority Priority, FOnCooperativeUpdate Callback) override;
+	UE_API virtual void UnregisterCooperativeUpdate(const FName& TaskName) override;
 	UE_API virtual bool IsAvailable() const override;
 	UE_API virtual void* GetExternalSystemAddress(UClass* Target) override;
 
@@ -210,15 +212,13 @@ public:
 	static UE_API void AddReferencedObjects(UObject* InThis, FReferenceCollector& Collector);
 
 private:
-	UE_API void PreparePhase(UE::Editor::DataStorage::EQueryTickPhase Phase, float DeltaTime);
-	UE_API void FinalizePhase(UE::Editor::DataStorage::EQueryTickPhase Phase, float DeltaTime);
-	UE_API void Reset();
+	void PreparePhase(UE::Editor::DataStorage::EQueryTickPhase Phase, float DeltaTime);
+	void FinalizePhase(UE::Editor::DataStorage::EQueryTickPhase Phase, float DeltaTime);
+	void Reset();
 
-	UE_API int32 GetTableChunkSize(FName TableName) const;
+	int32 GetTableChunkSize(FName TableName) const;
+	void TickCooperativeTasks();
 
-public:
-
-private:
 	struct FFactoryTypePair
 	{
 		// Used to find the factory by type without needing to dereference each one
@@ -248,6 +248,18 @@ private:
 	FDelegateHandle OnPreMassTickHandle;
 	FDelegateHandle OnPostMassTickHandle;
 
+	struct FCooperativeTask
+	{
+		FOnCooperativeUpdate Callback;
+		FName Name;
+		uint8 OrderResetValue;
+		uint8 Order;
+		bool bHasRun;
+	};
+	TArray<FCooperativeTask> CooperativeTasks;
+	uint64 CooperativeTickStartTime = 0;
+	bool bRunningCooperativeUpdate = false;
+
 	TSharedPtr<FMassEntityManager> ActiveEditorEntityManager;
 	TSharedPtr<FMassProcessingPhaseManager> ActiveEditorPhaseManager;
 };
diff --git a/Engine/Plugins/Experimental/EditorDataStorage/Source/TedsCore/Private/TypedElementDatabaseEnvironment.cpp b/Engine/Plugins/Experimental/EditorDataStorage/Source/TedsCore/Private/TypedElementDatabaseEnvironment.cpp
--- a/Engine/Plugins/Experimental/EditorDataStorage/Source/TedsCore/Private/TypedElementDatabaseEnvironment.cpp
+++ b/Engine/Plugins/Experimental/EditorDataStorage/Source/TedsCore/Private/TypedElementDatabaseEnvironment.cpp
@@ -140,8 +140,7 @@ namespace UE::Editor::DataStorage
 		Queries.UpdateActivatableQueries();
 		FlushCommands();
 		ScratchBuffer.BatchDelete();
-		MappingTable.RemoveInvalidRows();
-
+		
 		UpdateCycleId++;
 	}
 
diff --git a/Engine/Plugins/Experimental/EditorDataStorage/Source/TedsCore/Private/TypedElementDatabaseIndexTable.cpp b/Engine/Plugins/Experimental/EditorDataStorage/Source/TedsCore/Private/TypedElementDatabaseIndexTable.cpp
--- a/Engine/Plugins/Experimental/EditorDataStorage/Source/TedsCore/Private/TypedElementDatabaseIndexTable.cpp
+++ b/Engine/Plugins/Experimental/EditorDataStorage/Source/TedsCore/Private/TypedElementDatabaseIndexTable.cpp
@@ -11,7 +11,6 @@
 
 namespace UE::Editor::DataStorage
 {
-	const FTimespan FMappingTable::FCleanUpInfo::TargetDuration = FTimespan::FromMicroseconds(2000.0);
 	const FTimespan FMappingTable::FCleanUpInfo::JobShrinkThreshold = FTimespan::FromMicroseconds(100);
 	const FTimespan FMappingTable::FCleanUpInfo::JobGrowthThreshold = FTimespan::FromMicroseconds(500);
 
@@ -21,6 +20,16 @@ namespace UE::Editor::DataStorage
 	FMappingTable::FMappingTable(UEditorDataStorage& DataStorage)
 		: DataStorage(DataStorage)
 	{
+		DataStorage.RegisterCooperativeUpdate(FName("Mapping Table GC"), ICoreProvider::ECooperativeTaskPriority::Medium,
+			[this](FTimespan TimeAllowance)
+			{
+				RemoveInvalidRows(TimeAllowance);
+			});
+	}
+
+	FMappingTable::~FMappingTable()
+	{
+		DataStorage.UnregisterCooperativeUpdate(FName("Mapping Table GC"));
 	}
 
 	RowHandle FMappingTable::Lookup(EGlobalLockScope LockScope, const FName& Domain, const FMapKeyView& Key) const
@@ -139,7 +148,7 @@ namespace UE::Editor::DataStorage
 				Info.RemainingJobs -= NumBatches;
 
 				// Follow up with any adjustments needed.
-				DomainRemainingTime -= FTimespan(FPlatformTime::Cycles64() - DomainStartTime);
+				DomainRemainingTime -= FTimespan(FTimespan::FromSeconds(FPlatformTime::ToSeconds64(FPlatformTime::Cycles64() - DomainStartTime)));
 				if (NumBatches == Info.MaxNumJobs) // Only adjust if a full set of batches was used to avoid skewing.
 				{
 					AdjustJobCount(DomainRemainingTime);
@@ -158,7 +167,7 @@ namespace UE::Editor::DataStorage
 
 			UE_LOG(LogEditorDataStorage, Verbose,
 				TEXT("Mapping cleanup - %7.2fms - Domain %s has%sremaining rows, Batch size: %i, Job count: %i, Remaining jobs: %i, Max batches: %i"),
-				FTimespan(FPlatformTime::Cycles64() - DomainStartTime).GetTotalMilliseconds(),
+				FTimespan(FTimespan::FromSeconds(FPlatformTime::ToSeconds64(FPlatformTime::Cycles64() - DomainStartTime))).GetTotalMilliseconds(),
 				*DomainName.ToString(),
 				Info.DeletionQueue.IsEmpty() ? TEXT(" no ") : TEXT(" "),
 				Info.BatchSize,
@@ -167,7 +176,7 @@ namespace UE::Editor::DataStorage
 				Info.MaxNumJobs);
 
 			// See if there's time left for another pass for the next domain.
-			RemainingTime -= FTimespan(FPlatformTime::Cycles64() - DomainStartTime);
+			RemainingTime -= FTimespan(FTimespan::FromSeconds(FPlatformTime::ToSeconds64(FPlatformTime::Cycles64() - DomainStartTime)));
 			
 			// If there are no more remaining jobs and nothing queued for deletion that the domain has been fully cleaned up and should return false.
 			return (Info.RemainingJobs > 0 || !Info.DeletionQueue.IsEmpty()) 
@@ -180,15 +189,12 @@ namespace UE::Editor::DataStorage
 		}
 	}
 
-	void FMappingTable::RemoveInvalidRows()
+	void FMappingTable::RemoveInvalidRows(FTimespan RemainingTime)
 	{
 		TEDS_EVENT_SCOPE("Mapping clean up");
 
 		FScopedExclusiveLock Lock(EGlobalLockScope::Public);
 
-		FTimespan RemainingTime = FCleanUpInfo::TargetDuration;
-		uint64 StartTime = FPlatformTime::Cycles64();
-
 		// If there's no work from a previous pass, then start a new pass if needed.
 		if (RemainingUnverifiedDomains == 0)
 		{
@@ -328,7 +334,7 @@ namespace UE::Editor::DataStorage
 		// Only adjust if this is a full batch, otherwise partial batches skew performance stats.
 		if (bIsFullBatch)
 		{
-			FTimespan Duration(FPlatformTime::Cycles64() - StartTime);
+			FTimespan Duration(FTimespan::FromSeconds(FPlatformTime::ToSeconds64(FPlatformTime::Cycles64() - StartTime)));
 			if (Duration >= FCleanUpInfo::MaxBatchDuration)
 			{
 				++CleanUpInfo->BatchWentOverTime;
@@ -346,7 +352,8 @@ namespace UE::Editor::DataStorage
 
 		uint64 StartTime = FPlatformTime::Cycles64();
 
-		while (!CleanUpInfo->DeletionQueue.IsEmpty() && FTimespan(FPlatformTime::Cycles64() - StartTime) < RemainingFrameTime)
+		while (!CleanUpInfo->DeletionQueue.IsEmpty() && 
+			FTimespan(FTimespan::FromSeconds(FPlatformTime::ToSeconds64(FPlatformTime::Cycles64() - StartTime))) < RemainingFrameTime)
 		{
 			TOptional<FCleanUpInfo::InvalidRowContainer> OptionalContainer = CleanUpInfo->DeletionQueue.Dequeue();
 			checkf(OptionalContainer.IsSet(),
diff --git a/Engine/Plugins/Experimental/EditorDataStorage/Source/TedsCore/Private/TypedElementDatabaseIndexTable.h b/Engine/Plugins/Experimental/EditorDataStorage/Source/TedsCore/Private/TypedElementDatabaseIndexTable.h
--- a/Engine/Plugins/Experimental/EditorDataStorage/Source/TedsCore/Private/TypedElementDatabaseIndexTable.h
+++ b/Engine/Plugins/Experimental/EditorDataStorage/Source/TedsCore/Private/TypedElementDatabaseIndexTable.h
@@ -27,6 +27,7 @@ namespace UE::Editor::DataStorage
 	{
 	public:
 		explicit FMappingTable(UEditorDataStorage& DataStorage);
+		~FMappingTable();
 
 		RowHandle Lookup(EGlobalLockScope LockScope, const FName& Domain, const FMapKeyView& Key) const;
 		void Map(EGlobalLockScope LockScope, const FName& Domain, FMapKey Key, RowHandle Row);
@@ -35,16 +36,12 @@ namespace UE::Editor::DataStorage
 		void Remove(EGlobalLockScope LockScope, const FName& Domain, const FMapKeyView& Key);
 
 		void MarkDirty();
-		void RemoveInvalidRows();
 
 	private:
 		struct FCleanUpInfo final
 		{
 			using InvalidRowContainer = TArray<int32>;
 
-			// The target amount of time that should be spend per frame on cleaning up old rows.
-			static const FTimespan TargetDuration;
-
 			// If less than this time is available after running jobs in a frame, reduce the number of jobs.
 			static const FTimespan JobShrinkThreshold;
 			// If more than this time is left per frame, increase the number of jobs.
@@ -141,6 +138,7 @@ namespace UE::Editor::DataStorage
 		
 		UEditorDataStorage& DataStorage;
 		
+		void RemoveInvalidRows(FTimespan RemainingTime); 
 		bool IsDirty() const;
 	};
 } // namespace UE::Editor::DataStorage
diff --git a/Engine/Source/Runtime/TypedElementFramework/Public/Elements/Interfaces/TypedElementDataStorageInterface.h b/Engine/Source/Runtime/TypedElementFramework/Public/Elements/Interfaces/TypedElementDataStorageInterface.h
--- a/Engine/Source/Runtime/TypedElementFramework/Public/Elements/Interfaces/TypedElementDataStorageInterface.h
+++ b/Engine/Source/Runtime/TypedElementFramework/Public/Elements/Interfaces/TypedElementDataStorageInterface.h
@@ -21,6 +21,7 @@
 #include "Features/IModularFeature.h"
 #include "Math/NumericLimits.h"
 #include "Misc/EnumClassFlags.h"
+#include "Misc/Timespan.h"
 #include "Templates/Function.h"
 #include "UObject/Interface.h"
 #include "UObject/NameTypes.h"
@@ -526,7 +527,8 @@ public:
 	virtual void RemoveRowMapping(const FName& Domain, const FMapKeyView& Key) = 0;
 	
 	/**
-	 * @section Miscellaneous
+	 * @section Tick
+	 * Includes callbacks to respond to various steps during TEDS' tick.
 	 */
 	
 	/**
@@ -540,7 +542,30 @@ public:
 	 */
 	virtual FTypedElementOnDataStorageUpdate& OnUpdateCompleted() = 0;
 
+	using FOnCooperativeUpdate = TFunction<void(FTimespan TimeAllowance)>;
+	/** The higher the priority, the more frequently tasks get called. See RegisterCooperativeUpdate. */
+	enum class ECooperativeTaskPriority : uint8
+	{
+		High,
+		Medium,
+		Low
+	};
 	/**
+	 * Each tick TEDS will use the remaining time in a frame to execute task in the cooperative queue. A time slicer uses a 
+	 * cooperative threading model so once a task has been started it will not be interrupted and it's up to the registered 
+	 * function to try to stay within the provided time window.
+	 * Tasks with a higher priority will on average be run more frequently, but other priority tasks will be given a chance to
+	 * do work so guarantees are given that lower priority aren't starved out by higher priority tasks.
+	 */
+	virtual void RegisterCooperativeUpdate(const FName& TaskName, ECooperativeTaskPriority Priority, FOnCooperativeUpdate Callback) = 0;
+	/** Removes a previously registered time sliced callback. */
+	virtual void UnregisterCooperativeUpdate(const FName& TaskName) = 0;
+
+	/**
+	 * @section Miscellaneous
+	 */
+	
+	 /**
 	 * Whether or not the data storage is available. The data storage is available most of the time, but can be
 	 * unavailable for a brief time between being destroyed and a new one created.
 	 */
