FEnvironment : Add a destructor to clear pending commands in the Legacy::FCommandBuffer.

Previously, this class was relying on member destructors which are called in reverse order so the FScratchBuffer was being deleted before the Legacy::FCommandBuffer and the command buffer could have commands (e.g. FAddDataColumnCommand) that point to allocations from the now-deleted scratch buffer.

diff --git a/Engine/Plugins/Experimental/EditorDataStorage/Source/TedsCore/Private/TypedElementDatabaseCommandBuffer.cpp b/Engine/Plugins/Experimental/EditorDataStorage/Source/TedsCore/Private/TypedElementDatabaseCommandBuffer.cpp
--- a/Engine/Plugins/Experimental/EditorDataStorage/Source/TedsCore/Private/TypedElementDatabaseCommandBuffer.cpp
+++ b/Engine/Plugins/Experimental/EditorDataStorage/Source/TedsCore/Private/TypedElementDatabaseCommandBuffer.cpp
@@ -346,6 +346,12 @@ namespace UE::Editor::DataStorage::Legacy
 		Commands.Reset();
 	}
 
+	void FCommandBuffer::ClearCommands()
+	{
+		PendingColumns.Reset();
+		Commands.Reset();
+	}
+
 	// 
 	// misc
 	//
diff --git a/Engine/Plugins/Experimental/EditorDataStorage/Source/TedsCore/Private/TypedElementDatabaseCommandBuffer.h b/Engine/Plugins/Experimental/EditorDataStorage/Source/TedsCore/Private/TypedElementDatabaseCommandBuffer.h
--- a/Engine/Plugins/Experimental/EditorDataStorage/Source/TedsCore/Private/TypedElementDatabaseCommandBuffer.h
+++ b/Engine/Plugins/Experimental/EditorDataStorage/Source/TedsCore/Private/TypedElementDatabaseCommandBuffer.h
@@ -73,6 +73,7 @@ namespace UE::Editor::DataStorage
 				FMassFragmentBitSet FragmentsToRemove, FMassTagBitSet TagsToRemove);
 	
 			void ProcessCommands();
+			void ClearCommands();
 	
 		private:
 			struct FAddColumnCommand
diff --git a/Engine/Plugins/Experimental/EditorDataStorage/Source/TedsCore/Private/TypedElementDatabaseEnvironment.cpp b/Engine/Plugins/Experimental/EditorDataStorage/Source/TedsCore/Private/TypedElementDatabaseEnvironment.cpp
--- a/Engine/Plugins/Experimental/EditorDataStorage/Source/TedsCore/Private/TypedElementDatabaseEnvironment.cpp
+++ b/Engine/Plugins/Experimental/EditorDataStorage/Source/TedsCore/Private/TypedElementDatabaseEnvironment.cpp
@@ -19,6 +19,11 @@ namespace UE::Editor::DataStorage
 		DynamicColumnGenerator.SetQueryStore(Queries);
 	}
 
+	FEnvironment::~FEnvironment()
+	{
+		DirectDeferredCommands.ClearCommands();
+	}
+
 	Legacy::FCommandBuffer& FEnvironment::GetDirectDeferredCommands()
 	{
 		return DirectDeferredCommands;
diff --git a/Engine/Plugins/Experimental/EditorDataStorage/Source/TedsCore/Private/TypedElementDatabaseEnvironment.h b/Engine/Plugins/Experimental/EditorDataStorage/Source/TedsCore/Private/TypedElementDatabaseEnvironment.h
--- a/Engine/Plugins/Experimental/EditorDataStorage/Source/TedsCore/Private/TypedElementDatabaseEnvironment.h
+++ b/Engine/Plugins/Experimental/EditorDataStorage/Source/TedsCore/Private/TypedElementDatabaseEnvironment.h
@@ -23,6 +23,8 @@ namespace UE::Editor::DataStorage
 		FEnvironment(UEditorDataStorage& InDataStorage,
 			FMassEntityManager& InMassEntityManager, FMassProcessingPhaseManager& InMassPhaseManager);
 
+		~FEnvironment();
+
 		Legacy::FCommandBuffer& GetDirectDeferredCommands();
 		const Legacy::FCommandBuffer& GetDirectDeferredCommands() const;
 
