UEditorDataStorageCompatibility : Release this object's reference to the Environment from Deinitialize.

UEditorDataStorage is the last to Deinitialize and should be the one to delete the Environment.

diff --git a/Engine/Plugins/Experimental/EditorDataStorage/Source/TedsCore/Private/TypedElementDatabase.cpp b/Engine/Plugins/Experimental/EditorDataStorage/Source/TedsCore/Private/TypedElementDatabase.cpp
--- a/Engine/Plugins/Experimental/EditorDataStorage/Source/TedsCore/Private/TypedElementDatabase.cpp
+++ b/Engine/Plugins/Experimental/EditorDataStorage/Source/TedsCore/Private/TypedElementDatabase.cpp
@@ -1594,7 +1594,10 @@ void UEditorDataStorage::Reset()
 	}
 	Tables.Reset();
 	TableNameLookup.Reset();
+
+	checkf(Environment.IsUnique(), TEXT("UEditorDataStorage should hold the last reference to the Environment and be the one to delete it."));
 	Environment.Reset();
+
 	ActiveEditorPhaseManager.Reset();
 	ActiveEditorEntityManager.Reset();
 }
diff --git a/Engine/Plugins/Experimental/EditorDataStorage/Source/TedsCore/Private/TypedElementDatabaseCompatibility.cpp b/Engine/Plugins/Experimental/EditorDataStorage/Source/TedsCore/Private/TypedElementDatabaseCompatibility.cpp
--- a/Engine/Plugins/Experimental/EditorDataStorage/Source/TedsCore/Private/TypedElementDatabaseCompatibility.cpp
+++ b/Engine/Plugins/Experimental/EditorDataStorage/Source/TedsCore/Private/TypedElementDatabaseCompatibility.cpp
@@ -135,8 +135,14 @@ void UEditorDataStorageCompatibility::Deinitialize()
 	FCoreUObjectDelegates::OnObjectModified.Remove(ObjectModifiedDelegateHandle);
 	FCoreUObjectDelegates::OnObjectPropertyChanged.Remove(PostEditChangePropertyDelegateHandle);
 	FCoreUObjectDelegates::OnPreObjectPropertyChanged.Remove(PreEditChangePropertyDelegateHandle);
-	
-	Reset();
+
+	if (Storage)
+	{
+		Storage->OnUpdate().RemoveAll(this);
+	}
+
+	Environment.Reset();
+	Storage = nullptr;
 }
 
 void UEditorDataStorageCompatibility::RegisterRegistrationFilter(ObjectRegistrationFilter Filter)
@@ -367,10 +373,6 @@ void UEditorDataStorageCompatibility::Prepare()
 	RegisterTypeInformationQueries();
 }
 
-void UEditorDataStorageCompatibility::Reset()
-{
-}
-
 void UEditorDataStorageCompatibility::CreateStandardArchetypes()
 {
 	// Thread-safe as this is only called from a function that has an exclusive lock.
diff --git a/Engine/Plugins/Experimental/EditorDataStorage/Source/TedsCore/Private/TypedElementDatabaseCompatibility.h b/Engine/Plugins/Experimental/EditorDataStorage/Source/TedsCore/Private/TypedElementDatabaseCompatibility.h
--- a/Engine/Plugins/Experimental/EditorDataStorage/Source/TedsCore/Private/TypedElementDatabaseCompatibility.h
+++ b/Engine/Plugins/Experimental/EditorDataStorage/Source/TedsCore/Private/TypedElementDatabaseCompatibility.h
@@ -103,7 +103,6 @@ private:
 	};
 
 	UE_API void Prepare();
-	UE_API void Reset();
 	UE_API void CreateStandardArchetypes();
 	UE_API void RegisterTypeInformationQueries();
 	
