FRegistrationCommandChange and FDeregistrationCommandChange guard against a null DataStorageCompat->Storage in their destructors.

diff --git a/Engine/Plugins/Experimental/EditorDataStorage/Source/TedsCore/Private/TypedElementDatabaseCompatibility.cpp b/Engine/Plugins/Experimental/EditorDataStorage/Source/TedsCore/Private/TypedElementDatabaseCompatibility.cpp
--- a/Engine/Plugins/Experimental/EditorDataStorage/Source/TedsCore/Private/TypedElementDatabaseCompatibility.cpp
+++ b/Engine/Plugins/Experimental/EditorDataStorage/Source/TedsCore/Private/TypedElementDatabaseCompatibility.cpp
@@ -1272,7 +1272,7 @@ UEditorDataStorageCompatibility::FRegistrationCommandChange::~FRegistrationComma
 	
 	// If there has been no revert operation, there's also no memento.
 	if (UEditorDataStorageCompatibility* DataStorageCompat = Owner.Get(); 
-		DataStorageCompat && DataStorageCompat->Storage->IsRowAvailable(MementoRow))
+		DataStorageCompat && DataStorageCompat->Storage && DataStorageCompat->Storage->IsRowAvailable(MementoRow))
 	{
 		if (Private::bUseCommandBuffer)
 		{
@@ -1379,7 +1379,7 @@ UEditorDataStorageCompatibility::FDeregistrationCommandChange::~FDeregistrationC
 
 	// There's no memento row if target object was never registered with TEDS Compat.
 	if (UEditorDataStorageCompatibility* DataStorageCompat = Owner.Get();
-		DataStorageCompat && DataStorageCompat->Storage->IsRowAvailable(MementoRow))
+		DataStorageCompat && DataStorageCompat->Storage && DataStorageCompat->Storage->IsRowAvailable(MementoRow))
 	{
 		if (Private::bUseCommandBuffer)
 		{
