Misc TEDS Hierarchy Fixes, potential memory corruption + Unit Tests
* SetParentRow was holding onto a pointer to a relocated column (HierarchyDataColumn_Parent) which could occur when both the PreviousParent and NewParent are in the same chunk and the NewParent's data is at the end of the chunk.  If the PreviousParent's table is changed, generating an archtype change in mass, then a swap and pop with NewParent occurs.  This leads to HierarchyDataColumn_Parent pointing to an unallocated portion of the chunk.
* WalkDepthFirst change to not call function on root if root is not available.
* Several cases where deleted, reparented or unparented rows would leave hierarchy data column intact.
* Unit tests for topology changes

Non-functional refactor for readability duplicated code to find the HierarchyDataColumn within the chunk (see Hack_FindRowIndexInChunk) has been coalesced into a helper function.

diff --git a/Engine/Plugins/Experimental/EditorDataStorage/Source/TedsCore/Private/Hierarchy/EditorDataStorageHierarchyImplementation.cpp b/Engine/Plugins/Experimental/EditorDataStorage/Source/TedsCore/Private/Hierarchy/EditorDataStorageHierarchyImplementation.cpp
--- a/Engine/Plugins/Experimental/EditorDataStorage/Source/TedsCore/Private/Hierarchy/EditorDataStorageHierarchyImplementation.cpp
+++ b/Engine/Plugins/Experimental/EditorDataStorage/Source/TedsCore/Private/Hierarchy/EditorDataStorageHierarchyImplementation.cpp
@@ -162,6 +162,15 @@ namespace UE::Editor::DataStorage
 		return RowIndex;
 	};
 
+	static FEditorDataHierarchyData_Template* GetHierarchyDataComponentHelperInProcessor(ICommonQueryContext& Context, RowHandle Row, const UScriptStruct* ColumnType)
+	{
+		const int RowIndexInChunk = Hack_FindRowIndexInChunk(Context, Row);
+		check(RowIndexInChunk != INDEX_NONE);
+		void* HierarchyDataComponentArrayStart = Context.GetMutableColumn(ColumnType);
+		uint8* HierarchyDataComponentObjectAddress = static_cast<uint8*>(HierarchyDataComponentArrayStart) + (RowIndexInChunk * ColumnType->GetStructureSize());
+		return reinterpret_cast<FEditorDataHierarchyData_Template*>(HierarchyDataComponentObjectAddress);
+	};
+
 	void FTedsHierarchyRegistrar::RegisterObservers(ICoreProvider& DataStorage, const FHierarchyHandle& Handle)
 	{
 		using namespace UE::Editor::DataStorage::Queries;
@@ -182,37 +191,33 @@ namespace UE::Editor::DataStorage
 			{
 				const FRegisteredHierarchy* Hierarchy = FindRegisteredHierarchy(Handle);
 				
-				FEditorDataHierarchyData_Template* ChildHierarchyDataComponent;
-				{
-					const int RowIndexInChunk = Hack_FindRowIndexInChunk(Context, Row);
-					check(RowIndexInChunk != INDEX_NONE);
-					void* HierarchyDataComponentArrayStart = Context.GetMutableColumn(Hierarchy->HierarchyData);
-					uint8* HierarchyDataComponentObjectAddress = static_cast<uint8*>(HierarchyDataComponentArrayStart) + (RowIndexInChunk * Hierarchy->HierarchyData->GetStructureSize());
-					ChildHierarchyDataComponent = reinterpret_cast<FEditorDataHierarchyData_Template*>(HierarchyDataComponentObjectAddress);
-				}
+				FEditorDataHierarchyData_Template* ChildHierarchyDataComponent = GetHierarchyDataComponentHelperInProcessor(Context, Row, Hierarchy->HierarchyData);
+
 				if (ChildHierarchyDataComponent->Parent == InvalidRowHandle)
 				{
 					return;
 				}
 				// Remove the ParentTag from the parent row if it has no more children
 				bool bRemoveParentTag = false;
-				Context.RunSubquery(0, ChildHierarchyDataComponent->Parent, CreateSubqueryCallbackBinding([Row, &bRemoveParentTag, Hierarchy](ISubqueryContext& SubQueryContext, RowHandle ParentRow)
+				bool bRemoveParentsDataColumn = false;
+				Context.RunSubquery(0, ChildHierarchyDataComponent->Parent, CreateSubqueryCallbackBinding([Row, &bRemoveParentTag, &bRemoveParentsDataColumn, Hierarchy](ISubqueryContext& SubQueryContext, RowHandle ParentRow)
 				{
-					FEditorDataHierarchyData_Template* ParentHierarchyDataComponent;
-					{
-						const int RowIndexInChunk = Hack_FindRowIndexInChunk(SubQueryContext, ParentRow);
-						check(RowIndexInChunk != INDEX_NONE);
-						void* HierarchyDataComponentArrayStart = SubQueryContext.GetMutableColumn(Hierarchy->HierarchyData);
-						uint8* HierarchyDataComponentObjectAddress = static_cast<uint8*>(HierarchyDataComponentArrayStart) + (RowIndexInChunk * Hierarchy->HierarchyData->GetStructureSize());
-						ParentHierarchyDataComponent = reinterpret_cast<FEditorDataHierarchyData_Template*>(HierarchyDataComponentObjectAddress);
-					}
+					FEditorDataHierarchyData_Template* ParentHierarchyDataComponent = GetHierarchyDataComponentHelperInProcessor(SubQueryContext, ParentRow, Hierarchy->HierarchyData);
 					
 					ParentHierarchyDataComponent->Children.RemoveSwap(FTedsRowHandle(Row));
 					if (ParentHierarchyDataComponent->Children.IsEmpty())
 					{
 						bRemoveParentTag = true;
+						if (ParentHierarchyDataComponent->Parent == InvalidRowHandle)
+						{
+							bRemoveParentsDataColumn = true;
+						}
 					}
 				}));
+				if (bRemoveParentsDataColumn)
+				{
+					Context.RemoveColumns(ChildHierarchyDataComponent->Parent, {Hierarchy->HierarchyData});
+				}
 				if (bRemoveParentTag)
 				{
 					Context.RemoveColumns(ChildHierarchyDataComponent->Parent, {Hierarchy->ParentTag});
@@ -232,39 +237,32 @@ namespace UE::Editor::DataStorage
 			[this, Handle](IQueryContext& Context, RowHandle ParentRow)
 			{
 				const FRegisteredHierarchy* Hierarchy = FindRegisteredHierarchy(Handle);
-				
-				FEditorDataHierarchyData_Template* ParentHierarchyDataComponent;
-				{
-					const int RowIndexInChunk = Hack_FindRowIndexInChunk(Context, ParentRow);
-					check(RowIndexInChunk != INDEX_NONE);
-					void* HierarchyDataComponentArrayStart = Context.GetMutableColumn(Hierarchy->HierarchyData);
-					uint8* HierarchyDataComponentObjectAddress = static_cast<uint8*>(HierarchyDataComponentArrayStart) + (RowIndexInChunk * Hierarchy->HierarchyData->GetStructureSize());
-					ParentHierarchyDataComponent = reinterpret_cast<FEditorDataHierarchyData_Template*>(HierarchyDataComponentObjectAddress);
-				}
-				if (ParentHierarchyDataComponent->Children.IsEmpty())
-				{
-					return;
-				}
+
+				FEditorDataHierarchyData_Template* ParentHierarchyDataComponent = GetHierarchyDataComponentHelperInProcessor(Context, ParentRow, Hierarchy->HierarchyData);
+					
 				for (int32 ChildIndex = 0; ChildIndex < ParentHierarchyDataComponent->Children.Num(); ++ChildIndex)
 				{
 					RowHandle ChildRow = ParentHierarchyDataComponent->Children[ChildIndex];
 					Context.RunSubquery(0, ChildRow, CreateSubqueryCallbackBinding([ParentRow, Hierarchy](ISubqueryContext& SubQueryContext, RowHandle ChildRow)
 					{
-						FEditorDataHierarchyData_Template* ChildHierarchyDataComponent;
+						FEditorDataHierarchyData_Template* ChildHierarchyDataComponent = GetHierarchyDataComponentHelperInProcessor(SubQueryContext, ChildRow, Hierarchy->HierarchyData);
+
+						ChildHierarchyDataComponent->Parent = InvalidRowHandle;
+						if (ChildHierarchyDataComponent->Children.IsEmpty())
 						{
-							const int RowIndexInChunk = Hack_FindRowIndexInChunk(SubQueryContext, ChildRow);
-							check(RowIndexInChunk != INDEX_NONE);
-							void* HierarchyDataComponentArrayStart = SubQueryContext.GetMutableColumn(Hierarchy->HierarchyData);
-							uint8* HierarchyDataComponentObjectAddress = static_cast<uint8*>(HierarchyDataComponentArrayStart) + (RowIndexInChunk * Hierarchy->HierarchyData->GetStructureSize());
-							ChildHierarchyDataComponent = reinterpret_cast<FEditorDataHierarchyData_Template*>(HierarchyDataComponentObjectAddress);
+							SubQueryContext.RemoveColumns(ChildRow, {Hierarchy->HierarchyData});
 						}
-						ChildHierarchyDataComponent->Parent = InvalidRowHandle;
 					}));
-
+					
 					Context.RemoveColumns(ChildRow, {Hierarchy->ChildTag});
+					
 				}
 
 				ParentHierarchyDataComponent->Children.Reset();
+				if (ParentHierarchyDataComponent->Parent == InvalidRowHandle)
+				{
+					Context.RemoveColumns(ParentRow, {Hierarchy->HierarchyData});
+				}
 			}
 		)
 		.ReadWrite(Hierarchy->HierarchyData)
@@ -378,6 +376,10 @@ namespace UE::Editor::DataStorage
 		RowHandle Row,
 		TFunction<void(const ICoreProvider& Context, RowHandle Owner,RowHandle Target)> OnVisitedFn) const
 	{
+		if (!Context.IsRowAvailable(Row))
+		{
+			return;
+		}
 		// Recursively walk depth first to visit each child
 		auto WalkDepthFirstImpl = [this](const ICoreProvider& Context_, RowHandle Row, decltype(OnVisitedFn)& OnVisitedFn_, auto& WalkDepthFirstImplRef) -> void
 		{
@@ -425,19 +427,6 @@ namespace UE::Editor::DataStorage
 			);
 		};
 
-		FEditorDataHierarchyData_Template* HierarchyDataColumn_Parent;
-		if (void* ParentHierarchyDataColumnAddress = CoreProvider->GetColumnData(Parent, HierarchyDataColumnType); ParentHierarchyDataColumnAddress == nullptr)
-		{
-			AddHierarchyDataColumn(CoreProvider, Parent, InvalidRowHandle);
-			ParentHierarchyDataColumnAddress = CoreProvider->GetColumnData(Parent, HierarchyDataColumnType);
-			check(ParentHierarchyDataColumnAddress);
-			HierarchyDataColumn_Parent = static_cast<FEditorDataHierarchyData_Template*>(ParentHierarchyDataColumnAddress);
-		}
-		else
-		{
-			HierarchyDataColumn_Parent = static_cast<FEditorDataHierarchyData_Template*>(ParentHierarchyDataColumnAddress);
-		}
-
 		RowHandle PreviousParent = InvalidRowHandle;
 		FEditorDataHierarchyData_Template* HierarchyDataColumn_Target;
 		if (void* TargetHierarchyDataColumnAddress = CoreProvider->GetColumnData(Target, HierarchyDataColumnType); TargetHierarchyDataColumnAddress == nullptr)
@@ -471,6 +460,10 @@ namespace UE::Editor::DataStorage
 
 					if (HierarchyDataColumn_PreviousParent->Children.IsEmpty())
 					{
+						if (HierarchyDataColumn_PreviousParent->Parent == InvalidRowHandle)
+						{
+							CoreProvider->RemoveColumn(PreviousParent, HierarchyDataColumnType);
+						}
 						CoreProvider->RemoveColumn(PreviousParent, ParentTag);
 					}
 				}
@@ -479,6 +472,19 @@ namespace UE::Editor::DataStorage
 				HierarchyDataColumn_Target->Parent = Parent;
 				CoreProvider->AddColumn(Target, ChildTag);
 
+				FEditorDataHierarchyData_Template* HierarchyDataColumn_Parent;
+				if (void* ParentHierarchyDataColumnAddress = CoreProvider->GetColumnData(Parent, HierarchyDataColumnType); ParentHierarchyDataColumnAddress == nullptr)
+				{
+					AddHierarchyDataColumn(CoreProvider, Parent, InvalidRowHandle);
+					ParentHierarchyDataColumnAddress = CoreProvider->GetColumnData(Parent, HierarchyDataColumnType);
+					check(ParentHierarchyDataColumnAddress);
+					HierarchyDataColumn_Parent = static_cast<FEditorDataHierarchyData_Template*>(ParentHierarchyDataColumnAddress);
+				}
+				else
+				{
+					HierarchyDataColumn_Parent = static_cast<FEditorDataHierarchyData_Template*>(ParentHierarchyDataColumnAddress);
+				}
+
 				// Update the new parent with the new information
 				HierarchyDataColumn_Parent->Children.Add(FTedsRowHandle(Target));
 				CoreProvider->AddColumn(Parent, ParentTag);
diff --git a/Engine/Plugins/Experimental/EditorDataStorage/Source/TedsCore/Tests/DynamicColumnTests.cpp b/Engine/Plugins/Experimental/EditorDataStorage/Source/TedsCore/Tests/DynamicColumnTests.cpp
--- a/Engine/Plugins/Experimental/EditorDataStorage/Source/TedsCore/Tests/DynamicColumnTests.cpp
+++ b/Engine/Plugins/Experimental/EditorDataStorage/Source/TedsCore/Tests/DynamicColumnTests.cpp
@@ -282,4 +282,293 @@ namespace UE::Editor::DataStorage::Tests
 	}
 } // namespace UE::Editor::DataStorage::Tests
 
+namespace UE::Editor::DataStorage::Tests
+{
+	BEGIN_DEFINE_SPEC(TedsHierarchyTestFixture, "Editor.DataStorage.Hierarchy", EAutomationTestFlags::EditorContext | EAutomationTestFlags::EngineFilter)
+		const FName TestTableName = TEXT("TestTable_HierarchyTestTable");
+		const FName HierarchyName = TEXT("TedsHierarchyTestFixture_Hierarchy");
+		FHierarchyHandle HierarchyHandle;
+		ICoreProvider* TedsInterface = nullptr;
+		TableHandle TestTable;
+		TArray<RowHandle> Rows;
+
+		TableHandle RegisterTestTable() const
+		{
+			const TableHandle Table = TedsInterface->FindTable(TestTableName);
+					
+			if (Table != InvalidTableHandle)
+			{
+				return Table;
+			}
+					
+			return TedsInterface->RegisterTable(
+			{
+				FTestColumnB::StaticStruct()
+			},
+			TestTableName);
+		}
+
+		FHierarchyHandle RegisterHierarchy()
+		{
+			HierarchyHandle = TedsInterface->FindHierarchyByName(HierarchyName);
+			if (!TedsInterface->IsValidHierarchyHandle(HierarchyHandle))
+			{
+				UE::Editor::DataStorage::FHierarchyRegistrationParams EntityHierarchyParams{.Name = HierarchyName};
+				HierarchyHandle = TedsInterface->RegisterHierarchy(EntityHierarchyParams);
+			}
+			return HierarchyHandle;
+		}
+
+		RowHandle CreateTestRow(TableHandle InTableHandle)
+		{
+			RowHandle Row = TedsInterface->AddRow(InTableHandle);
+			Rows.Add(Row);
+			return Row;
+		}
+
+		struct FTreeNodes
+		{
+			RowHandle A = InvalidRowHandle;
+			RowHandle B = InvalidRowHandle;
+			RowHandle C = InvalidRowHandle;
+			RowHandle D = InvalidRowHandle;
+			RowHandle E = InvalidRowHandle;
+		};
+
+		int32 ComputeTreeSize(RowHandle Root)
+		{
+			int32 Count = 0;
+			TedsInterface->WalkDepthFirst(HierarchyHandle, Root, [&Count](const ICoreProvider& Context, RowHandle Owner, RowHandle Target)
+				{
+					Count++;
+				});
+			return Count;
+		};
+
+		// Build out this topology:
+		// A
+		// - B
+		FTreeNodes BuildSimpleTopology()
+		{
+			RowHandle A = CreateTestRow(TestTable);
+			RowHandle B = CreateTestRow(TestTable);
+
+			TedsInterface->SetParentRow(HierarchyHandle, B, A);
+
+			return FTreeNodes
+			{
+				.A = A,
+				.B = B,
+			};
+		};
+
+		// Build out this topology:
+		// A
+		// - B
+		// - C
+		//   - D
+		FTreeNodes BuildNestedTopology()
+		{
+			RowHandle A = CreateTestRow(TestTable);
+			RowHandle B = CreateTestRow(TestTable);
+			RowHandle C = CreateTestRow(TestTable);
+			RowHandle D = CreateTestRow(TestTable);
+
+			TedsInterface->SetParentRow(HierarchyHandle, B, A);
+			TedsInterface->SetParentRow(HierarchyHandle, C, A);
+			TedsInterface->SetParentRow(HierarchyHandle, D, C);
+
+			return FTreeNodes
+			{
+				.A = A,
+				.B = B,
+				.C = C,
+				.D = D
+			};
+		};
+		
+	END_DEFINE_SPEC(TedsHierarchyTestFixture)
+
+	void TedsHierarchyTestFixture::Define()
+	{
+		BeforeEach([this]()
+		{
+			TedsInterface = GetMutableDataStorageFeature<ICoreProvider>(StorageFeatureName);
+			TestTable = RegisterTestTable();
+			HierarchyHandle = RegisterHierarchy();
+		});	
+
+		It("Test Simple Topology", [this]()
+			{
+				auto Tree = BuildSimpleTopology();
+
+				TestEqual("", TedsInterface->GetParentRow(HierarchyHandle, Tree.A), InvalidRowHandle);
+				TestEqual("", TedsInterface->GetParentRow(HierarchyHandle, Tree.B), Tree.A);
+
+				TestEqual("", ComputeTreeSize(Tree.A), 2);
+				TestEqual("", ComputeTreeSize(Tree.B), 1);
+			});
+		
+		It("Test Nested Topology", [this]()
+		{
+			auto Tree = BuildNestedTopology();
+
+			TestEqual("", TedsInterface->GetParentRow(HierarchyHandle, Tree.A), InvalidRowHandle);
+			TestEqual("", TedsInterface->GetParentRow(HierarchyHandle, Tree.B), Tree.A);
+			TestEqual("", TedsInterface->GetParentRow(HierarchyHandle, Tree.C), Tree.A);
+			TestEqual("", TedsInterface->GetParentRow(HierarchyHandle, Tree.D), Tree.C);
+
+			TestEqual("", ComputeTreeSize(Tree.A), 4);
+			TestEqual("", ComputeTreeSize(Tree.B), 1);
+			TestEqual("", ComputeTreeSize(Tree.C), 2);
+			TestEqual("", ComputeTreeSize(Tree.D), 1);
+		});
+			
+		It("Unparent childless leaf from single-child parent", [this]()
+		{
+				auto Tree = BuildSimpleTopology();
+				TedsInterface->SetParentRow(HierarchyHandle, Tree.B, InvalidRowHandle);
+
+				TestEqual("", TedsInterface->GetParentRow(HierarchyHandle, Tree.A), InvalidRowHandle);
+				TestEqual("", TedsInterface->GetParentRow(HierarchyHandle, Tree.B), InvalidRowHandle);
+
+				TestEqual("", ComputeTreeSize(Tree.A), 1);
+				TestEqual("", ComputeTreeSize(Tree.B), 1);
+		});
+
+		It("Delete childless leaf from single-child parent", [this]()
+		{
+				auto Tree = BuildSimpleTopology();
+				TedsInterface->RemoveRow(Tree.B);
+
+				TestEqual("", TedsInterface->GetParentRow(HierarchyHandle, Tree.A), InvalidRowHandle);
+				TestEqual("", TedsInterface->GetParentRow(HierarchyHandle, Tree.B), InvalidRowHandle);
+
+				TestEqual("", ComputeTreeSize(Tree.A), 1);
+				TestEqual("", ComputeTreeSize(Tree.B), 0);
+		});
+
+		It("Delete single-child parent of childless leaf node ", [this]()
+		{
+				auto Tree = BuildSimpleTopology();
+				TedsInterface->RemoveRow(Tree.A);
+
+				TestEqual("", TedsInterface->GetParentRow(HierarchyHandle, Tree.A), InvalidRowHandle);
+				TestEqual("", TedsInterface->GetParentRow(HierarchyHandle, Tree.B), InvalidRowHandle);
+
+				TestEqual("", ComputeTreeSize(Tree.A), 0);
+				TestEqual("", ComputeTreeSize(Tree.B), 1);
+		});
+			
+		It("Unparent childless leaf from multi-child parent", [this]()
+		{
+				auto Tree = BuildNestedTopology();
+				TedsInterface->SetParentRow(HierarchyHandle, Tree.B, InvalidRowHandle);
+
+				TestEqual("", TedsInterface->GetParentRow(HierarchyHandle, Tree.A), InvalidRowHandle);
+				TestEqual("", TedsInterface->GetParentRow(HierarchyHandle, Tree.B), InvalidRowHandle);
+				TestEqual("", TedsInterface->GetParentRow(HierarchyHandle, Tree.C), Tree.A);
+				TestEqual("", TedsInterface->GetParentRow(HierarchyHandle, Tree.D), Tree.C);
+
+				TestEqual("", ComputeTreeSize(Tree.A), 3);
+				TestEqual("", ComputeTreeSize(Tree.B), 1);
+				TestEqual("", ComputeTreeSize(Tree.C), 2);
+				TestEqual("", ComputeTreeSize(Tree.D), 1);
+		});
+
+		It("Unparent middle parent from middle parent", [this]()
+		{
+				auto Tree = BuildNestedTopology();
+				TedsInterface->SetParentRow(HierarchyHandle, Tree.C, InvalidRowHandle);
+
+				TestEqual("", TedsInterface->GetParentRow(HierarchyHandle, Tree.A), InvalidRowHandle);
+				TestEqual("", TedsInterface->GetParentRow(HierarchyHandle, Tree.B), Tree.A);
+				TestEqual("", TedsInterface->GetParentRow(HierarchyHandle, Tree.C), InvalidRowHandle);
+				TestEqual("", TedsInterface->GetParentRow(HierarchyHandle, Tree.D), Tree.C);
+
+				TestEqual("", ComputeTreeSize(Tree.A), 2);
+				TestEqual("", ComputeTreeSize(Tree.B), 1);
+				TestEqual("", ComputeTreeSize(Tree.C), 2);
+				TestEqual("", ComputeTreeSize(Tree.D), 1);
+		});
+
+		It("Unparent childless leaf from middle parent", [this]()
+		{
+				auto Tree = BuildNestedTopology();
+				TedsInterface->SetParentRow(HierarchyHandle, Tree.D, InvalidRowHandle);
+
+				TestEqual("", TedsInterface->GetParentRow(HierarchyHandle, Tree.A), InvalidRowHandle);
+				TestEqual("", TedsInterface->GetParentRow(HierarchyHandle, Tree.B), Tree.A);
+				TestEqual("", TedsInterface->GetParentRow(HierarchyHandle, Tree.C), Tree.A);
+				TestEqual("", TedsInterface->GetParentRow(HierarchyHandle, Tree.D), InvalidRowHandle);
+
+				TestEqual("", ComputeTreeSize(Tree.A), 3);
+				TestEqual("", ComputeTreeSize(Tree.B), 1);
+				TestEqual("", ComputeTreeSize(Tree.C), 1);
+				TestEqual("", ComputeTreeSize(Tree.D), 1);
+		});
+
+		It("Reparent childless leaf to childless leaf", [this]()
+		{
+				auto Tree = BuildNestedTopology();
+				TedsInterface->SetParentRow(HierarchyHandle, Tree.B, Tree.D);
+
+				TestEqual("", TedsInterface->GetParentRow(HierarchyHandle, Tree.A), InvalidRowHandle);
+				TestEqual("", TedsInterface->GetParentRow(HierarchyHandle, Tree.B), Tree.D);
+				TestEqual("", TedsInterface->GetParentRow(HierarchyHandle, Tree.C), Tree.A);
+				TestEqual("", TedsInterface->GetParentRow(HierarchyHandle, Tree.D), Tree.C);
+
+				TestEqual("", ComputeTreeSize(Tree.A), 4);
+				TestEqual("", ComputeTreeSize(Tree.B), 1);
+				TestEqual("", ComputeTreeSize(Tree.C), 3);
+				TestEqual("", ComputeTreeSize(Tree.D), 2);
+		});
+
+		It("Delete childless leaf of multi-child root", [this]()
+		{
+				auto Tree = BuildNestedTopology();
+				TedsInterface->RemoveRow(Tree.B);
+
+				TestEqual("", TedsInterface->GetParentRow(HierarchyHandle, Tree.A), InvalidRowHandle);
+				TestEqual("", TedsInterface->GetParentRow(HierarchyHandle, Tree.B), InvalidRowHandle);
+				TestEqual("", TedsInterface->GetParentRow(HierarchyHandle, Tree.C), Tree.A);
+				TestEqual("", TedsInterface->GetParentRow(HierarchyHandle, Tree.D), Tree.C);
+
+				TestEqual("ComputeTreeSize(Tree.A)", ComputeTreeSize(Tree.A), 3);
+				TestEqual("ComputeTreeSize(Tree.B)", ComputeTreeSize(Tree.B), 0);
+				TestEqual("ComputeTreeSize(Tree.C)", ComputeTreeSize(Tree.C), 2);
+				TestEqual("ComputeTreeSize(Tree.D)", ComputeTreeSize(Tree.D), 1);
+		});
+
+		It("Delete childless leaf of middle parent", [this]()
+		{
+				auto Tree = BuildNestedTopology();
+				TedsInterface->RemoveRow(Tree.D);
+
+				TestEqual("", TedsInterface->GetParentRow(HierarchyHandle, Tree.A), InvalidRowHandle);
+				TestEqual("", TedsInterface->GetParentRow(HierarchyHandle, Tree.B), Tree.A);
+				TestEqual("", TedsInterface->GetParentRow(HierarchyHandle, Tree.C), Tree.A);
+				TestEqual("", TedsInterface->GetParentRow(HierarchyHandle, Tree.D), InvalidRowHandle);
+
+				TestEqual("", ComputeTreeSize(Tree.A), 3);
+				TestEqual("", ComputeTreeSize(Tree.B), 1);
+				TestEqual("", ComputeTreeSize(Tree.C), 1);
+				TestEqual("", ComputeTreeSize(Tree.D), 0);
+		});
+
+		AfterEach([this]()
+		{
+			for (RowHandle Row : Rows)
+			{
+				TedsInterface->RemoveRow(Row);
+			}
+			Rows.Empty(Rows.Num());
+
+			HierarchyHandle = FHierarchyHandle();
+
+			TedsInterface = nullptr;
+		});
+	}
+}
+
 #endif // WITH_TESTS
