[TEDS] Allow dynamic columns to be declared at compile time.

This change introduces a way to capture string literals as template arguments and uses that to allow dynamic columns to be defined at compile time. In the future this can also be done for other uses like hierarchies. The use for dynamic columns is useful for users that want to reuse an existing column but do still need a specialization. An example could be using the transform column, but having a specialization for skeleton bones. In this case the benefit of using the existing transform column is that the existing widgets will be used, but it does allow the column type is known at compile time. This still allows using the regular syntax for columns like in Query Conditions without having to resort to the API that allows setting those values at runtime.
To help with tests, `HasColumns` has been updated to be able to use `TColumn` as well. In the future more interface functions will be getting support for this.

diff --git a/Engine/Plugins/Experimental/EditorDataStorage/Source/TedsCore/Tests/DynamicColumnTests.cpp b/Engine/Plugins/Experimental/EditorDataStorage/Source/TedsCore/Tests/DynamicColumnTests.cpp
--- a/Engine/Plugins/Experimental/EditorDataStorage/Source/TedsCore/Tests/DynamicColumnTests.cpp
+++ b/Engine/Plugins/Experimental/EditorDataStorage/Source/TedsCore/Tests/DynamicColumnTests.cpp
@@ -12,6 +12,10 @@
 
 #include "Misc/AutomationTest.h"
 
+#define TEST_IDENTIFIER_0 "StaticMesh"
+#define TEST_IDENTIFIER_1 "Animation"
+#define TEST_IDENTIFIER_2 "AudioClip"
+
 namespace UE::Editor::DataStorage::Tests
 {
 	BEGIN_DEFINE_SPEC(DynamicColumnTestFixture, "Editor.DataStorage.DynamicColumns", EAutomationTestFlags::EditorContext | EAutomationTestFlags::EngineFilter)
@@ -63,7 +67,7 @@ namespace UE::Editor::DataStorage::Tests
 		{
 			TedsInterface = GetMutableDataStorageFeature<ICoreProvider>(StorageFeatureName);
 			TestTable = RegisterTestTable();
-			Identifiers = {TEXT("StaticMesh"), TEXT("Animation"), TEXT("AudioClip")};
+			Identifiers = { TEXT( TEST_IDENTIFIER_0 ), TEXT( TEST_IDENTIFIER_1 ), TEXT( TEST_IDENTIFIER_2 ) };
 			TestTrue("", TedsInterface != nullptr);
 		});
 			
@@ -141,8 +145,8 @@ namespace UE::Editor::DataStorage::Tests
 							}
 						});					
 
-						// Should match Rows[0]
 						{
+							// Should match Rows[0]
 							QueryHandle Query = RegisterQuery(
 								Select().
 								Where(TColumn<FTestDynamicTag>(Identifiers[0]) && TColumn<FTestDynamicTag>(Identifiers[1])).
@@ -152,6 +156,17 @@ namespace UE::Editor::DataStorage::Tests
 							FQueryResult Result = TedsInterface->RunQuery(Query, Callback);
 							TestEqual("Match Row[0]", Result.Count, GetMatchCount());
 						}
+						{
+							// Should match Rows[0] (using compile time dynamic columns)
+							QueryHandle Query = RegisterQuery(
+								Select().
+								Where(TColumn<FTestDynamicTag, TEST_IDENTIFIER_0>() && TColumn<FTestDynamicTag, TEST_IDENTIFIER_1>()).
+								Compile());
+
+							SetExpectedMatches({ Rows[0] });
+							FQueryResult Result = TedsInterface->RunQuery(Query, Callback);
+							TestEqual("Match Row[0]", Result.Count, GetMatchCount());
+						}
 						{
 							// Should match Rows 0 and 1
 							QueryHandle Query = RegisterQuery(
@@ -162,6 +177,16 @@ namespace UE::Editor::DataStorage::Tests
 							FQueryResult Result = TedsInterface->RunQuery(Query, Callback);
 							TestEqual("Match Row[0] and Row[1]", Result.Count, GetMatchCount());
 						}
+						{
+							// Should match Rows 0 and 1 (using compile time dynamic columns)
+							QueryHandle Query = RegisterQuery(
+								Select().
+								Where(TColumn<FTestDynamicTag, TEST_IDENTIFIER_0>()).
+								Compile());
+							SetExpectedMatches({ Rows[0], Rows[1] });
+							FQueryResult Result = TedsInterface->RunQuery(Query, Callback);
+							TestEqual("Match Row[0] and Row[1]", Result.Count, GetMatchCount());
+						}
 						{
 							// Should match Rows 0 and 2
 							QueryHandle Query = RegisterQuery(
@@ -172,6 +197,16 @@ namespace UE::Editor::DataStorage::Tests
 							FQueryResult Result = TedsInterface->RunQuery(Query, Callback);
 							TestEqual("Match Row[0] and Row[2]", Result.Count, GetMatchCount());
 						}
+						{
+							// Should match Rows 0 and 2 (using compile time dynamic columns)
+							QueryHandle Query = RegisterQuery(
+								Select().
+								Where(TColumn<FTestDynamicTag, TEST_IDENTIFIER_1>()).
+								Compile());
+							SetExpectedMatches({ Rows[0], Rows[2] });
+							FQueryResult Result = TedsInterface->RunQuery(Query, Callback);
+							TestEqual("Match Row[0] and Row[2]", Result.Count, GetMatchCount());
+						}
 						{
 							// Should match Rows 0, 1 and 2
 							QueryHandle Query = RegisterQuery(
@@ -182,6 +217,16 @@ namespace UE::Editor::DataStorage::Tests
 							FQueryResult Result = TedsInterface->RunQuery(Query, Callback);
 							TestEqual("Match All Rows", Result.Count, GetMatchCount());
 						}
+						{
+							// Should match Rows 0, 1 and 2 (using compile time dynamic columns)
+							QueryHandle Query = RegisterQuery(
+								Select().
+								Where(TColumn<FTestDynamicTag, TEST_IDENTIFIER_0>() || TColumn<FTestDynamicTag, TEST_IDENTIFIER_1>()).
+								Compile());
+							SetExpectedMatches({ Rows[0], Rows[1], Rows[2] });
+							FQueryResult Result = TedsInterface->RunQuery(Query, Callback);
+							TestEqual("Match All Rows", Result.Count, GetMatchCount());
+						}
 						{
 							// Should match Rows 0, 1 and 2
 							QueryHandle Query = RegisterQuery(
@@ -198,6 +243,8 @@ namespace UE::Editor::DataStorage::Tests
 
 			It(TEXT("Columns"), EAsyncExecution::TaskGraphMainThread, [this]()
 			{
+				using namespace UE::Editor::DataStorage::Queries;
+				
 				// Add dynamic columns that have data
 				{
 					for (int32 Index = 0; Index <= 2; ++Index)
@@ -225,27 +272,18 @@ namespace UE::Editor::DataStorage::Tests
 						.IntArray = {11, 22, 33, 43},
 					});
 
-					// Check they were added
-					// Note: There is no HasColumn function for syntactic sugar to get a dynamic column type
-					TArray<const UScriptStruct*> DynamicTagTypes;
-								
-					DynamicTagTypes.Emplace(TedsInterface->FindDynamicColumn(FDynamicColumnDescription
-						{
-							.TemplateType = FTestDynamicColumn::StaticStruct(),
-							.Identifier = Identifiers[0]
-						}));
-					DynamicTagTypes.Emplace(TedsInterface->FindDynamicColumn(FDynamicColumnDescription
-						{
-							.TemplateType = FTestDynamicColumn::StaticStruct(),
-							.Identifier = Identifiers[1]
-						}));
-					TestTrue("Expected columns not found", TedsInterface->HasColumns(Rows[0], MakeConstArrayView({DynamicTagTypes[0], DynamicTagTypes[1]})));
-								
-					TestTrue("Expected columns not found", TedsInterface->HasColumns(Rows[1], MakeConstArrayView({DynamicTagTypes[0]})));
-					TestFalse("Unexpected columns found", TedsInterface->HasColumns(Rows[1], MakeConstArrayView({DynamicTagTypes[1]})));
+					TestTrue("Expected columns not found", TedsInterface->HasColumns(Rows[0], TColumn<FTestDynamicColumn>(Identifiers[0]), TColumn<FTestDynamicColumn>(Identifiers[1])));
+					TestTrue("Expected columns not found", TedsInterface->HasColumns(Rows[0], TColumn<FTestDynamicColumn, TEST_IDENTIFIER_0>(), TColumn<FTestDynamicColumn, TEST_IDENTIFIER_1>()));
 
-					TestFalse("Unexpected columns found", TedsInterface->HasColumns(Rows[2], MakeConstArrayView({DynamicTagTypes[0]})));
-					TestTrue("Expected columns not found", TedsInterface->HasColumns(Rows[2], MakeConstArrayView({DynamicTagTypes[1]})));
+					TestTrue("Expected columns not found", TedsInterface->HasColumns(Rows[1], TColumn<FTestDynamicColumn>(Identifiers[0])));
+					TestTrue("Expected columns not found", TedsInterface->HasColumns(Rows[1], TColumn<FTestDynamicColumn, TEST_IDENTIFIER_0>()));
+					TestFalse("Unexpected columns found", TedsInterface->HasColumns(Rows[1], TColumn<FTestDynamicColumn>(Identifiers[1])));
+					TestFalse("Unexpected columns found", TedsInterface->HasColumns(Rows[1], TColumn<FTestDynamicColumn, TEST_IDENTIFIER_1>()));
+
+					TestFalse("Unexpected columns found", TedsInterface->HasColumns(Rows[2], TColumn<FTestDynamicColumn>(Identifiers[0])));
+					TestFalse("Unexpected columns found", TedsInterface->HasColumns(Rows[2], TColumn<FTestDynamicColumn, TEST_IDENTIFIER_0>()));
+					TestTrue("Expected columns not found", TedsInterface->HasColumns(Rows[2], TColumn<FTestDynamicColumn>(Identifiers[1])));
+					TestTrue("Expected columns not found", TedsInterface->HasColumns(Rows[2], TColumn<FTestDynamicColumn, TEST_IDENTIFIER_1>()));
 					
 					using namespace UE::Editor::DataStorage::Queries;
 
@@ -282,6 +320,10 @@ namespace UE::Editor::DataStorage::Tests
 	}
 } // namespace UE::Editor::DataStorage::Tests
 
+#undef TEST_IDENTIFIER_2
+#undef TEST_IDENTIFIER_1
+#undef TEST_IDENTIFIER_0
+
 namespace UE::Editor::DataStorage::Tests
 {
 	BEGIN_DEFINE_SPEC(TedsHierarchyTestFixture, "Editor.DataStorage.Hierarchy", EAutomationTestFlags::EditorContext | EAutomationTestFlags::EngineFilter)
diff --git a/Engine/Source/Runtime/TypedElementFramework/Public/DataStorage/CommonTypes.h b/Engine/Source/Runtime/TypedElementFramework/Public/DataStorage/CommonTypes.h
--- a/Engine/Source/Runtime/TypedElementFramework/Public/DataStorage/CommonTypes.h
+++ b/Engine/Source/Runtime/TypedElementFramework/Public/DataStorage/CommonTypes.h
@@ -4,6 +4,7 @@
 
 #include "DataStorage/Handles.h"
 
+#include <algorithm>
 #include "CommonTypes.generated.h"
 
 /**
@@ -102,5 +103,25 @@ namespace UE
 
 		template<typename T>
 		concept TEnumType = std::is_enum_v<T>;
+
+		/** Capture for variant names, e.g. for dynamic columns or hierarchies, so the name can be passed in as a template argument. */
+		template <int N>
+		struct TVariantName
+		{
+			static constexpr int Size = N;
+			char Name[N > 1 ? (N - 1) : 1];
+
+			constexpr TVariantName() requires (N == 0) = default;
+			constexpr TVariantName(const char (&InName)[N]) requires (N > 0)
+			{
+				std::copy_n(InName, N - 1, Name);
+			}
+
+			FName GetStaticName() const
+			{
+				static FName Result = N > 1 ? FName(N - 1, Name) : NAME_None;
+				return Result;
+			}
+		};
 	} // namespace Editor::DataStorage
 } // namespace UE
diff --git a/Engine/Source/Runtime/TypedElementFramework/Public/DataStorage/Queries/Conditions.h b/Engine/Source/Runtime/TypedElementFramework/Public/DataStorage/Queries/Conditions.h
--- a/Engine/Source/Runtime/TypedElementFramework/Public/DataStorage/Queries/Conditions.h
+++ b/Engine/Source/Runtime/TypedElementFramework/Public/DataStorage/Queries/Conditions.h
@@ -2,6 +2,7 @@
 
 #pragma once
 
+#include <algorithm>
 #include "DataStorage/CommonTypes.h"
 #include "Templates/SubclassOf.h"
 #include "UObject/WeakObjectPtrTemplates.h"
@@ -21,24 +22,40 @@ namespace Queries
 
 	protected:
 		FColumnBase() = default;
-		constexpr explicit FColumnBase(TWeakObjectPtr<const UScriptStruct> ColumnTypeInfo, const FName& InIdentifier)
+		explicit FColumnBase(TWeakObjectPtr<const UScriptStruct> ColumnTypeInfo, const FName& InIdentifier)
 			: TypeInfo(ColumnTypeInfo)
 			, Identifier(InIdentifier)
 		{}
 	};
 
-	template<typename T = void>
+	template<typename T, TVariantName Id = TVariantName<0>()>
 	struct TColumn final : public FColumnBase
 	{
-		template <typename U = T> requires (!std::is_same_v<U, void>)
-			constexpr explicit TColumn(const FName& Identifier = NAME_None) : FColumnBase(T::StaticStruct(), Identifier) {}
-		
-		template <typename U = T> requires (std::is_same_v<U, void>)
-			constexpr explicit TColumn(TWeakObjectPtr<const UScriptStruct> ColumnTypeInfo, const FName& Identifier = NAME_None)
-				: FColumnBase(ColumnTypeInfo, Identifier)
-		{}
+		explicit TColumn() : FColumnBase(T::StaticStruct(), NAME_None) {}
+		explicit TColumn(const FName& Identifier) : FColumnBase(T::StaticStruct(), Identifier) {}
+	};
+
+	template<typename T, TVariantName Id> requires (decltype(Id)::Size > 0)
+	struct TColumn<T, Id> final : public FColumnBase
+	{
+		explicit TColumn() : FColumnBase(T::StaticStruct(), Id.GetStaticName()) {}
 	};
 
+	template<>
+	struct TColumn<void> final : public FColumnBase
+	{
+		explicit TColumn(TWeakObjectPtr<const UScriptStruct> ColumnTypeInfo) : FColumnBase(ColumnTypeInfo, NAME_None) {}
+		explicit TColumn(TWeakObjectPtr<const UScriptStruct> ColumnTypeInfo, const FName& Identifier)
+			: FColumnBase(ColumnTypeInfo, Identifier)
+		{
+		}
+	};
+	TColumn(TWeakObjectPtr<const UScriptStruct>) -> TColumn<void>;
+	TColumn(TWeakObjectPtr<const UScriptStruct>, const FName&) -> TColumn<void>;
+
+	template<typename T>
+	concept TColumnDescType = std::is_base_of_v<FColumnBase, T>;
+
 	/*
 	 * Compile Context used to resolve dynamic columns when FConditions are compiled
 	 */
diff --git a/Engine/Source/Runtime/TypedElementFramework/Public/Elements/Interfaces/TypedElementDataStorageInterface.h b/Engine/Source/Runtime/TypedElementFramework/Public/Elements/Interfaces/TypedElementDataStorageInterface.h
--- a/Engine/Source/Runtime/TypedElementFramework/Public/Elements/Interfaces/TypedElementDataStorageInterface.h
+++ b/Engine/Source/Runtime/TypedElementFramework/Public/Elements/Interfaces/TypedElementDataStorageInterface.h
@@ -332,6 +332,8 @@ public:
 	virtual bool HasColumns(RowHandle Row, TConstArrayView<TWeakObjectPtr<const UScriptStruct>> ColumnTypes) const = 0;
 	template<TColumnType... ColumnTypes>
 	bool HasColumns(RowHandle Row) const;
+	template<Queries::TColumnDescType... ColumnTypes>
+	bool HasColumns(RowHandle Row, const ColumnTypes&... Columns) const;
 
 	/** 
 	 * Lists the columns on a row. This includes data and tag columns.
@@ -848,6 +850,39 @@ bool ICoreProvider::HasColumns(RowHandle Row) const
 	return HasColumns(Row, TConstArrayView<const UScriptStruct*>({ ColumnType::StaticStruct()... }));
 }
 
+template<Queries::TColumnDescType... ColumnTypes>
+bool ICoreProvider::HasColumns(RowHandle Row, const ColumnTypes&... Columns) const
+{
+	TArray<const UScriptStruct*, TInlineAllocator<sizeof...(ColumnTypes)>> TypeInfo;
+	TypeInfo.Reserve(sizeof...(ColumnTypes));
+
+	auto GetType = [this, &TypeInfo]<Queries::TColumnDescType ColumnType>(const ColumnType & Column) -> bool
+		{
+			const UScriptStruct* Result = Column.Identifier.IsNone()
+				? Column.TypeInfo.Get()
+				: FindDynamicColumn(FDynamicColumnDescription
+					{
+						.TemplateType = Column.TypeInfo.Get(),
+						.Identifier = Column.Identifier
+					});
+			if (Result)
+			{
+				TypeInfo.Add(Result);
+				return true;
+			}
+			else
+			{
+				return false;
+			}
+		};
+
+	if (!(GetType(Columns) && ...))
+	{
+		return false;
+	}
+	return HasColumns(Row, TConstArrayView<const UScriptStruct*>(TypeInfo));
+}
+
 template<typename SystemType>
 SystemType* ICoreProvider::GetExternalSystem()
 {
