[TEDS] Addressed an issue with the new templated dynamic columns.

Using static in the literal itself doesn't work reliably because the template is specialized on just the size in the final submitted version. This would mean that two strings with the same string would return the first string to be called. Instead the storage has been moved to where it's used. A test was added to make sure this problem can't accidently be reintroduced.

diff --git a/Engine/Plugins/Experimental/EditorDataStorage/Source/TedsCore/Tests/DynamicColumnTests.cpp b/Engine/Plugins/Experimental/EditorDataStorage/Source/TedsCore/Tests/DynamicColumnTests.cpp
--- a/Engine/Plugins/Experimental/EditorDataStorage/Source/TedsCore/Tests/DynamicColumnTests.cpp
+++ b/Engine/Plugins/Experimental/EditorDataStorage/Source/TedsCore/Tests/DynamicColumnTests.cpp
@@ -112,6 +112,16 @@ namespace UE::Editor::DataStorage::Tests
 					TestTrue("Expected columns not found", TedsInterface->HasColumns(Rows[2], MakeConstArrayView({DynamicTagTypes[1]})));
 				}
 				
+				// Check if two dynamic columns with the same string length are unique
+				{
+					using namespace UE::Editor::DataStorage::Queries;
+					{
+						FName FirstName = TColumn<FTestDynamicTag, "Truck">().Identifier;
+						FName SecondName = TColumn<FTestDynamicTag, "Hello">().Identifier;
+						TestNotEqual(TEXT("Different id with the same length should produce two different ids."), FirstName, SecondName);
+					}
+				}
+
 				// Direct Query
 				{
 					{
diff --git a/Engine/Source/Runtime/TypedElementFramework/Public/DataStorage/CommonTypes.h b/Engine/Source/Runtime/TypedElementFramework/Public/DataStorage/CommonTypes.h
--- a/Engine/Source/Runtime/TypedElementFramework/Public/DataStorage/CommonTypes.h
+++ b/Engine/Source/Runtime/TypedElementFramework/Public/DataStorage/CommonTypes.h
@@ -117,10 +117,9 @@ namespace UE
 				std::copy_n(InName, N - 1, Name);
 			}
 
-			FName GetStaticName() const
+			FName GetName() const
 			{
-				static FName Result = N > 1 ? FName(N - 1, Name) : NAME_None;
-				return Result;
+				return N > 1 ? FName(N - 1, Name) : NAME_None;
 			}
 		};
 	} // namespace Editor::DataStorage
diff --git a/Engine/Source/Runtime/TypedElementFramework/Public/DataStorage/Queries/Conditions.h b/Engine/Source/Runtime/TypedElementFramework/Public/DataStorage/Queries/Conditions.h
--- a/Engine/Source/Runtime/TypedElementFramework/Public/DataStorage/Queries/Conditions.h
+++ b/Engine/Source/Runtime/TypedElementFramework/Public/DataStorage/Queries/Conditions.h
@@ -38,7 +38,11 @@ namespace Queries
 	template<typename T, TVariantName Id> requires (decltype(Id)::Size > 0)
 	struct TColumn<T, Id> final : public FColumnBase
 	{
-		explicit TColumn() : FColumnBase(T::StaticStruct(), Id.GetStaticName()) {}
+		explicit TColumn() : FColumnBase(T::StaticStruct(), []()
+			{
+				static FName Name = Id.GetName();
+				return Name;
+			}()) {}
 	};
 
 	template<>
