Slate Dynamic Invalidation - Scene Outliner Widgets

diff --git a/Engine/Plugins/Experimental/EditorDataStorage/Source/TedsCore/Private/Widgets/STedsWidget.h b/Engine/Plugins/Experimental/EditorDataStorage/Source/TedsCore/Private/Widgets/STedsWidget.h
--- a/Engine/Plugins/Experimental/EditorDataStorage/Source/TedsCore/Private/Widgets/STedsWidget.h
+++ b/Engine/Plugins/Experimental/EditorDataStorage/Source/TedsCore/Private/Widgets/STedsWidget.h
@@ -11,8 +11,20 @@ struct FTypedElementWidgetConstructor;
 namespace UE::Editor::DataStorage
 {
 	class IUiProvider;
+	namespace Widgets
+	{
+		class STedsWidget;
+	} // namespace UE::Editor::DataStorage::Widgets
 } // namespace UE::Editor::DataStorage
 
+
+template<>
+struct TWidgetTypeTraits<UE::Editor::DataStorage::Widgets::STedsWidget>
+{
+	static constexpr bool SupportsInvalidation() { return true; }
+};
+
+
 namespace UE::Editor::DataStorage::Widgets
 {
 	/*
diff --git a/Engine/Plugins/Experimental/EditorDataStorageFeatures/Source/TedsOutliner/Public/Widgets/OutlinerVisibilityWidget.h b/Engine/Plugins/Experimental/EditorDataStorageFeatures/Source/TedsOutliner/Public/Widgets/OutlinerVisibilityWidget.h
--- a/Engine/Plugins/Experimental/EditorDataStorageFeatures/Source/TedsOutliner/Public/Widgets/OutlinerVisibilityWidget.h
+++ b/Engine/Plugins/Experimental/EditorDataStorageFeatures/Source/TedsOutliner/Public/Widgets/OutlinerVisibilityWidget.h
@@ -77,6 +77,12 @@ public:
 		const UE::Editor::DataStorage::FMetaDataView& Arguments) override;
 };
 
+template<>
+struct TWidgetTypeTraits<class STedsVisibilityWidget>
+{
+	static constexpr bool SupportsInvalidation() { return true; }
+};
+
 /** Widget responsible for managing the visibility for a single item */
 class STedsVisibilityWidget : public SImage
 {
diff --git a/Engine/Source/Editor/SceneOutliner/Private/ActorDescTreeItem.cpp b/Engine/Source/Editor/SceneOutliner/Private/ActorDescTreeItem.cpp
--- a/Engine/Source/Editor/SceneOutliner/Private/ActorDescTreeItem.cpp
+++ b/Engine/Source/Editor/SceneOutliner/Private/ActorDescTreeItem.cpp
@@ -31,6 +31,12 @@
 
 const FSceneOutlinerTreeItemType FActorDescTreeItem::Type(&IActorBaseTreeItem::Type);
 
+template<>
+struct TWidgetTypeTraits<struct SActorDescTreeLabel>
+{
+	static constexpr bool SupportsInvalidation() { return true; }
+};
+
 struct SActorDescTreeLabel : FSceneOutlinerCommonLabelData, public SCompoundWidget
 {
 	SLATE_BEGIN_ARGS(SActorDescTreeLabel) {}
diff --git a/Engine/Source/Editor/SceneOutliner/Private/LevelTreeItem.cpp b/Engine/Source/Editor/SceneOutliner/Private/LevelTreeItem.cpp
--- a/Engine/Source/Editor/SceneOutliner/Private/LevelTreeItem.cpp
+++ b/Engine/Source/Editor/SceneOutliner/Private/LevelTreeItem.cpp
@@ -16,6 +16,12 @@
 
 const FSceneOutlinerTreeItemType FLevelTreeItem::Type(&ISceneOutlinerTreeItem::Type);
 
+template<>
+struct TWidgetTypeTraits<struct SLevelTreeLabel>
+{
+	static constexpr bool SupportsInvalidation() { return true; }
+};
+
 struct SLevelTreeLabel : FSceneOutlinerCommonLabelData, public SCompoundWidget
 {
 	SLATE_BEGIN_ARGS(SLevelTreeLabel) {}
diff --git a/Engine/Source/Editor/SceneOutliner/Private/SSourceControlWidget.h b/Engine/Source/Editor/SceneOutliner/Private/SSourceControlWidget.h
--- a/Engine/Source/Editor/SceneOutliner/Private/SSourceControlWidget.h
+++ b/Engine/Source/Editor/SceneOutliner/Private/SSourceControlWidget.h
@@ -6,6 +6,12 @@
 #include "SourceControlHelpers.h"
 #include "SceneOutlinerTreeItemSCC.h"
 
+template<>
+struct TWidgetTypeTraits<class SSourceControlWidget>
+{
+	static constexpr bool SupportsInvalidation() { return true; }
+};
+
 class SSourceControlWidget : public SLayeredImage
 {
 public:
diff --git a/Engine/Source/Editor/SceneOutliner/Private/SceneOutlinerPinnedColumn.cpp b/Engine/Source/Editor/SceneOutliner/Private/SceneOutlinerPinnedColumn.cpp
--- a/Engine/Source/Editor/SceneOutliner/Private/SceneOutlinerPinnedColumn.cpp
+++ b/Engine/Source/Editor/SceneOutliner/Private/SceneOutlinerPinnedColumn.cpp
@@ -52,6 +52,12 @@ void FSceneOutlinerPinnedColumn::FSceneOutlinerPinnedStateCache::Empty()
 	PinnedStateInfo.Empty();
 }
 
+template<>
+struct TWidgetTypeTraits<class SPinnedWidget>
+{
+	static constexpr bool SupportsInvalidation() { return true; }
+};
+
 class SPinnedWidget : public SImage
 {
 public:
