[TEDS] Added support for dynamic columns to the new query context.

This review contains the first new Query Context capability that can be shared, so this also includes updates to support this. The Direct API, current Query Context and new Query Context (currently only used for filtering) now all use the same implementation and interface creating greater consistency moving forward.
New features added are the ability to retrieve the dynamic column type from the new Query Context, several new functions to make it easier to work with dynamic columns and the ability to iterate over all variants of a dynamic columns from both the current and new Query Context.

diff --git a/Engine/Plugins/Experimental/EditorDataStorage/Source/TedsCore/Private/Compatibility/TypedElementObjectReinstancingManager.cpp b/Engine/Plugins/Experimental/EditorDataStorage/Source/TedsCore/Private/Compatibility/TypedElementObjectReinstancingManager.cpp
--- a/Engine/Plugins/Experimental/EditorDataStorage/Source/TedsCore/Private/Compatibility/TypedElementObjectReinstancingManager.cpp
+++ b/Engine/Plugins/Experimental/EditorDataStorage/Source/TedsCore/Private/Compatibility/TypedElementObjectReinstancingManager.cpp
@@ -50,7 +50,7 @@ void UTedsObjectReinstancingManager::UpdateCompleted()
 {
 	using namespace UE::Editor::DataStorage;
 
-	FMementoSystem& MementoSystem = DataStorage->GetEnvironment()->GetMementoSystem();
+	FMementoSystem& MementoSystem = DataStorage->GetEnvironment().GetMementoSystem();
 	for (TMap<const void*, RowHandle>::TConstIterator It = OldObjectToMementoMap.CreateConstIterator(); It; ++It)
 	{
 		MementoSystem.DestroyMemento(It.Value());
@@ -64,7 +64,7 @@ void UTedsObjectReinstancingManager::HandleOnObjectPreRemoved(
 	UE::Editor::DataStorage::RowHandle ObjectRow)
 {
 	// This is the chance to record the old object to memento
-	UE::Editor::DataStorage::RowHandle Memento = DataStorage->GetEnvironment()->GetMementoSystem().CreateMemento(ObjectRow);
+	UE::Editor::DataStorage::RowHandle Memento = DataStorage->GetEnvironment().GetMementoSystem().CreateMemento(ObjectRow);
 	OldObjectToMementoMap.Add(Object, Memento);
 }
 
@@ -96,7 +96,7 @@ void UTedsObjectReinstancingManager::HandleOnObjectsReinstanced(
 			}
 
 			// Kick off re-instantiation of NewObjectRow from the Memento
-			DataStorage->GetEnvironment()->GetMementoSystem().RestoreMemento(Memento, NewObjectRow);
+			DataStorage->GetEnvironment().GetMementoSystem().RestoreMemento(Memento, NewObjectRow);
 		}
 	}
 }
diff --git a/Engine/Plugins/Experimental/EditorDataStorage/Source/TedsCore/Private/Processors/TypedElementProcessorAdaptors.cpp b/Engine/Plugins/Experimental/EditorDataStorage/Source/TedsCore/Private/Processors/TypedElementProcessorAdaptors.cpp
--- a/Engine/Plugins/Experimental/EditorDataStorage/Source/TedsCore/Private/Processors/TypedElementProcessorAdaptors.cpp
+++ b/Engine/Plugins/Experimental/EditorDataStorage/Source/TedsCore/Private/Processors/TypedElementProcessorAdaptors.cpp
@@ -8,10 +8,12 @@
 #include "MassExecutionContext.h"
 #include "Queries/TypedElementExtendedQueryStore.h"
 #include "Queries/TypedElementQueryContext_MassForwarder.h"
+#include "Processors/TypedElementQueryCapabilityImplemenationForwarder.h"
 #include "StructUtils/StructTypeBitSet.h"
 #include "TypedElementDatabase.h"
 #include "TypedElementDatabaseEnvironment.h"
 #include "TypedElementUtils.h"
+#include "Queries/QueryContextCapabilities/EditorDataStorageDynamicColumnInfo.h"
 
 #include UE_INLINE_GENERATED_CPP_BY_NAME(TypedElementProcessorAdaptors)
 
@@ -115,12 +117,6 @@ namespace UE::Editor::DataStorage
 				return false;
 			}
 
-			const UScriptStruct* FindDynamicColumnType(const UE::Editor::DataStorage::FDynamicColumnDescription& Description) const
-			{
-				const UScriptStruct* DynamicColumnType = Environment.FindDynamicColumn(*Description.TemplateType, Description.Identifier);
-				return DynamicColumnType;
-			}
-
 			float GetDeltaTimeSeconds() const
 			{
 				return Context.GetDeltaTimeSeconds();
@@ -1112,7 +1108,7 @@ namespace UE::Editor::DataStorage
 			const FTedsHierarchyAccessInterface* HierarchyAccessInterface = nullptr;
 		};
 
-		struct FMassDirectContextForwarder final : public IDirectQueryContext
+		struct FMassDirectContextForwarder final : Queries::Private::TImplementationForwarder<IDirectQueryContext, Queries::FDynamicColumnInfoContextCapability>
 		{
 			FMassDirectContextForwarder(FMassExecutionContext& InContext, FEnvironment& InEnvironment, FQueryDescription& Description)
 				: Implementation(InContext, InEnvironment)
@@ -1129,7 +1125,6 @@ namespace UE::Editor::DataStorage
 			void GetColumnsUnguarded(int32 TypeCount, char** RetrievedAddresses, const TWeakObjectPtr<const UScriptStruct>* ColumnTypes, const EQueryAccessType* AccessTypes) override { return Implementation.GetColumnsUnguarded(TypeCount, RetrievedAddresses, ColumnTypes, AccessTypes); }
 			bool HasColumn(const UScriptStruct* ColumnType) const override { return Implementation.HasColumn(ColumnType); }
 			bool HasColumn(RowHandle Row, const UScriptStruct* ColumnType) const override { return Implementation.HasColumn(Row, ColumnType); }
-			const UScriptStruct* FindDynamicColumnType(const UE::Editor::DataStorage::FDynamicColumnDescription& Description) const override { return Implementation.FindDynamicColumnType(Description); }
 			float GetDeltaTimeSeconds() const override { return Implementation.GetDeltaTimeSeconds(); }
 			virtual void SetParentRow(RowHandle Target, RowHandle Parent) override { if (HierarchyAccessInterface) { HierarchyAccessInterface->SetParentRow(*static_cast<IDirectQueryContext*>(this), Target, Parent); } }
 			virtual void SetUnresolvedParent(RowHandle Target, FMapKey ParentId, FName MappingDomain) override
@@ -1145,13 +1140,15 @@ namespace UE::Editor::DataStorage
 		protected:
 			void* EmplaceObjectInScratch(const FEmplaceObjectParams& Params) override { return Implementation.EmplaceObjectInScratch(Params.ObjectSize, Params.Alignment, Params.Construct, Params.Destroy, Params.SourceObject); }
 
+			virtual FEnvironment& GetEnvironment() override { return Implementation.Environment; }
+			virtual const FEnvironment& GetEnvironment() const override { return Implementation.Environment; }
 		public:
 
 			FMassContextCommon Implementation;
 			const FTedsHierarchyAccessInterface* HierarchyAccessInterface = nullptr;
 		};
 
-		struct FMassSubqueryContextForwarder  : public ISubqueryContext
+		struct FMassSubqueryContextForwarder : Queries::Private::TImplementationForwarder<ISubqueryContext, Queries::FDynamicColumnInfoContextCapability>
 		{
 			FMassSubqueryContextForwarder(FMassExecutionContext& InContext, FEnvironment& InEnvironment, FQueryDescription& Description)
 				: Implementation(InContext, InEnvironment, Description.bShouldBatchModifications)
@@ -1189,7 +1186,6 @@ namespace UE::Editor::DataStorage
 			void* AddColumnUninitialized(RowHandle Row, const UE::Editor::DataStorage::FDynamicColumnDescription& DynamicColumnDescription, ObjectCopyOrMove Relocator) override { return Implementation.AddColumnUninitialized(Row, DynamicColumnDescription, Relocator); }
 			void RemoveColumns(RowHandle Row, TConstArrayView<const UScriptStruct*> ColumnTypes) override { Implementation.RemoveColumns(Row, ColumnTypes); }
 			void RemoveColumns(TConstArrayView<RowHandle> Rows, TConstArrayView<const UScriptStruct*> ColumnTypes) override { Implementation.RemoveColumns(Rows, ColumnTypes); }
-			const UScriptStruct* FindDynamicColumnType(const UE::Editor::DataStorage::FDynamicColumnDescription& Description) const override { return Implementation.FindDynamicColumnType(Description); }
 			float GetDeltaTimeSeconds() const override { return Implementation.GetDeltaTimeSeconds(); }
 			virtual void SetParentRow(RowHandle Target, RowHandle Parent) override { Implementation.SetParent(*static_cast<ISubqueryContext*>(this), Target, Parent); }
 			virtual void SetUnresolvedParent(RowHandle Target, FMapKey ParentId, FName MappingDomain) override { Implementation.SetUnresolvedParent(*static_cast<ISubqueryContext*>(this), Target, ParentId, MappingDomain); }
@@ -1198,6 +1194,9 @@ namespace UE::Editor::DataStorage
 		protected:
 			void* EmplaceObjectInScratch(const FEmplaceObjectParams& Params) override { return Implementation.EmplaceObjectInScratch(Params.ObjectSize, Params.Alignment, Params.Construct, Params.Destroy, Params.SourceObject); }
 			
+			virtual FEnvironment& GetEnvironment() override { return Implementation.Environment; }
+			virtual const FEnvironment& GetEnvironment() const override { return Implementation.Environment; }
+
 			FMassWithEnvironmentContextCommon Implementation;
 		};
 
@@ -1308,7 +1307,7 @@ namespace UE::Editor::DataStorage
 			FExtendedQueryStore& QueryStore;
 		};
 
-		struct FMassContextForwarder final : public IQueryContext
+		struct FMassContextForwarder final : Queries::Private::TImplementationForwarder<IQueryContext, Queries::FDynamicColumnInfoContextCapability>
 		{
 			FMassContextForwarder(
 				FQueryDescription& InQueryDescription,
@@ -1343,7 +1342,6 @@ namespace UE::Editor::DataStorage
 			void* AddColumnUninitialized(RowHandle Row, const FDynamicColumnDescription& DynamicColumnDescription) override { return Implementation.AddColumnUninitialized(Row, DynamicColumnDescription); };
 			void RemoveColumns(RowHandle Row, TConstArrayView<const UScriptStruct*> ColumnTypes) override { Implementation.RemoveColumns(Row, ColumnTypes); }
 			void RemoveColumns(TConstArrayView<RowHandle> Rows, TConstArrayView<const UScriptStruct*> ColumnTypes) override { Implementation.RemoveColumns(Rows, ColumnTypes); }
-			const UScriptStruct* FindDynamicColumnType(const UE::Editor::DataStorage::FDynamicColumnDescription& Description) const override { return Implementation.FindDynamicColumnType(Description); }
 			float GetDeltaTimeSeconds() const override { return Implementation.GetDeltaTimeSeconds(); }
 			virtual void SetParentRow(RowHandle Target, RowHandle Parent) override { Implementation.SetParent(*static_cast<IQueryContext*>(this), Target, Parent); }
 			virtual void SetUnresolvedParent(RowHandle Target, FMapKey ParentId, FName MappingDomain) override { Implementation.SetUnresolvedParent(*static_cast<IQueryContext*>(this), Target, ParentId, MappingDomain); }
@@ -1358,8 +1356,11 @@ namespace UE::Editor::DataStorage
 			FQueryResult RunSubquery(int32 SubqueryIndex, UE::Editor::DataStorage::SubqueryCallbackRef Callback) override { return Implementation.RunSubquery(SubqueryIndex, Callback); }
 			FQueryResult RunSubquery(int32 SubqueryIndex, RowHandle Row, UE::Editor::DataStorage::SubqueryCallbackRef Callback) override { return Implementation.RunSubquery(SubqueryIndex, Row, Callback); }
 
-protected:
+		protected:
 			void* EmplaceObjectInScratch(const FEmplaceObjectParams& Params) override { return Implementation.EmplaceObjectInScratch(Params.ObjectSize, Params.Alignment, Params.Construct, Params.Destroy, Params.SourceObject); }
+
+			virtual FEnvironment& GetEnvironment() override { return Implementation.Environment; }
+			virtual const FEnvironment& GetEnvironment() const override { return Implementation.Environment; }
 		
 		protected:
 			FMassQueryContextImplementation Implementation;
@@ -1654,11 +1655,11 @@ UE::Editor::DataStorage::FQueryResult FTypedElementQueryProcessorData::Execute(
 	if (EnumHasAnyFlags(ExecutionFlags, EDirectQueryExecutionFlags::AllowBoundQueries) || !Description.Callback.Function)
 	{
 		FMassExecutionContext Context(EntityManager);
-		auto ExecuteFunction = [&Result, &Callback, &Description, &Environment](FMassExecutionContext& Context)
+		auto ExecuteFunction = [&Result, &Callback, &Environment](FMassExecutionContext& Context)
 			{
 				using namespace Processors::Private;
 				// No need to cache any subsystem dependencies as these are not accessible from a direct query.
-				QueryContext_MassForwarder QueryContext(Context);
+				QueryContext_MassForwarder QueryContext(Environment, Context);
 				// This assumes that there already has been a check that the provided function is compatible with the query.
 				Callback.Call(QueryContext, QueryContext.GetContextImplementation());
 				Result.Count += Context.GetNumEntities();
diff --git a/Engine/Plugins/Experimental/EditorDataStorage/Source/TedsCore/Private/Processors/TypedElementQueryCapabilityImplemenationForwarder.h b/Engine/Plugins/Experimental/EditorDataStorage/Source/TedsCore/Private/Processors/TypedElementQueryCapabilityImplemenationForwarder.h
new file mode 100644
--- /dev/null
+++ b/Engine/Plugins/Experimental/EditorDataStorage/Source/TedsCore/Private/Processors/TypedElementQueryCapabilityImplemenationForwarder.h
@@ -0,0 +1,77 @@
+// Copyright Epic Games, Inc. All Rights Reserved.
+
+#pragma once
+
+#include "DataStorage/CommonTypes.h"
+#include "Elements/Framework/TypedElementQueryCapabilities.h"
+#include "Elements/Framework/TypedElementQueryContract.h"
+#include "Elements/Framework/TypedElementQueryFunctionArgumentTypes.h"
+
+namespace UE::Editor::DataStorage::Queries::Private
+{
+	// Because the original context has a different inheritance hierarchy than the new query context implementations can't
+	// be directly used in the various IQueryContext implementation. This wrapper can be used to bridge the two until
+	// the new query context can fully replace the original query context.
+
+	template<typename Interface, typename... Implementations> struct TImplementationForwarder
+	{
+		static_assert(sizeof...(Implementations) < 0,
+			"Specialization for query context forwards missing. The most likely reason is the addition of a new query "
+			"context capability that wasn't specialized.");
+	};
+
+	template<typename Interface>
+	struct TImplementationForwarder<Interface> : Interface
+	{
+	};
+
+#define ArgTypeName(Type, Name) Type Name
+#define ArgName(Type, Name) Name
+
+#define Function0(Capability, Return, Function) virtual Return Function() \
+	{ return Implementation:: Function (); }
+#define Function1(Capability, Return, Function, Arg1) virtual Return Function(ArgTypeName Arg1 ) override \
+	{ return Implementation:: Function (ArgName Arg1 ); }
+#define Function2(Capability, Return, Function, Arg1, Arg2) virtual Return Function(ArgTypeName Arg1 , ArgTypeName Arg2 ) override \
+	{ return Implementation:: Function (ArgName Arg1 , ArgName Arg2 ); }
+#define Function3(Capability, Return, Function, Arg1, Arg2, Arg3) virtual Return Function(ArgTypeName Arg1 , ArgTypeName Arg2 , ArgTypeName Arg3 ) override \
+	{ return Implementation:: Function (ArgName Arg1 , ArgName Arg2 , ArgName Arg3); }
+#define Function4(Capability, Return, Function, Arg1, Arg2, Arg3, Arg4) virtual Return Function(ArgTypeName Arg1 , ArgTypeName Arg2 , ArgTypeName Arg3 , ArgTypeName Arg4 ) override \
+	{ return Implementation:: Function (ArgName Arg1 , ArgName Arg2 , ArgName Arg3 , ArgName Arg4); }
+
+#define ConstFunction0(Capability, Return, Function) virtual Return Function() const \
+	{ return Implementation:: Function (); }
+#define ConstFunction1(Capability, Return, Function, Arg1) virtual Return Function(ArgTypeName Arg1 ) const override \
+	{ return Implementation:: Function (ArgName Arg1 ); }
+#define ConstFunction2(Capability, Return, Function, Arg1, Arg2) virtual Return Function(ArgTypeName Arg1 , ArgTypeName Arg2 ) const override \
+	{ return Implementation:: Function (ArgName Arg1 , ArgName Arg2 ); }
+#define ConstFunction3(Capability, Return, Function, Arg1, Arg2, Arg3) virtual Return Function(ArgTypeName Arg1 , ArgTypeName Arg2 , ArgTypeName Arg3 ) const override \
+	{ return Implementation:: Function (ArgName Arg1 , ArgName Arg2 , ArgName Arg3); }
+#define ConstFunction4(Capability, Return, Function, Arg1, Arg2, Arg3, Arg4) virtual Return Function(ArgTypeName Arg1 , ArgTypeName Arg2 , ArgTypeName Arg3 , ArgTypeName Arg4 ) const override \
+	{ return Implementation:: Function (ArgName Arg1 , ArgName Arg2 , ArgName Arg3 , ArgName Arg4); }
+
+#define CapabilityStart(Capability, Flags) \
+	template<typename Interface, typename Implementation, typename... Implementations> requires std::is_base_of_v< Queries:: Capability , Implementation> \
+	struct TImplementationForwarder< Interface, Implementation, Implementations...> : Implementation, TImplementationForwarder<Interface, Implementations...> \
+	{ 
+
+#define CapabilityEnd(Capability) \
+	}; \
+	
+#include "Elements/Framework/TypedElementQueryCapabilities.inl"
+
+#undef ArgTypeName
+#undef ArgName
+#undef Function0
+#undef Function1
+#undef Function2
+#undef Function3
+#undef Function4
+#undef ConstFunction0
+#undef ConstFunction1
+#undef ConstFunction2
+#undef ConstFunction3
+#undef ConstFunction4
+#undef CapabilityStart
+#undef CapabilityEnd
+} // namespace UE::Editor::DataStorage::Queries::Private
\ No newline at end of file
diff --git a/Engine/Plugins/Experimental/EditorDataStorage/Source/TedsCore/Private/Queries/QueryContextCapabilities/EditorDataStorageDynamicColumnInfo.cpp b/Engine/Plugins/Experimental/EditorDataStorage/Source/TedsCore/Private/Queries/QueryContextCapabilities/EditorDataStorageDynamicColumnInfo.cpp
new file mode 100644
--- /dev/null
+++ b/Engine/Plugins/Experimental/EditorDataStorage/Source/TedsCore/Private/Queries/QueryContextCapabilities/EditorDataStorageDynamicColumnInfo.cpp
@@ -0,0 +1,18 @@
+// Copyright Epic Games, Inc. All Rights Reserved.
+
+#include "Queries/QueryContextCapabilities/EditorDataStorageDynamicColumnInfo.h"
+
+namespace UE::Editor::DataStorage::Queries
+{
+	const UScriptStruct* FDynamicColumnInfoContextCapability::FindDynamicColumnType(
+		const UScriptStruct& TemplateType, const FName& Identifier) const
+	{
+		return GetEnvironment().FindDynamicColumn(TemplateType, Identifier);
+	}
+
+	void FDynamicColumnInfoContextCapability::ForEachDynamicColumnType(
+		const UScriptStruct& Template, TFunctionRef<void(const UScriptStruct& Type)> Callback) const
+	{
+		GetEnvironment().ForEachDynamicColumn(Template, Callback);
+	}
+} // namespace UE::Editor::DataStorage::Queries
diff --git a/Engine/Plugins/Experimental/EditorDataStorage/Source/TedsCore/Private/Queries/QueryContextCapabilities/EditorDataStorageDynamicColumnInfo.h b/Engine/Plugins/Experimental/EditorDataStorage/Source/TedsCore/Private/Queries/QueryContextCapabilities/EditorDataStorageDynamicColumnInfo.h
new file mode 100644
--- /dev/null
+++ b/Engine/Plugins/Experimental/EditorDataStorage/Source/TedsCore/Private/Queries/QueryContextCapabilities/EditorDataStorageDynamicColumnInfo.h
@@ -0,0 +1,28 @@
+// Copyright Epic Games, Inc. All Rights Reserved.
+
+#pragma once
+
+#include "Elements/Framework/TypedElementQueryContextImplementation.h"
+#include "TypedElementDatabaseEnvironment.h"
+
+class FName;
+class UScriptStruct;
+
+namespace UE::Editor::DataStorage::Queries
+{
+	/**
+	 * Implementation for the DynamicColumnInfo context capability.
+	 * This provides read access to the dynamic columns in order to retrieve information about them from any thread.
+	 * This class provides an implementation block and is not meant to be directly used but should be inherited from
+	 * by a class that's compatible with TQueryContextImpl.
+	 */
+	class FDynamicColumnInfoContextCapability : public DynamicColumnInfo, public IQueryEnvironmentRetrieval
+	{
+	public:
+		virtual ~FDynamicColumnInfoContextCapability() override = default;
+
+		virtual const UScriptStruct* FindDynamicColumnType(const UScriptStruct& TemplateType, const FName& Identifier) const override;
+		virtual void ForEachDynamicColumnType(
+			const UScriptStruct& Template, TFunctionRef<void(const UScriptStruct& Type)> Callback) const override;
+	};
+} // namespace UE::Editor::DataStorage::Queries
diff --git a/Engine/Plugins/Experimental/EditorDataStorage/Source/TedsCore/Private/Queries/TypedElementQueryContext_DirectSingleApi.cpp b/Engine/Plugins/Experimental/EditorDataStorage/Source/TedsCore/Private/Queries/TypedElementQueryContext_DirectSingleApi.cpp
--- a/Engine/Plugins/Experimental/EditorDataStorage/Source/TedsCore/Private/Queries/TypedElementQueryContext_DirectSingleApi.cpp
+++ b/Engine/Plugins/Experimental/EditorDataStorage/Source/TedsCore/Private/Queries/TypedElementQueryContext_DirectSingleApi.cpp
@@ -7,8 +7,9 @@ namespace UE::Editor::DataStorage::Queries
 {
 	const FName FQueryContext_DirectSingleApi::Name = "FQueryContext_DirectSingleApi";
 
-	FQueryContext_DirectSingleApi::FQueryContext_DirectSingleApi(UEditorDataStorage& DataStorage)
-		: DataStorage(DataStorage)
+	FQueryContext_DirectSingleApi::FQueryContext_DirectSingleApi(FEnvironment& Environment, UEditorDataStorage& DataStorage)
+		: Environment(Environment)
+		, DataStorage(DataStorage)
 	{
 	}
 
@@ -71,4 +72,13 @@ namespace UE::Editor::DataStorage::Queries
 		return CurrentRowHasColumns(Columns);
 	}
 
+	FEnvironment& FQueryContext_DirectSingleApi::GetEnvironment()
+	{
+		return Environment;
+	}
+
+	const FEnvironment& FQueryContext_DirectSingleApi::GetEnvironment() const
+	{
+		return Environment;
+	}
 } // namespace UE::Editor::DataStorage::Queries
\ No newline at end of file
diff --git a/Engine/Plugins/Experimental/EditorDataStorage/Source/TedsCore/Private/Queries/TypedElementQueryContext_DirectSingleApi.h b/Engine/Plugins/Experimental/EditorDataStorage/Source/TedsCore/Private/Queries/TypedElementQueryContext_DirectSingleApi.h
--- a/Engine/Plugins/Experimental/EditorDataStorage/Source/TedsCore/Private/Queries/TypedElementQueryContext_DirectSingleApi.h
+++ b/Engine/Plugins/Experimental/EditorDataStorage/Source/TedsCore/Private/Queries/TypedElementQueryContext_DirectSingleApi.h
@@ -4,6 +4,8 @@
 
 #include "Elements/Framework/TypedElementQueryContextImplementation.h"
 #include "Elements/Framework/TypedElementQueryFunctions.h"
+#include "TypedElementDatabaseEnvironment.h"
+#include "Queries/QueryContextCapabilities/EditorDataStorageDynamicColumnInfo.h"
 
 class UEditorDataStorage;
 
@@ -15,12 +17,16 @@ namespace UE::Editor::DataStorage::Queries
 	 * will in practice only run a single row. Batch mode was added for compatibility and offers
 	 * no performance benefits when using this context.
 	 */
-	class FQueryContext_DirectSingleApi final : public SingleRowInfo, public RowBatchInfo, public IQueryFunctionResponse
+	class FQueryContext_DirectSingleApi final : 
+		public SingleRowInfo, 
+		public RowBatchInfo, 
+		public FDynamicColumnInfoContextCapability,
+		public IQueryFunctionResponse
 	{
 	public:
 		static const FName Name;
 
-		explicit FQueryContext_DirectSingleApi(UEditorDataStorage& DataStorage);
+		FQueryContext_DirectSingleApi(FEnvironment& Environment, UEditorDataStorage& DataStorage);
 
 		/**
 		 * @section ISingleRowInfo
@@ -46,7 +52,14 @@ namespace UE::Editor::DataStorage::Queries
 
 		void SetCurrentRow(RowHandle Row);
 
+		/**
+		 * @section IQueryEnvironmentRetrieval
+		 */
+		virtual FEnvironment& GetEnvironment() override;
+		virtual const FEnvironment& GetEnvironment() const override;
+
 	private:
+		FEnvironment& Environment;
 		UEditorDataStorage& DataStorage;
 		RowHandle CurrentRow = InvalidRowHandle;
 	};
diff --git a/Engine/Plugins/Experimental/EditorDataStorage/Source/TedsCore/Private/Queries/TypedElementQueryContext_MassForwarder.cpp b/Engine/Plugins/Experimental/EditorDataStorage/Source/TedsCore/Private/Queries/TypedElementQueryContext_MassForwarder.cpp
--- a/Engine/Plugins/Experimental/EditorDataStorage/Source/TedsCore/Private/Queries/TypedElementQueryContext_MassForwarder.cpp
+++ b/Engine/Plugins/Experimental/EditorDataStorage/Source/TedsCore/Private/Queries/TypedElementQueryContext_MassForwarder.cpp
@@ -8,8 +8,9 @@ namespace UE::Editor::DataStorage::Queries
 {
 	const FName FQueryContext_MassForwarder::Name = "FQueryContext_MassForwarder";
 
-	FQueryContext_MassForwarder::FQueryContext_MassForwarder(FMassExecutionContext& Context)
-		: Context(Context)
+	FQueryContext_MassForwarder::FQueryContext_MassForwarder(FEnvironment& Environment, FMassExecutionContext& Context)
+		: Environment(Environment)
+		, Context(Context)
 	{
 		TConstArrayView<RowHandle> Rows = MassEntitiesToRowsConversion(Context.GetEntities());
 		CurrentRow = Rows.GetData();
@@ -103,4 +104,14 @@ namespace UE::Editor::DataStorage::Queries
 			*ColumnsDataIt++ = Context.GetMutableFragmentView(ColumnType).GetData();
 		}
 	}
+
+	FEnvironment& FQueryContext_MassForwarder::GetEnvironment()
+	{
+		return Environment;
+	}
+
+	const FEnvironment& FQueryContext_MassForwarder::GetEnvironment() const
+	{
+		return Environment;
+	}
 } // namespace UE::Editor::DataStorage::Queries
\ No newline at end of file
diff --git a/Engine/Plugins/Experimental/EditorDataStorage/Source/TedsCore/Private/Queries/TypedElementQueryContext_MassForwarder.h b/Engine/Plugins/Experimental/EditorDataStorage/Source/TedsCore/Private/Queries/TypedElementQueryContext_MassForwarder.h
--- a/Engine/Plugins/Experimental/EditorDataStorage/Source/TedsCore/Private/Queries/TypedElementQueryContext_MassForwarder.h
+++ b/Engine/Plugins/Experimental/EditorDataStorage/Source/TedsCore/Private/Queries/TypedElementQueryContext_MassForwarder.h
@@ -4,18 +4,24 @@
 
 #include "Elements/Framework/TypedElementQueryContextImplementation.h"
 #include "Elements/Framework/TypedElementQueryFunctions.h"
+#include "TypedElementDatabaseEnvironment.h"
+#include "Queries/QueryContextCapabilities/EditorDataStorageDynamicColumnInfo.h"
 
 struct FMassExecutionContext;
 
 namespace UE::Editor::DataStorage::Queries
 {
-	/** Query context implementation for single row queries that call into TEDS Direct API. */
-	class FQueryContext_MassForwarder final : public SingleRowInfo, public RowBatchInfo, public IQueryFunctionResponse
+	/** Query context implementation for single row and batch queries that call into TEDS Direct API. */
+	class FQueryContext_MassForwarder final : 
+		public SingleRowInfo, 
+		public RowBatchInfo,
+		public FDynamicColumnInfoContextCapability,
+		public IQueryFunctionResponse
 	{
 	public:
 		static const FName Name;
 
-		explicit FQueryContext_MassForwarder(FMassExecutionContext& Context);
+		FQueryContext_MassForwarder(FEnvironment& Environment, FMassExecutionContext& Context);
 
 		/**
 		 * @section ISingleRowInfo
@@ -39,7 +45,14 @@ namespace UE::Editor::DataStorage::Queries
 		virtual void GetConstColumns(TArrayView<const void*> ColumnsData, TConstArrayView<const UScriptStruct*> ColumnTypes) override;
 		virtual void GetMutableColumns(TArrayView<void*> ColumnsData, TConstArrayView<const UScriptStruct*> ColumnTypes) override;
 
+		/**
+		 * @section IQueryEnvironmentRetrieval
+		 */
+		virtual FEnvironment& GetEnvironment() override;
+		virtual const FEnvironment& GetEnvironment() const override;
+
 	private:
+		FEnvironment& Environment;
 		FMassExecutionContext& Context;
 		const RowHandle* CurrentRow;
 		const RowHandle* EndRow;
diff --git a/Engine/Plugins/Experimental/EditorDataStorage/Source/TedsCore/Private/TypedElementDatabase.cpp b/Engine/Plugins/Experimental/EditorDataStorage/Source/TedsCore/Private/TypedElementDatabase.cpp
--- a/Engine/Plugins/Experimental/EditorDataStorage/Source/TedsCore/Private/TypedElementDatabase.cpp
+++ b/Engine/Plugins/Experimental/EditorDataStorage/Source/TedsCore/Private/TypedElementDatabase.cpp
@@ -598,7 +598,7 @@ void UEditorDataStorage::FilterRowsBy(UE::Editor::DataStorage::FRowHandleArray&
 	using namespace UE::Editor::DataStorage;
 	using namespace UE::Editor::DataStorage::Queries;
 
-	if (ActiveEditorEntityManager)
+	if (ActiveEditorEntityManager && Environment)
 	{
 		struct FQueryResult : TResult<bool>
 		{
@@ -609,7 +609,7 @@ void UEditorDataStorage::FilterRowsBy(UE::Editor::DataStorage::FRowHandleArray&
 			}
 		};
 
-		QueryContext_DirectSingleApi Context(*this);
+		QueryContext_DirectSingleApi Context(*Environment, *this);
 		Context.CheckCompatiblity(Filter);
 
 		bool bAcceptedResult = Options == EFilterOptions::Inclusive ? true : false;
@@ -1069,11 +1069,6 @@ bool UEditorDataStorage::MatchesColumns(UE::Editor::DataStorage::RowHandle Row,
 	return false;
 }
 
-const UScriptStruct* UEditorDataStorage::FindDynamicColumn(const UE::Editor::DataStorage::FDynamicColumnDescription& Description) const
-{
-	return Environment->FindDynamicColumn(*Description.TemplateType, Description.Identifier);
-}
-
 const UScriptStruct* UEditorDataStorage::GenerateDynamicColumn(const UE::Editor::DataStorage::FDynamicColumnDescription& Description)
 {
 	return Environment->GenerateDynamicColumn(*Description.TemplateType, Description.Identifier);
@@ -1448,7 +1443,7 @@ void UEditorDataStorage::TickCooperativeTasks()
 	// The expected duration per frame.
 	int64 TargetFrameRate = FMath::Clamp(CooperativeTasks_TargetFrameRate.GetValueOnGameThread(), 16, 120);
 	const FTimespan TargetFrameDuration = FTimespan::FromMilliseconds(1000.0 / static_cast<double>(TargetFrameRate));
-	// The amount of time a task has to go over the alloted time before it's reported.
+	// The amount of time a task has to go over the allotted time before it's reported.
 	const FTimespan OverBudgetReportThreshold = FTimespan::FromMilliseconds(CooperativeTasks_ReportThresholdMs.GetValueOnGameThread());
 
 	if (!CooperativeTasks.IsEmpty())
@@ -1471,21 +1466,37 @@ void UEditorDataStorage::TickCooperativeTasks()
 			: MinTimePerTick;
 		
 		int32 CurrentTaskIndex = 0;
-		while (RemainingTime > MinTimePerTask && CurrentTaskIndex < CooperativeTasks.Num())
+		
+		uint64 ProcessingStartTime = FPlatformTime::Cycles64();
+		while (CurrentTaskIndex < CooperativeTasks.Num())
 		{
 			FCooperativeTask& CurrentTask = CooperativeTasks[CurrentTaskIndex];
 			// Make sure that tasks aren't run multiple times per frame and eat up cycles from lower priority tasks
 			if (!CurrentTask.bHasRun)
 			{
-				uint64 TaskStartTime = FPlatformTime::Cycles64();
-
 				TRACE_CPUPROFILER_EVENT_SCOPE_TEXT(
 					*FString::Printf(TEXT("[TEDS] Tick time sliced tasks: '%s'"), *CurrentTask.Name.ToString()));
 
 				// Run the task.
-				CurrentTask.Callback(RemainingTime);
-				CurrentTask.Order = CurrentTask.OrderResetValue;
-				CurrentTask.bHasRun = true;
+				{
+					// Measure only the task specifically going over time, not the cost of any bookkeeping.
+					uint64 TaskStartTime = FPlatformTime::Cycles64();
+
+					CurrentTask.Callback(RemainingTime);
+					CurrentTask.Order = CurrentTask.OrderResetValue;
+					CurrentTask.bHasRun = true;
+
+					FTimespan TaskDuration(FTimespan::FromSeconds(FPlatformTime::ToSeconds64(FPlatformTime::Cycles64() - TaskStartTime)));
+					if (TaskDuration > RemainingTime + OverBudgetReportThreshold)
+					{
+						// Report this as verbose and not a warning because it's expected that when the frame behaves erratic during loading or
+						// under heavy load this can trigger for even well behaving tasks. It's also possible that the OS takes control away from
+						// the editor for a while to give other programs a chance to run which can also cause spikes. This becomes useful if there's
+						// a pattern of frequent runs going over budget.
+						UE_LOG(LogEditorDataStorage, Verbose, TEXT("Time sliced task: '%s' took %.2fms, but was allotted %.2fms"),
+							*CurrentTask.Name.ToString(), TaskDuration.GetTotalMilliseconds(), RemainingTime.GetTotalMilliseconds());
+					}
+				}
 
 				// Find the new slot to put the processed task in by bubbling it up the queue.
 				for (int32 Index = CurrentTaskIndex + 1; Index < CooperativeTasks.Num(); Index++)
@@ -1499,23 +1510,16 @@ void UEditorDataStorage::TickCooperativeTasks()
 						break;
 					}
 				}
-				
-				FTimespan TaskDuration(FTimespan::FromSeconds(FPlatformTime::ToSeconds64(FPlatformTime::Cycles64() - TaskStartTime)));
-				if (TaskDuration > RemainingTime + OverBudgetReportThreshold)
-				{
-					// Report this as verbose and not a warning because it's expected that when the frame behaves erratic during loading or
-					// under heavy load this can trigger for even well behaving tasks. It's also possible that the OS takes control away from
-					// the editor for a while to give other programs a chance to run which can also cause spikes. This becomes useful if there's
-					// a pattern of frequent runs going over budget.
-					UE_LOG(LogEditorDataStorage, Verbose, TEXT("Time sliced task: '%s' took %.2fms, but was alloted %.2fms"),
-						*CurrentTask.Name.ToString(), TaskDuration.GetTotalMilliseconds(), RemainingTime.GetTotalMilliseconds());
-				}
-				RemainingTime -= TaskDuration;
 			}
 			else
 			{
 				CurrentTaskIndex++;
 			}
+
+			if (FTimespan::FromSeconds(FPlatformTime::ToSeconds64(FPlatformTime::Cycles64() - ProcessingStartTime)) >= RemainingTime)
+			{
+				break;
+			}
 		}
 
 		bRunningCooperativeUpdate = false;
@@ -1624,16 +1628,28 @@ int32 UEditorDataStorage::GetTableChunkSize(FName TableName) const
 	}
 }
 
-TSharedPtr<UE::Editor::DataStorage::FEnvironment> UEditorDataStorage::GetEnvironment()
+TSharedPtr<UE::Editor::DataStorage::FEnvironment> UEditorDataStorage::GetEnvironmentPtr()
 {
 	return Environment;
 }
 
-TSharedPtr<const UE::Editor::DataStorage::FEnvironment> UEditorDataStorage::GetEnvironment() const
+TSharedPtr<const UE::Editor::DataStorage::FEnvironment> UEditorDataStorage::GetEnvironmentPtr() const
 {
 	return Environment;
 }
 
+UE::Editor::DataStorage::FEnvironment& UEditorDataStorage::GetEnvironment()
+{
+	checkf(Environment, TEXT("Unable to retrieve the environment before TEDS Core has been initialized."));
+	return *Environment;
+}
+
+const UE::Editor::DataStorage::FEnvironment& UEditorDataStorage::GetEnvironment() const
+{
+	checkf(Environment, TEXT("Unable to retrieve the environment before TEDS Core has been initialized."));
+	return *Environment;
+}
+
 FMassArchetypeHandle UEditorDataStorage::LookupArchetype(UE::Editor::DataStorage::TableHandle InTableHandle) const
 {
 	using namespace UE::Editor::DataStorage;
diff --git a/Engine/Plugins/Experimental/EditorDataStorage/Source/TedsCore/Private/TypedElementDatabase.h b/Engine/Plugins/Experimental/EditorDataStorage/Source/TedsCore/Private/TypedElementDatabase.h
--- a/Engine/Plugins/Experimental/EditorDataStorage/Source/TedsCore/Private/TypedElementDatabase.h
+++ b/Engine/Plugins/Experimental/EditorDataStorage/Source/TedsCore/Private/TypedElementDatabase.h
@@ -6,9 +6,11 @@
 #include "Elements/Interfaces/TypedElementDataStorageInterface.h"
 #include "MassArchetypeTypes.h"
 #include "Misc/TVariant.h"
+#include "Processors/TypedElementQueryCapabilityImplemenationForwarder.h"
 #include "Templates/SharedPointer.h"
 #include "TypedElementDatabaseCommandBuffer.h"
 #include "UObject/ObjectMacros.h"
+#include "Queries/QueryContextCapabilities/EditorDataStorageDynamicColumnInfo.h"
 
 #include "TypedElementDatabase.generated.h"
 
@@ -28,7 +30,9 @@ namespace UE::Editor::DataStorage
 UCLASS(MinimalAPI)
 class UEditorDataStorage final
 	: public UObject
-	, public UE::Editor::DataStorage::ICoreProvider
+	, public UE::Editor::DataStorage::Queries::Private::TImplementationForwarder<
+		UE::Editor::DataStorage::ICoreProvider, 
+		UE::Editor::DataStorage::Queries::FDynamicColumnInfoContextCapability>
 {
 	GENERATED_BODY()
 
@@ -136,7 +140,6 @@ public:
 	UE_API virtual void ListColumns(UE::Editor::DataStorage::RowHandle Row, UE::Editor::DataStorage::ColumnListWithDataCallbackRef Callback) override;
 	UE_API virtual bool MatchesColumns(UE::Editor::DataStorage::RowHandle Row, const UE::Editor::DataStorage::Queries::FConditions& Conditions) const override;
 
-	UE_API virtual const UScriptStruct* FindDynamicColumn(const UE::Editor::DataStorage::FDynamicColumnDescription& Description) const override;
 	UE_API virtual const UScriptStruct* GenerateDynamicColumn(const UE::Editor::DataStorage::FDynamicColumnDescription& Description) override;
 	UE_API virtual void ForEachDynamicColumn(const UScriptStruct& Template, TFunctionRef<void(const UScriptStruct& Type)> Callback) const override;
 
@@ -204,8 +207,11 @@ public:
 	UE_API virtual bool SupportsExtension(FName Extension) const override;
 	UE_API virtual void ListExtensions(TFunctionRef<void(FName)> Callback) const override;
 	
-	UE_API TSharedPtr<UE::Editor::DataStorage::FEnvironment> GetEnvironment();
-	UE_API TSharedPtr<const UE::Editor::DataStorage::FEnvironment> GetEnvironment() const;
+	UE_API TSharedPtr<UE::Editor::DataStorage::FEnvironment> GetEnvironmentPtr();
+	UE_API TSharedPtr<const UE::Editor::DataStorage::FEnvironment> GetEnvironmentPtr() const;
+
+	virtual UE::Editor::DataStorage::FEnvironment& GetEnvironment() override;
+	virtual const UE::Editor::DataStorage::FEnvironment& GetEnvironment() const override;
 
 	UE_API FMassArchetypeHandle LookupArchetype(UE::Editor::DataStorage::TableHandle InTableHandle) const;
 
diff --git a/Engine/Plugins/Experimental/EditorDataStorage/Source/TedsCore/Private/TypedElementDatabaseCompatibility.cpp b/Engine/Plugins/Experimental/EditorDataStorage/Source/TedsCore/Private/TypedElementDatabaseCompatibility.cpp
--- a/Engine/Plugins/Experimental/EditorDataStorage/Source/TedsCore/Private/TypedElementDatabaseCompatibility.cpp
+++ b/Engine/Plugins/Experimental/EditorDataStorage/Source/TedsCore/Private/TypedElementDatabaseCompatibility.cpp
@@ -104,7 +104,7 @@ void UEditorDataStorageCompatibility::Initialize(UEditorDataStorage* InStorage)
 	FScopedExclusiveLock Lock(EGlobalLockScope::Public);
 
 	Storage = InStorage;
-	Environment = InStorage->GetEnvironment();
+	Environment = InStorage->GetEnvironmentPtr();
 	QueuedCommands.Initialize(Environment->GetScratchBuffer());
 	
 	Prepare();
diff --git a/Engine/Plugins/Experimental/EditorDataStorage/Source/TedsCore/Private/TypedElementDatabaseEnvironment.cpp b/Engine/Plugins/Experimental/EditorDataStorage/Source/TedsCore/Private/TypedElementDatabaseEnvironment.cpp
--- a/Engine/Plugins/Experimental/EditorDataStorage/Source/TedsCore/Private/TypedElementDatabaseEnvironment.cpp
+++ b/Engine/Plugins/Experimental/EditorDataStorage/Source/TedsCore/Private/TypedElementDatabaseEnvironment.cpp
@@ -109,7 +109,7 @@ namespace UE::Editor::DataStorage
 		return MassPhaseManager;
 	}
 
-	const UScriptStruct* FEnvironment::FindDynamicColumn(const UScriptStruct& Template, const FName& Identifier)
+	const UScriptStruct* FEnvironment::FindDynamicColumn(const UScriptStruct& Template, const FName& Identifier) const
 	{
 		return DynamicColumnGenerator.FindByTemplateId(Template, Identifier);
 	}
diff --git a/Engine/Plugins/Experimental/EditorDataStorage/Source/TedsCore/Private/TypedElementDatabaseEnvironment.h b/Engine/Plugins/Experimental/EditorDataStorage/Source/TedsCore/Private/TypedElementDatabaseEnvironment.h
--- a/Engine/Plugins/Experimental/EditorDataStorage/Source/TedsCore/Private/TypedElementDatabaseEnvironment.h
+++ b/Engine/Plugins/Experimental/EditorDataStorage/Source/TedsCore/Private/TypedElementDatabaseEnvironment.h
@@ -53,7 +53,7 @@ namespace UE::Editor::DataStorage
 
 		// Finds the type information for the dynamic column
 		// Dynamic columns are specified by a template layout and a FName Identifier
-		const UScriptStruct* FindDynamicColumn(const UScriptStruct& Template, const FName& Identifier);
+		const UScriptStruct* FindDynamicColumn(const UScriptStruct& Template, const FName& Identifier) const;
 		// Generates or returns an existing type for the dynamic column
 		// Dynamic columns are specified by a template layout and a FName Identifier
 		const UScriptStruct* GenerateDynamicColumn(const UScriptStruct& Template, const FName& Identifier);
@@ -101,4 +101,14 @@ namespace UE::Editor::DataStorage
 
 		uint64 UpdateCycleId = 0;
 	};
+
+	class IQueryEnvironmentRetrieval
+	{
+	public:
+		virtual ~IQueryEnvironmentRetrieval() = default;
+
+	protected:
+		virtual FEnvironment& GetEnvironment() = 0;
+		virtual const FEnvironment& GetEnvironment() const = 0;
+	};
 } // namespace UE::Editor::DataStorage
diff --git a/Engine/Plugins/Experimental/EditorDataStorage/Source/TedsCore/Tests/DynamicColumnTests.cpp b/Engine/Plugins/Experimental/EditorDataStorage/Source/TedsCore/Tests/DynamicColumnTests.cpp
--- a/Engine/Plugins/Experimental/EditorDataStorage/Source/TedsCore/Tests/DynamicColumnTests.cpp
+++ b/Engine/Plugins/Experimental/EditorDataStorage/Source/TedsCore/Tests/DynamicColumnTests.cpp
@@ -93,16 +93,9 @@ namespace UE::Editor::DataStorage::Tests
 					// Note: There is no HasColumn function for syntactic sugar to get a dynamic column type
 					TArray<const UScriptStruct*> DynamicTagTypes;
 					
-					DynamicTagTypes.Emplace(TedsInterface->FindDynamicColumn(FDynamicColumnDescription
-						{
-							.TemplateType = FTestDynamicTag::StaticStruct(),
-							.Identifier = Identifiers[0]
-						}));
-					DynamicTagTypes.Emplace(TedsInterface->FindDynamicColumn(FDynamicColumnDescription
-						{
-							.TemplateType = FTestDynamicTag::StaticStruct(),
-							.Identifier = Identifiers[1]
-						}));
+					DynamicTagTypes.Emplace(TedsInterface->FindDynamicColumnType<FTestDynamicTag>(Identifiers[0]));
+					DynamicTagTypes.Emplace(TedsInterface->FindDynamicColumnType<FTestDynamicTag>(Identifiers[1]));
+					
 					TestTrue("Expected columns not found", TedsInterface->HasColumns(Rows[0], MakeConstArrayView({DynamicTagTypes[0], DynamicTagTypes[1]})));
 					
 					TestTrue("Expected columns not found", TedsInterface->HasColumns(Rows[1], MakeConstArrayView({DynamicTagTypes[0]})));
diff --git a/Engine/Source/Runtime/TypedElementFramework/Public/DataStorage/CommonTypes.h b/Engine/Source/Runtime/TypedElementFramework/Public/DataStorage/CommonTypes.h
--- a/Engine/Source/Runtime/TypedElementFramework/Public/DataStorage/CommonTypes.h
+++ b/Engine/Source/Runtime/TypedElementFramework/Public/DataStorage/CommonTypes.h
@@ -48,6 +48,19 @@ namespace UE
 	template< class T >
 	using type_identity_t = type_identity<T>::type;
 
+	template <class T>
+	constexpr T MakeDefault()
+	{
+		if constexpr (std::is_void_v<T>)
+		{
+			return;
+		}
+		else
+		{
+			return T{};
+		}
+	}
+
 	namespace Editor::DataStorage
 	{
 		using FColumn = FEditorDataStorageColumn;
diff --git a/Engine/Source/Runtime/TypedElementFramework/Public/Elements/Framework/TypedElementQueryCapabilities.h b/Engine/Source/Runtime/TypedElementFramework/Public/Elements/Framework/TypedElementQueryCapabilities.h
--- a/Engine/Source/Runtime/TypedElementFramework/Public/Elements/Framework/TypedElementQueryCapabilities.h
+++ b/Engine/Source/Runtime/TypedElementFramework/Public/Elements/Framework/TypedElementQueryCapabilities.h
@@ -5,6 +5,8 @@
 #include <type_traits>
 #include "Containers/ContainersFwd.h"
 #include "DataStorage/Handles.h"
+#include "DataStorage/Queries/Conditions.h"
+#include "Elements/Framework/TypedElementQueryFunctionArgumentTypes.h"
 #include "Misc/EnumClassFlags.h"
 #include "UObject/NameTypes.h"
 
@@ -31,7 +33,7 @@ namespace UE::Editor::DataStorage::Queries
 		virtual ~IContextCapability() = default;
 	};
 
-	
+#define WithWrappers 1	
 #define ArgTypeName(Type, Name) Type Name
 
 #define Function0(Capability, Return, Function) virtual Return Function() = 0;
@@ -52,7 +54,7 @@ namespace UE::Editor::DataStorage::Queries
 	{ \
 		inline static const FName Name = #InCapability ; \
 		static constexpr EContextCapabilityFlags Flags = InFlags; \
-		virtual ~I##InCapability () override = default;
+		virtual ~I##InCapability () = default;
 
 #define CapabilityEnd(Capability) \
 	}; \
@@ -73,4 +75,5 @@ namespace UE::Editor::DataStorage::Queries
 #undef ConstFunction4
 #undef CapabilityStart
 #undef CapabilityEnd
+#undef WithWrappers
 } // namespace UE::Editor::DataStorage::Queries
\ No newline at end of file
diff --git a/Engine/Source/Runtime/TypedElementFramework/Public/Elements/Framework/TypedElementQueryCapabilities.inl b/Engine/Source/Runtime/TypedElementFramework/Public/Elements/Framework/TypedElementQueryCapabilities.inl
--- a/Engine/Source/Runtime/TypedElementFramework/Public/Elements/Framework/TypedElementQueryCapabilities.inl
+++ b/Engine/Source/Runtime/TypedElementFramework/Public/Elements/Framework/TypedElementQueryCapabilities.inl
@@ -1,6 +1,6 @@
 // Copyright Epic Games, Inc. All Rights Reserved.
 
-// Mini-DSL to declare TEDS capabilities.
+// Mechanic to declare TEDS capabilities.
 // Do not include #pragma once here or any other guards against repeated inclusion as this file is meant to be included multiple times.
 // This file contains the definitions of all available query capabilities. The macros get expended in various ways to generate interfaces,
 // function forwarders, contexts, etc.
@@ -63,4 +63,84 @@ CapabilityStart(RowBatchInfo, EContextCapabilityFlags::SupportsBatch)
 		}
 	}
 #endif
-CapabilityEnd(RowBatchInfo)
\ No newline at end of file
+CapabilityEnd(RowBatchInfo)
+
+/**
+ * Capability to work with dynamic columns
+ */
+CapabilityStart(DynamicColumnInfo, EContextCapabilityFlags::SupportsSingle | EContextCapabilityFlags::SupportsBatch)
+	/**
+	 * Finds the type information for a dynamic column.
+	 * If the dynamic column has not been generated, then return nullptr
+	 * The TemplateType may be a type derived from either FColumn or FTag, anything else will return nullptr
+	 */
+	ConstFunction2(DynamicColumnInfo, const UScriptStruct*, FindDynamicColumnType, (const UScriptStruct&, TemplateType), (const FName&, Identifier))
+	
+	/**
+	 * Executes the given callback for each known dynamic column that derives from the base template provided.
+	 */
+	ConstFunction2(DynamicColumnInfo, void, ForEachDynamicColumnType, (const UScriptStruct&, Template), (TFunctionRef<void(const UScriptStruct& Type)>, Callback))
+
+#if defined(WithWrappers)
+	/**
+	 * Finds the type information for a dynamic column.
+	 * If the dynamic column has not been generated, then return nullptr
+	 * The DynamicColumnTemplateType may be a type derived from either FColumn or FTag, anything else will result in a compile error.
+	 */
+	template<TDynamicColumnTemplate DynamicColumnTemplateType>
+	const UScriptStruct* FindDynamicColumnType(const FName& Identifier) const
+	{
+		return this->FindDynamicColumnType(*DynamicColumnTemplateType::StaticStruct(), Identifier);
+	}
+	
+	/**
+	 * Finds the type information for a dynamic column.
+	 * If the dynamic column has not been generated, then return nullptr
+	 * The DynamicColumnTemplateType may be a type derived from either FColumn or FTag, anything else will result in a compile error.
+	 */
+	template<TDynamicColumnTemplate DynamicColumnTemplateType, TVariantName Identifier>
+	const UScriptStruct* FindDynamicColumnType() const
+	{
+		return this->FindDynamicColumnType(*DynamicColumnTemplateType::StaticStruct(),
+			[]()
+			{
+				static FName Name = Identifier.GetName();
+				return Name;
+			}());
+	}
+
+	/**
+	 * Finds the type information for a dynamic column.
+	 * If the dynamic column has not been generated, then return nullptr
+	 * The ColumnType can be any TColumn but TColumns without a name or dynamic column type will return a nullptr.
+	 */
+	template<TColumnDescType ColumnType>
+	const UScriptStruct* FindDynamicColumnType(const ColumnType& Column) const
+	{
+		return Column.TypeInfo
+			? this->FindDynamicColumnType(*Column.TypeInfo, Column.Identifier)
+			: nullptr;
+	}
+
+	/**
+	 * Executes the given callback for each known dynamic column that derives from the base template provided.
+	 * The DynamicColumnTemplateType may be a type derived from either FColumn or FTag, anything else will result in a compile error. 
+	 */
+	template<TDynamicColumnTemplate DynamicColumnTemplateType>
+	void ForEachDynamicColumnType(TFunctionRef<void(const UScriptStruct& Type)> Callback) const
+	{
+		this->ForEachDynamicColumnType(*DynamicColumnTemplateType::StaticStruct(), Callback);
+	}
+
+	UE_DEPRECATED(5.8, "Please one of the new versions of FindDynamicColumnType.")
+	const UScriptStruct* FindDynamicColumnType(const FDynamicColumnDescription& Description) const
+	{
+		return Description.TemplateType ? this->FindDynamicColumnType(*Description.TemplateType, Description.Identifier) : nullptr;
+	}
+	UE_DEPRECATED(5.8, "Please one of the new versions of FindDynamicColumnType.")
+	const UScriptStruct* FindDynamicColumn(const FDynamicColumnDescription& Description) const
+	{
+		return Description.TemplateType ? this->FindDynamicColumnType(*Description.TemplateType, Description.Identifier) : nullptr;
+	}
+#endif
+CapabilityEnd(DynamicColumnInfo)
\ No newline at end of file
diff --git a/Engine/Source/Runtime/TypedElementFramework/Public/Elements/Framework/TypedElementQueryContextImplementation.h b/Engine/Source/Runtime/TypedElementFramework/Public/Elements/Framework/TypedElementQueryContextImplementation.h
--- a/Engine/Source/Runtime/TypedElementFramework/Public/Elements/Framework/TypedElementQueryContextImplementation.h
+++ b/Engine/Source/Runtime/TypedElementFramework/Public/Elements/Framework/TypedElementQueryContextImplementation.h
@@ -2,6 +2,7 @@
 
 #pragma once
 
+#include "DataStorage/CommonTypes.h"
 #include "Elements/Framework/TypedElementQueryCapabilities.h"
 #include "Elements/Framework/TypedElementQueryContract.h"
 #include "Elements/Framework/TypedElementQueryContext.h"
diff --git a/Engine/Source/Runtime/TypedElementFramework/Public/Elements/Framework/TypedElementQueryContextImplementation.inl b/Engine/Source/Runtime/TypedElementFramework/Public/Elements/Framework/TypedElementQueryContextImplementation.inl
--- a/Engine/Source/Runtime/TypedElementFramework/Public/Elements/Framework/TypedElementQueryContextImplementation.inl
+++ b/Engine/Source/Runtime/TypedElementFramework/Public/Elements/Framework/TypedElementQueryContextImplementation.inl
@@ -62,7 +62,7 @@ namespace UE::Editor::DataStorage::Queries
 	else \
 	{ \
 		checkf(false, TEXT("Function '" #FunctionName "' in capability '" #Capability "' is not supported by the current query context implementation.")); \
-		return UE::type_identity_t< ReturnValue >{}; \
+		return UE::MakeDefault< ReturnValue >(); \
 	}
 
 #define ArgTypeName(Type, Name) Type Name
diff --git a/Engine/Source/Runtime/TypedElementFramework/Public/Elements/Interfaces/TypedElementDataStorageInterface.h b/Engine/Source/Runtime/TypedElementFramework/Public/Elements/Interfaces/TypedElementDataStorageInterface.h
--- a/Engine/Source/Runtime/TypedElementFramework/Public/Elements/Interfaces/TypedElementDataStorageInterface.h
+++ b/Engine/Source/Runtime/TypedElementFramework/Public/Elements/Interfaces/TypedElementDataStorageInterface.h
@@ -61,7 +61,7 @@ struct FHierarchyRegistrationParams
 	const UScriptStruct* ParentChangedColumn = nullptr;
 };
 	
-class ICoreProvider : public IModularFeature
+class ICoreProvider : public Queries::IDynamicColumnInfo<IModularFeature>
 {
 public:
 	/**
@@ -353,13 +353,6 @@ public:
 	/** Determines if the columns in the row match the query conditions. */
 	virtual bool MatchesColumns(RowHandle Row, const Queries::FConditions& Conditions) const = 0;
 
-	/**
-	 * Finds the type information for a dynamic column.
-	 * If the dynamic column has not been generated, then return nullptr
-	 * The TemplateType may be a typed derived from either FColumn or FTag, anything else will return nullptr
-	 */
-	virtual const UScriptStruct* FindDynamicColumn(const FDynamicColumnDescription& Description) const = 0;
-
 	/**
 	 * Generates a new dynamic column from a Template.  A dynamic column is uniquely identified using the given template and an Identifier
 	 * This function is idempotent - multiple calls with the same parameters will result in subsequent calls returning the same type
@@ -686,13 +679,7 @@ inline void ICoreProvider::RemoveColumn<FValueTag>(RowHandle Row, const FName& T
 template<TDynamicColumnTemplate DynamicColumnTemplate>
 void ICoreProvider::RemoveColumn(RowHandle Row, const FName& Identifier)
 {
-	const FDynamicColumnDescription Description
-	{
-		.TemplateType = DynamicColumnTemplate::StaticStruct(),
-		.Identifier = Identifier
-	};
-	const UScriptStruct* StructInfo = FindDynamicColumn(Description);
-	RemoveColumn(Row, StructInfo);
+	RemoveColumn(Row, FindDynamicColumnType<DynamicColumnTemplate>(Identifier));
 }
 
 template<TEnumType EnumT>
@@ -837,13 +824,7 @@ DynamicColumnTemplate* ICoreProvider::GetColumn(RowHandle Row, const FName& Iden
 template <TDynamicColumnTemplate DynamicColumnTemplate>
 const DynamicColumnTemplate* ICoreProvider::GetColumn(RowHandle Row, const FName& Identifier) const
 {
-	const FDynamicColumnDescription Description
-	{
-		.TemplateType = DynamicColumnTemplate::StaticStruct(),
-		.Identifier = Identifier
-	};
-	const UScriptStruct* StructInfo = FindDynamicColumn(Description);
-	if (StructInfo)
+	if (const UScriptStruct* StructInfo = FindDynamicColumnType<DynamicColumnTemplate>(Identifier))
 	{
 		return static_cast<const DynamicColumnTemplate*>(GetColumnData(Row, StructInfo));
 	}
@@ -862,15 +843,12 @@ bool ICoreProvider::HasColumns(RowHandle Row, const ColumnTypes&... Columns) con
 	TArray<const UScriptStruct*, TInlineAllocator<sizeof...(ColumnTypes)>> TypeInfo;
 	TypeInfo.Reserve(sizeof...(ColumnTypes));
 
-	auto GetType = [this, &TypeInfo]<Queries::TColumnDescType ColumnType>(const ColumnType & Column) -> bool
+	auto GetType = [this, &TypeInfo]<Queries::TColumnDescType ColumnType>(const ColumnType& Column) -> bool
 		{
+			const UScriptStruct* BaseColumnType = Column.TypeInfo.Get();
 			const UScriptStruct* Result = Column.Identifier.IsNone()
-				? Column.TypeInfo.Get()
-				: FindDynamicColumn(FDynamicColumnDescription
-					{
-						.TemplateType = Column.TypeInfo.Get(),
-						.Identifier = Column.Identifier
-					});
+				? BaseColumnType
+				: (BaseColumnType ? FindDynamicColumnType(*BaseColumnType, Column.Identifier) : nullptr);
 			if (Result)
 			{
 				TypeInfo.Add(Result);
diff --git a/Engine/Source/Runtime/TypedElementFramework/Public/Elements/Interfaces/TypedElementQueryStorageInterfaces.h b/Engine/Source/Runtime/TypedElementFramework/Public/Elements/Interfaces/TypedElementQueryStorageInterfaces.h
--- a/Engine/Source/Runtime/TypedElementFramework/Public/Elements/Interfaces/TypedElementQueryStorageInterfaces.h
+++ b/Engine/Source/Runtime/TypedElementFramework/Public/Elements/Interfaces/TypedElementQueryStorageInterfaces.h
@@ -7,6 +7,7 @@
 #include "DataStorage/Handles.h"
 #include "DataStorage/MapKey.h"
 #include "DataStorage/Queries/Types.h"
+#include "Elements/Framework/TypedElementQueryCapabilities.h"
 #include "UObject/Class.h"
 #include "UObject/WeakObjectPtrTemplates.h"
 #include "Templates/Function.h"
@@ -26,7 +27,7 @@ namespace UE::Editor::DataStorage
 	/**
 	 * Base interface for any contexts provided to query callbacks.
 	 */
-	struct ICommonQueryContext
+	struct ICommonQueryContext : public Queries::IDynamicColumnInfo<Queries::IContextCapability>
 	{
 		virtual ~ICommonQueryContext() = default;
 
@@ -96,14 +97,6 @@ namespace UE::Editor::DataStorage
 		template<TDynamicColumnTemplate DynamicColumnTemplate>
 		bool HasColumn(RowHandle Row, const FName& Identifier) const;
 
-		/**
-		 * Finds the type of a dynamic column
-		 * @return The UScriptStruct for a previously created Dynamic Column.  If no column exists, then nullptr
-		 */
-		virtual const UScriptStruct* FindDynamicColumnType(const FDynamicColumnDescription& Description) const = 0;
-		template<TDynamicColumnTemplate TemplateType>
-		const UScriptStruct* FindDynamicColumnType(const FName& Identifier) const;
-
 		// Sets the Target's Parent if the query has a configured Hierarchy
 		virtual void SetParentRow(
 			RowHandle Target,
@@ -492,33 +485,15 @@ namespace UE::Editor::DataStorage
 	template <TDynamicColumnTemplate DynamicColumnTemplate>
 	bool ICommonQueryContext::HasColumn(RowHandle Row, const FName& Identifier) const
 	{
-		if (const UScriptStruct* DynamicColumnType = FindDynamicColumnType<DynamicColumnTemplate>(Identifier))
-		{
-			return HasColumn(DynamicColumnType);
-		}
-		return false;
-	}
-
-	template <TDynamicColumnTemplate TemplateType>
-	const UScriptStruct* ICommonQueryContext::FindDynamicColumnType(const FName& Identifier) const
-	{
-		return FindDynamicColumnType(FDynamicColumnDescription
-			{
-				.TemplateType = TemplateType::StaticStruct(),
-				.Identifier = Identifier
-			});
+		const UScriptStruct* DynamicColumnType = FindDynamicColumnType<DynamicColumnTemplate>(Identifier);
+		return DynamicColumnType ? HasColumn(Row, DynamicColumnType) : false;
 	}
 
 	template <typename DynamicColumnTemplate>
 	const DynamicColumnTemplate* ICommonQueryWithEnvironmentContext::GetColumn(const FName& Identifier) const
 	{
 		using namespace UE::Editor::DataStorage;
-		const FDynamicColumnDescription Description
-		{
-			.TemplateType = DynamicColumnTemplate::StaticStruct(),
-			.Identifier = Identifier
-		};
-		const UScriptStruct* StructInfo = FindDynamicColumnType(Description);
+		const UScriptStruct* StructInfo = FindDynamicColumnType(DynamicColumnTemplate::StaticStruct(), Identifier);
 		if (StructInfo)
 		{
 			const void* ColumnData = GetColumn(StructInfo);
@@ -531,12 +506,7 @@ namespace UE::Editor::DataStorage
 	DynamicColumnTemplate* ICommonQueryWithEnvironmentContext::GetMutableColumn(const FName& Identifier)
 	{
 		using namespace UE::Editor::DataStorage;
-		const FDynamicColumnDescription Description
-		{
-			.TemplateType = DynamicColumnTemplate::StaticStruct(),
-			.Identifier = Identifier
-		};
-		const UScriptStruct* StructInfo = FindDynamicColumnType(Description);
+		const UScriptStruct* StructInfo = FindDynamicColumnType(DynamicColumnTemplate::StaticStruct(), Identifier);
 		if (StructInfo)
 		{
 			const void* ColumnData = GetMutableColumn(StructInfo);
@@ -549,12 +519,7 @@ namespace UE::Editor::DataStorage
 	bool ICommonQueryWithEnvironmentContext::HasColumn(const FName& Identifier) const
 	{
 		using namespace UE::Editor::DataStorage;
-		const FDynamicColumnDescription Description
-		{
-			.TemplateType = DynamicColumnTemplate::StaticStruct(),
-			.Identifier = Identifier
-		};
-		const UScriptStruct* StructInfo = FindDynamicColumnType(Description);
+		const UScriptStruct* StructInfo = FindDynamicColumnType(DynamicColumnTemplate::StaticStruct(), Identifier);
 		if (StructInfo)
 		{
 			return ICommonQueryContext::HasColumn(StructInfo);
@@ -566,12 +531,7 @@ namespace UE::Editor::DataStorage
 	bool ICommonQueryWithEnvironmentContext::HasColumn(RowHandle Row, const FName& Identifier) const
 	{
 		using namespace UE::Editor::DataStorage;
-		const FDynamicColumnDescription Description
-		{
-			.TemplateType = DynamicColumnTemplate::StaticStruct(),
-			.Identifier = Identifier
-		};
-		const UScriptStruct* StructInfo = FindDynamicColumnType(Description);
+		const UScriptStruct* StructInfo = FindDynamicColumnType(DynamicColumnTemplate::StaticStruct(), Identifier);
 		if (StructInfo)
 		{
 			return ICommonQueryContext::HasColumn(Row, StructInfo);
diff --git a/Engine/Source/Runtime/TypedElementFramework/Tests/Elements/Framework/TypedElementQueryContextMock.h b/Engine/Source/Runtime/TypedElementFramework/Tests/Elements/Framework/TypedElementQueryContextMock.h
--- a/Engine/Source/Runtime/TypedElementFramework/Tests/Elements/Framework/TypedElementQueryContextMock.h
+++ b/Engine/Source/Runtime/TypedElementFramework/Tests/Elements/Framework/TypedElementQueryContextMock.h
@@ -133,7 +133,8 @@ namespace UE::Editor::DataStorage::Queries
 #undef WithWrappers
 
 #define MockFunctionImplementation(Capability, Return, Function) \
-		checkf(false, TEXT( "Context function " #Function " in capability " #Capability " was unexpectedly called." )); return UE::type_identity_t<Return>{};
+		checkf(false, TEXT( "Context function " #Function " in capability " #Capability " was unexpectedly called." )); \
+		return UE::MakeDefault< Return >();
 #define Function0(Capability, Return, Function) \
 		MockFunction(Capability, Function) = []() { MockFunctionImplementation(Capability, Return, Function) };
 #define Function1(Capability, Return, Function, Arg1) \
