[TedsCore]
* Fixed hides class member error happening when using latest toolchain/header units

diff --git a/Engine/Plugins/Experimental/EditorDataStorage/Source/TedsCore/Private/TypedElementDatabase.cpp b/Engine/Plugins/Experimental/EditorDataStorage/Source/TedsCore/Private/TypedElementDatabase.cpp
--- a/Engine/Plugins/Experimental/EditorDataStorage/Source/TedsCore/Private/TypedElementDatabase.cpp
+++ b/Engine/Plugins/Experimental/EditorDataStorage/Source/TedsCore/Private/TypedElementDatabase.cpp
@@ -300,20 +300,20 @@ TSharedPtr<const FMassEntityManager> UEditorDataStorage::GetActiveEditorEntityMa
 	return ActiveEditorEntityManager;
 }
 
-UE::Editor::DataStorage::TableHandle UEditorDataStorage::RegisterTable(TConstArrayView<const UScriptStruct*> ColumnList, const FName& Name)
+UE::Editor::DataStorage::TableHandle UEditorDataStorage::RegisterTable(TConstArrayView<const UScriptStruct*> ColumnList, const FName& InName)
 {
 	using namespace UE::Editor::DataStorage;
 
-	if (ActiveEditorEntityManager && !TableNameLookup.Contains(Name))
+	if (ActiveEditorEntityManager && !TableNameLookup.Contains(InName))
 	{
 		TableHandle Result = Tables.Num();
 		FMassArchetypeCreationParams ArchetypeCreationParams;
-		ArchetypeCreationParams.DebugName = Name;
-		ArchetypeCreationParams.ChunkMemorySize = GetTableChunkSize(Name);
+		ArchetypeCreationParams.DebugName = InName;
+		ArchetypeCreationParams.ChunkMemorySize = GetTableChunkSize(InName);
 		Tables.Add(ActiveEditorEntityManager->CreateArchetype(ColumnList, ArchetypeCreationParams));
-		if (Name.IsValid())
+		if (InName.IsValid())
 		{
-			TableNameLookup.Add(Name, Result);
+			TableNameLookup.Add(InName, Result);
 		}
 		return Result;
 	}
@@ -321,32 +321,32 @@ UE::Editor::DataStorage::TableHandle UEditorDataStorage::RegisterTable(TConstArr
 }
 
 UE::Editor::DataStorage::TableHandle UEditorDataStorage::RegisterTable(UE::Editor::DataStorage::TableHandle SourceTable,
-	TConstArrayView<const UScriptStruct*> ColumnList, const FName& Name)
+	TConstArrayView<const UScriptStruct*> ColumnList, const FName& InName)
 {
 	using namespace UE::Editor::DataStorage;
 
 	const int32 TableHandleAsIndex = Private::ConvertTableHandleToIndex(SourceTable);
-	if (ActiveEditorEntityManager && TableHandleAsIndex < Tables.Num() && !TableNameLookup.Contains(Name))
+	if (ActiveEditorEntityManager && TableHandleAsIndex < Tables.Num() && !TableNameLookup.Contains(InName))
 	{
 		TableHandle Result = Tables.Num();
 		FMassArchetypeCreationParams ArchetypeCreationParams;
-		ArchetypeCreationParams.DebugName = Name;
-		ArchetypeCreationParams.ChunkMemorySize = GetTableChunkSize(Name);
+		ArchetypeCreationParams.DebugName = InName;
+		ArchetypeCreationParams.ChunkMemorySize = GetTableChunkSize(InName);
 		Tables.Add(ActiveEditorEntityManager->CreateArchetype(Tables[TableHandleAsIndex], ColumnList, ArchetypeCreationParams));
-		if (Name.IsValid())
+		if (InName.IsValid())
 		{
-			TableNameLookup.Add(Name, Result);
+			TableNameLookup.Add(InName, Result);
 		}
 		return Result;
 	}
 	return InvalidTableHandle;
 }
 
-UE::Editor::DataStorage::TableHandle UEditorDataStorage::FindTable(const FName& Name)
+UE::Editor::DataStorage::TableHandle UEditorDataStorage::FindTable(const FName& InName)
 {
 	using namespace UE::Editor::DataStorage;
 
-	TableHandle* TableHandle = TableNameLookup.Find(Name);
+	TableHandle* TableHandle = TableNameLookup.Find(InName);
 	return TableHandle ? *TableHandle : InvalidTableHandle;
 }
 
@@ -1085,9 +1085,9 @@ UE::Editor::DataStorage::FHierarchyHandle UEditorDataStorage::RegisterHierarchy(
 	return Environment->GetHierarchyRegistrar().RegisterHierarchy(this, Params);
 }
 
-UE::Editor::DataStorage::FHierarchyHandle UEditorDataStorage::FindHierarchyByName(const FName& Name) const
+UE::Editor::DataStorage::FHierarchyHandle UEditorDataStorage::FindHierarchyByName(const FName& InName) const
 {
-	return Environment->GetHierarchyRegistrar().FindHierarchyByName(Name);
+	return Environment->GetHierarchyRegistrar().FindHierarchyByName(InName);
 }
 
 bool UEditorDataStorage::IsValidHierarchyHandle(UE::Editor::DataStorage::FHierarchyHandle Handle) const
@@ -1289,7 +1289,7 @@ UE::Editor::DataStorage::FQueryResult UEditorDataStorage::RunQuery(
 }
 
 UE::Editor::DataStorage::FQueryResult UEditorDataStorage::RunQuery(
-	UE::Editor::DataStorage::QueryHandle Query, UE::Editor::DataStorage::EDirectQueryExecutionFlags Flags,
+	UE::Editor::DataStorage::QueryHandle Query, UE::Editor::DataStorage::EDirectQueryExecutionFlags InFlags,
 	UE::Editor::DataStorage::DirectQueryCallbackRef Callback)
 {
 	using namespace UE::Editor::DataStorage;
@@ -1299,7 +1299,7 @@ UE::Editor::DataStorage::FQueryResult UEditorDataStorage::RunQuery(
 	if (ActiveEditorEntityManager)
 	{
 		const FExtendedQueryStore::Handle StorageHandle(Query);
-		return Environment->GetQueryStore().RunQuery(*ActiveEditorEntityManager, *Environment, StorageHandle, Flags, Callback);
+		return Environment->GetQueryStore().RunQuery(*ActiveEditorEntityManager, *Environment, StorageHandle, InFlags, Callback);
 	}
 	else
 	{
@@ -1309,7 +1309,7 @@ UE::Editor::DataStorage::FQueryResult UEditorDataStorage::RunQuery(
 
 UE::Editor::DataStorage::FQueryResult UEditorDataStorage::RunQuery(
 	UE::Editor::DataStorage::QueryHandle Query,
-	UE::Editor::DataStorage::ERunQueryFlags Flags,
+	UE::Editor::DataStorage::ERunQueryFlags InFlags,
 	const UE::Editor::DataStorage::Queries::TQueryFunction<void>& Callback)
 {
 	using namespace UE::Editor::DataStorage;
@@ -1319,7 +1319,7 @@ UE::Editor::DataStorage::FQueryResult UEditorDataStorage::RunQuery(
 	if (ActiveEditorEntityManager)
 	{
 		return Environment->GetQueryStore().RunQuery(*ActiveEditorEntityManager, *Environment, 
-			FExtendedQueryStore::Handle(Query), Flags, Callback);
+			FExtendedQueryStore::Handle(Query), InFlags, Callback);
 	}
 	else
 	{
