TEDS UI - Fix asan error caused by accessing invalid ptr to metadata after it is destroyed

diff --git a/Engine/Plugins/Experimental/EditorDataStorage/Source/TedsUI/Private/Widgets/SlateColorWidget.cpp b/Engine/Plugins/Experimental/EditorDataStorage/Source/TedsUI/Private/Widgets/SlateColorWidget.cpp
--- a/Engine/Plugins/Experimental/EditorDataStorage/Source/TedsUI/Private/Widgets/SlateColorWidget.cpp
+++ b/Engine/Plugins/Experimental/EditorDataStorage/Source/TedsUI/Private/Widgets/SlateColorWidget.cpp
@@ -63,12 +63,18 @@ TSharedPtr<SWidget> FSlateColorWidgetConstructor::CreateWidget(UE::Editor::DataS
 	UE::Editor::DataStorage::FAttributeBinder Binder(TargetRow, DataStorage);
 
 	FPointerEventHandler OnMouseButtonDown;
-	const bool* IsEditable = Arguments.FindForColumn<FSlateColorColumn>(UE::Editor::DataStorage::IsEditableName).TryGetExact<bool>();
+	
+	UE::Editor::DataStorage::FMetaDataEntryView MetaDataEntryView = Arguments.FindForColumn<FSlateColorColumn>(UE::Editor::DataStorage::IsEditableName);
 
-	// If the column is editable, summon a color picker on click
-	if(IsEditable && *IsEditable)
+	if (MetaDataEntryView.IsSet())
 	{
-		OnMouseButtonDown = FPointerEventHandler::CreateStatic(&UE::SlateColorWidget::Private::SummonColorPicker, DataStorage, TargetRow);
+		const bool* IsEditable = MetaDataEntryView.TryGetExact<bool>();
+		
+		// If the column is editable, summon a color picker on click
+		if(IsEditable && *IsEditable)
+		{
+			OnMouseButtonDown = FPointerEventHandler::CreateStatic(&UE::SlateColorWidget::Private::SummonColorPicker, DataStorage, TargetRow);
+		}
 	}
 	
 	return SNew(SHorizontalBox)
diff --git a/Engine/Plugins/Experimental/EditorDataStorageFeatures/Source/TedsOutliner/Private/Widgets/OutlinerFolderLabelWidget.cpp b/Engine/Plugins/Experimental/EditorDataStorageFeatures/Source/TedsOutliner/Private/Widgets/OutlinerFolderLabelWidget.cpp
--- a/Engine/Plugins/Experimental/EditorDataStorageFeatures/Source/TedsOutliner/Private/Widgets/OutlinerFolderLabelWidget.cpp
+++ b/Engine/Plugins/Experimental/EditorDataStorageFeatures/Source/TedsOutliner/Private/Widgets/OutlinerFolderLabelWidget.cpp
@@ -105,10 +105,20 @@ TSharedRef<SWidget> FOutlinerFolderLabelWidgetConstructor::CreateLabel(UE::Edito
 {
 	using namespace UE::Editor::DataStorage;
 	TSharedRef<SWidget> Result = SNullWidget::NullWidget;
+
+	bool bEditable = false;
 	
-	const bool* IsEditable = Arguments.FindForColumn<FTypedElementLabelColumn>(IsEditableName).TryGetExact<bool>();
+	FMetaDataEntryView MetaDataEntryView = Arguments.FindForColumn<FTypedElementLabelColumn>(UE::Editor::DataStorage::IsEditableName);
 
-	if (IsEditable && *IsEditable)
+	if (MetaDataEntryView.IsSet())
+	{
+		if (const bool* EditableMetaDataPtr = MetaDataEntryView.TryGetExact<bool>())
+		{
+			bEditable = *EditableMetaDataPtr;
+		}
+	}
+	
+	if (bEditable)
 	{
 		FAttributeBinder WidgetRowBinder(WidgetRow, DataStorage);
 			
diff --git a/Engine/Plugins/Experimental/EditorDataStorageFeatures/Source/TedsOutliner/Private/Widgets/OutlinerTextWidget.cpp b/Engine/Plugins/Experimental/EditorDataStorageFeatures/Source/TedsOutliner/Private/Widgets/OutlinerTextWidget.cpp
--- a/Engine/Plugins/Experimental/EditorDataStorageFeatures/Source/TedsOutliner/Private/Widgets/OutlinerTextWidget.cpp
+++ b/Engine/Plugins/Experimental/EditorDataStorageFeatures/Source/TedsOutliner/Private/Widgets/OutlinerTextWidget.cpp
@@ -163,16 +163,23 @@ TSharedPtr<SWidget> FOutlinerTextWidgetConstructor::CreateWidget(UE::Editor::Dat
 
 	TSharedPtr<SWidget> Result = SNullWidget::NullWidget;
 
-	const bool* IsEditable = Arguments.FindForColumn<FTypedElementLabelColumn>(IsEditableName).TryGetExact<bool>();
+	FMetaDataEntryView MetaDataEntryView = Arguments.FindForColumn<FTypedElementLabelColumn>(IsEditableName);
 
-	if (IsEditable && *IsEditable)
+	if (MetaDataEntryView.IsSet())
 	{
-		Result = CreateEditableWidget(DataStorage, DataStorageUi, TargetRow, WidgetRow, Arguments);
-	}
-	else
-	{
-		Result = CreateNonEditableWidget(DataStorage, DataStorageUi, TargetRow, WidgetRow, Arguments);
+		const bool* IsEditable = MetaDataEntryView.TryGetExact<bool>();
+
+		if (IsEditable && *IsEditable)
+		{
+			Result = CreateEditableWidget(DataStorage, DataStorageUi, TargetRow, WidgetRow, Arguments);
+		}
+		else
+		{
+			Result = CreateNonEditableWidget(DataStorage, DataStorageUi, TargetRow, WidgetRow, Arguments);
+		}
 	}
+	
+	
 	return Result;
 }
 
