Fix crash that can occur in TEDS passing a nullptr UObject when it is passed to registration filters
* Sync tag handler for actors correctly adds actor components to be registered to TEDS. Was always passing nullptr.
* Fix reserved row leak passing null struct to AddCompatibleObject

diff --git a/Engine/Plugins/Experimental/EditorDataStorage/Source/TedsCore/Private/TypedElementDatabaseCompatibility.cpp b/Engine/Plugins/Experimental/EditorDataStorage/Source/TedsCore/Private/TypedElementDatabaseCompatibility.cpp
--- a/Engine/Plugins/Experimental/EditorDataStorage/Source/TedsCore/Private/TypedElementDatabaseCompatibility.cpp
+++ b/Engine/Plugins/Experimental/EditorDataStorage/Source/TedsCore/Private/TypedElementDatabaseCompatibility.cpp
@@ -244,6 +244,7 @@ UE::Editor::DataStorage::RowHandle UEditorDataStorageCompatibility::AddCompatibl
 
 	bool bCanAddObject =
 		ensureMsgf(Storage, TEXT("Trying to add a UObject to Typed Element's Data Storage before the storage is available.")) &&
+		Object != nullptr &&
 		ShouldAddObject(Object);
 	return bCanAddObject ? AddCompatibleObjectExplicitTransactionable<true>(Object) : UE::Editor::DataStorage::InvalidRowHandle;
 }
@@ -252,6 +253,11 @@ UE::Editor::DataStorage::RowHandle UEditorDataStorageCompatibility::AddCompatibl
 UE::Editor::DataStorage::RowHandle UEditorDataStorageCompatibility::AddCompatibleObjectExplicit(void* Object, TWeakObjectPtr<const UScriptStruct> TypeInfo)
 {
 	using namespace UE::Editor::DataStorage;
+
+	if (Object == nullptr)
+	{
+		return InvalidRowHandle;
+	}
 	
 	checkf(Storage, TEXT("Trying to add an object to Typed Element's Data Storage before the storage is available."));
 	FScopedExclusiveLock Lock(EGlobalLockScope::Public);
diff --git a/Engine/Plugins/Experimental/EditorDataStorageFeatures/Source/TedsActorCompatibility/Private/TedsActorCompatibilityModule.cpp b/Engine/Plugins/Experimental/EditorDataStorageFeatures/Source/TedsActorCompatibility/Private/TedsActorCompatibilityModule.cpp
--- a/Engine/Plugins/Experimental/EditorDataStorageFeatures/Source/TedsActorCompatibility/Private/TedsActorCompatibilityModule.cpp
+++ b/Engine/Plugins/Experimental/EditorDataStorageFeatures/Source/TedsActorCompatibility/Private/TedsActorCompatibilityModule.cpp
@@ -94,7 +94,7 @@ namespace UE::Editor::DataStorage
 						for (TWeakObjectPtr<UActorComponent> Component : ComponentsToRegister)
 						{
 							UActorComponent* NewComponent = Component.Get();
-							if (!NewComponent)
+							if (NewComponent)
 							{
 								CompatibilityProvider->AddCompatibleObject(NewComponent);
 							}
