[TEDS] Coopertive task refinement.

The following changes were made:
- When determining the amount of time available for cooperative tasks, don't count the amount of time used in the previous frame towards the total time used. This time is what's reserved for cooperative tasks so if it's counted as normal time spend it will double count. Say that there's 5ms for cooperative tasks and one tasks exactly uses that 5ms, then next frame it would seem there's no time left at all, causing the cooperative tasks to ping-pong between 0ms (capped at 2ms) and 5ms) in this case.
- When there's no time left in a frame, the cooperative task will wait a set number of frames (default 5) to allow the editor to recover. This is to avoid spikes that last one or two frame in the editor, for instance due to Slate activity, to be made worse by the cooperative task in in the process cause alarms to trigger for tasks that struggle to fit in 2 ms.
- Added stats to track the amount of available time.

diff --git a/Engine/Plugins/Experimental/EditorDataStorage/Source/TedsCore/Private/TypedElementDatabase.cpp b/Engine/Plugins/Experimental/EditorDataStorage/Source/TedsCore/Private/TypedElementDatabase.cpp
--- a/Engine/Plugins/Experimental/EditorDataStorage/Source/TedsCore/Private/TypedElementDatabase.cpp
+++ b/Engine/Plugins/Experimental/EditorDataStorage/Source/TedsCore/Private/TypedElementDatabase.cpp
@@ -17,6 +17,7 @@
 #include "MassSubsystemAccess.h"
 #include "Processors/TypedElementProcessorAdaptors.h"
 #include "ProfilingDebugging/CpuProfilerTrace.h"
+#include "ProfilingDebugging/CountersTrace.h"
 #include "Queries/TypedElementQueryContext_DirectSingleApi.h"
 #include "Stats/Stats.h"
 #include "TickTaskManagerInterface.h"
@@ -27,6 +28,15 @@
 
 #include UE_INLINE_GENERATED_CPP_BY_NAME(TypedElementDatabase)
 
+#if COUNTERSTRACE_ENABLED
+
+DECLARE_STATS_GROUP(TEXT("Editor Data Storage (Core)"), STATGROUP_TedsCore, STATCAT_Advanced);
+
+DECLARE_FLOAT_COUNTER_STAT(TEXT("Coop task time (ms)"), STAT_TedsCore_CoopTaskTimeMs, STATGROUP_TedsCore);
+TRACE_DECLARE_FLOAT_COUNTER(COUNTER_TedsCore_CoopTaskTimeMs, TEXT("TedsCore/Tick/CoopTaskTimeMs"));
+
+#endif
+
 const FName UEditorDataStorage::TickGroupName_Default(TEXT("Default"));
 const FName UEditorDataStorage::TickGroupName_PreUpdate(TEXT("PreUpdate"));
 const FName UEditorDataStorage::TickGroupName_Update(TEXT("Update"));
@@ -76,19 +86,24 @@ FAutoConsoleCommandWithOutputDevice PrintSupportedColumnsConsoleCommand(
 			Output.Log(TEXT("End of Typed Elements Data Storage supported column list."));
 		}));
 
-static TAutoConsoleVariable<int32> CooperativeTasks_TargetFrameRate(TEXT("TEDS.Features.CooperativeTasks.TargetFrameRate"), 60,
+static TAutoConsoleVariable<int32> CooperativeTasks_TargetFrameRate(TEXT("TEDS.Feature.CooperativeTasks.TargetFrameRate"), 60,
 	TEXT("TEDS runs additional cooperative background tasks with the time that is left at the end of a frame. Adjust this value to set "
 		"how much TEDS considers as available time. Clamped between 16 and 120 fps."));
 
-static TAutoConsoleVariable<float> CooperativeTasks_MinTimePerFrameMs(TEXT("TEDS.Features.CooperativeTasks.MinTimePerFrameMs"), 2.0f,
+static TAutoConsoleVariable<float> CooperativeTasks_MinTimePerFrameMs(TEXT("TEDS.Feature.CooperativeTasks.MinTimePerFrameMs"), 2.0f,
 	TEXT("The minimal amount of time in milliseconds that TEDS will spend guaranteed on cooperative tasks."));
 
-static TAutoConsoleVariable<float> CooperativeTasks_MinTimePerTaskMs(TEXT("TEDS.Features.CooperativeTasks.MinTimePerTaskMs"), 1.0f,
+static TAutoConsoleVariable<float> CooperativeTasks_MinTimePerTaskMs(TEXT("TEDS.Feature.CooperativeTasks.MinTimePerTaskMs"), 1.0f,
 	TEXT("The minimal amount of time in milliseconds that needs to remain for TEDS to consider running another cooperative task."));
 
-static TAutoConsoleVariable<float> CooperativeTasks_ReportThresholdMs(TEXT("TEDS.Features.CooperativeTasks.ReportThresholdMs"), 1.0f,
+static TAutoConsoleVariable<float> CooperativeTasks_ReportThresholdMs(TEXT("TEDS.Feature.CooperativeTasks.ReportThresholdMs"), 1.0f,
 	TEXT("The maximum amount of time in milliseconds that cooperative tasks are allowed to go over their allotted time before they're reported."));
 
+static TAutoConsoleVariable<int32> CooperativeTasks_MaxIdleFrames(TEXT("TEDS.Feature.CooperativeTasks.MaxIdleFrames"), 5,
+	TEXT("The maximum number of frames the cooperative task scheduler will remain idle during low frame rates before starting to call tasks ")
+	TEXT("again with the minimum time per task. This allows the editor to recover from spikes and avoids moving to emergency execution for ")
+	TEXT("smaller interruptions due to for instance OS threading changes, but prevents cooperative tasks from starving."));
+
 namespace UE::Editor::DataStorage::Private
 {
 	struct ColumnsToBitSetsResult
@@ -1443,6 +1458,46 @@ void UEditorDataStorage::TickCooperativeTasks()
 
 		bRunningCooperativeUpdate = true;
 
+		// Don't count the time from last frame spend in cooperative tasks because that's part of the time that needs
+		// to be distributed again.
+		FTimespan LastFrameDuration(
+			FTimespan::FromSeconds(FPlatformTime::ToSeconds64(FPlatformTime::Cycles64() - CooperativeTickStartTime)) + 
+			LastFrameCooperativeTickDuration);
+
+		FTimespan RemainingTime;
+		if (LastFrameDuration < (TargetFrameDuration - MinTimePerTick))
+		{
+			// There's time left to do cooperative tasks
+			RemainingTime = TargetFrameDuration - LastFrameDuration;
+			CooperativeTaskIdleFrameCount = 0;
+		}
+		else
+		{
+			// There's no time left to do work;
+			if (CooperativeTaskIdleFrameCount < CooperativeTasks_MaxIdleFrames.GetValueOnGameThread())
+			{
+#if COUNTERSTRACE_ENABLED
+				SET_FLOAT_STAT(STAT_TedsCore_CoopTaskTimeMs, 0.0);
+				TRACE_COUNTER_SET(COUNTER_TedsCore_CoopTaskTimeMs, 0.0);
+#endif
+				LastFrameCooperativeTickDuration = 0;
+				bRunningCooperativeUpdate = false;
+				CooperativeTickStartTime = FPlatformTime::Cycles64();
+				CooperativeTaskIdleFrameCount++;
+				return;
+			}
+			else
+			{
+				RemainingTime = MinTimePerTick;
+			}
+		}
+
+#if COUNTERSTRACE_ENABLED
+		double RemainingTimeMs = RemainingTime.GetTotalMilliseconds();
+		SET_FLOAT_STAT(STAT_TedsCore_CoopTaskTimeMs, RemainingTimeMs);
+		TRACE_COUNTER_SET(COUNTER_TedsCore_CoopTaskTimeMs, RemainingTimeMs);
+#endif
+		
 		// Update the priorities. This causes all tasks to eventually float to the top so even low priority tasks
 		// get to run on occasion.
 		for (FCooperativeTask& Task : CooperativeTasks)
@@ -1451,11 +1506,6 @@ void UEditorDataStorage::TickCooperativeTasks()
 			Task.bHasRun = false;
 		}
 
-		FTimespan LastFrameDuration(FTimespan::FromSeconds(FPlatformTime::ToSeconds64(FPlatformTime::Cycles64() - CooperativeTickStartTime)));
-		FTimespan RemainingTime = LastFrameDuration < (TargetFrameDuration - MinTimePerTick)
-			? TargetFrameDuration - LastFrameDuration
-			: MinTimePerTick;
-		
 		int32 CurrentTaskIndex = 0;
 		
 		uint64 ProcessingStartTime = FPlatformTime::Cycles64();
@@ -1507,7 +1557,8 @@ void UEditorDataStorage::TickCooperativeTasks()
 				CurrentTaskIndex++;
 			}
 
-			if (FTimespan::FromSeconds(FPlatformTime::ToSeconds64(FPlatformTime::Cycles64() - ProcessingStartTime)) >= RemainingTime)
+			LastFrameCooperativeTickDuration = FTimespan::FromSeconds(FPlatformTime::ToSeconds64(FPlatformTime::Cycles64() - ProcessingStartTime));
+			if (LastFrameCooperativeTickDuration >= RemainingTime)
 			{
 				break;
 			}
diff --git a/Engine/Plugins/Experimental/EditorDataStorage/Source/TedsCore/Private/TypedElementDatabase.h b/Engine/Plugins/Experimental/EditorDataStorage/Source/TedsCore/Private/TypedElementDatabase.h
--- a/Engine/Plugins/Experimental/EditorDataStorage/Source/TedsCore/Private/TypedElementDatabase.h
+++ b/Engine/Plugins/Experimental/EditorDataStorage/Source/TedsCore/Private/TypedElementDatabase.h
@@ -265,7 +265,12 @@ private:
 		bool bHasRun;
 	};
 	TArray<FCooperativeTask> CooperativeTasks;
+	FTimespan LastFrameCooperativeTickDuration;
 	uint64 CooperativeTickStartTime = 0;
+	// The number of frames the cooperative tasks have been idle due to a lack of available time. Once this
+	// reaches a user specified number the cooperative tasks will start processing with a minimal amount of time
+	// to avoid tasks starving.
+	int32 CooperativeTaskIdleFrameCount = 0;
 	bool bRunningCooperativeUpdate = false;
 
 	TSharedPtr<FMassEntityManager> ActiveEditorEntityManager;
