[TEDS] Explicitly flush the command buffer when used in a direct API call to provide consistent behavior.

The command buffer used by Mass in during direct API calls wasn't consistently being flushed because they all ended up in the Entity Manager's shared buffer. As a result some operations happened out of order causing for instance deletes to happen after adds if called deleted from an observer. This would depend on what other functions were called in between and if those calls triggered a flush. To provide consistency direct calls now have their own local command buffer that get flushed explicitly.

#Synced-CL 49564377

diff --git a/Engine/Plugins/Experimental/EditorDataStorage/Source/TedsCore/Private/Processors/TypedElementProcessorAdaptors.cpp b/Engine/Plugins/Experimental/EditorDataStorage/Source/TedsCore/Private/Processors/TypedElementProcessorAdaptors.cpp
--- a/Engine/Plugins/Experimental/EditorDataStorage/Source/TedsCore/Private/Processors/TypedElementProcessorAdaptors.cpp
+++ b/Engine/Plugins/Experimental/EditorDataStorage/Source/TedsCore/Private/Processors/TypedElementProcessorAdaptors.cpp
@@ -815,7 +815,7 @@ namespace UE::Editor::DataStorage
 					}
 					else
 					{ 
-						this->Context.Defer().template PushCommand<FMassDeferredAddCommand>(
+						this->Context.Defer().template PushCommand<FMassDeferredRemoveCommand>(
 							[RemovedColumns](FMassEntityManager& System)
 							{
 								if (System.IsEntityValid(RemovedColumns->Entity) && System.IsEntityBuilt(RemovedColumns->Entity))
@@ -1605,6 +1605,8 @@ UE::Editor::DataStorage::FQueryResult FTypedElementQueryProcessorData::Execute(
 	if (EnumHasAnyFlags(ExecutionFlags, EDirectQueryExecutionFlags::AllowBoundQueries) || !Description.Callback.Function)
 	{
 		FMassExecutionContext Context(EntityManager);
+		Context.SetDeferredCommandBuffer(Environment.GetSharedMassCommandBuffer());
+		Context.SetFlushDeferredCommands(true);
 		auto ExecuteFunction = [&Result, &Callback, &Description, &Environment](FMassExecutionContext& Context)
 			{
 				using namespace Processors::Private;
@@ -1626,6 +1628,8 @@ UE::Editor::DataStorage::FQueryResult FTypedElementQueryProcessorData::Execute(
 			NativeQuery.ForEachEntityChunk(Context, ExecuteFunction);
 		}
 		Context.FlushDeferred();
+		checkf(Environment.GetSharedMassCommandBuffer()->HasPendingCommands() == false,
+			TEXT("Shared Mass command buffer used by the editor data storage wasn't cleared after flushing."));
 	}
 	else
 	{
@@ -1640,7 +1644,7 @@ UE::Editor::DataStorage::FQueryResult FTypedElementQueryProcessorData::Execute(
 	FMassEntityQuery& NativeQuery,
 	FMassEntityManager& EntityManager,
 	UE::Editor::DataStorage::FEnvironment& Environment,
-	UE::Editor::DataStorage::EDirectQueryExecutionFlags ExecutionFlags)
+	UE::Editor::DataStorage::ERunQueryFlags ExecutionFlags)
 {
 	using namespace UE::Editor::DataStorage;
 	using namespace UE::Editor::DataStorage::Queries;
@@ -1648,9 +1652,14 @@ UE::Editor::DataStorage::FQueryResult FTypedElementQueryProcessorData::Execute(
 	FQueryResult Result;
 	Result.Completed = FQueryResult::ECompletion::Fully;
 
-	if (EnumHasAnyFlags(ExecutionFlags, EDirectQueryExecutionFlags::AllowBoundQueries) || !Description.Callback.Function)
+	if (EnumHasAnyFlags(ExecutionFlags, ERunQueryFlags::AllowBoundQueries) || !Description.Callback.Function)
 	{
 		FMassExecutionContext Context(EntityManager);
+		if (EnumHasAnyFlags(ExecutionFlags, ERunQueryFlags::ForceCommandBufferFlush))
+		{
+			Context.SetDeferredCommandBuffer(Environment.GetSharedMassCommandBuffer());
+			Context.SetFlushDeferredCommands(true);
+		}
 		auto ExecuteFunction = [&Result, &Callback, &Environment](FMassExecutionContext& Context)
 			{
 				using namespace Processors::Private;
@@ -1660,10 +1669,10 @@ UE::Editor::DataStorage::FQueryResult FTypedElementQueryProcessorData::Execute(
 				Callback.Call(QueryContext, QueryContext.GetContextImplementation());
 				Result.Count += Context.GetNumEntities();
 			};
-		if (EnumHasAnyFlags(ExecutionFlags, EDirectQueryExecutionFlags::ParallelizeChunks))
+		if (EnumHasAnyFlags(ExecutionFlags, ERunQueryFlags::ParallelizeChunks))
 		{
 			FMassEntityQuery::EParallelExecutionFlags Mode =
-				EnumHasAnyFlags(ExecutionFlags, EDirectQueryExecutionFlags::AutoBalanceParallelChunkProcessing)
+				EnumHasAnyFlags(ExecutionFlags, ERunQueryFlags::AutoBalanceParallelChunkProcessing)
 				? FMassEntityQuery::EParallelExecutionFlags::AutoBalance
 				: FMassEntityQuery::EParallelExecutionFlags::Default;
 			NativeQuery.ParallelForEachEntityChunk(Context, ExecuteFunction, Mode);
@@ -1672,7 +1681,13 @@ UE::Editor::DataStorage::FQueryResult FTypedElementQueryProcessorData::Execute(
 		{
 			NativeQuery.ForEachEntityChunk(Context, ExecuteFunction);
 		}
-		Context.FlushDeferred();
+		
+		if (EnumHasAnyFlags(ExecutionFlags, ERunQueryFlags::ForceCommandBufferFlush))
+		{
+			Context.FlushDeferred();
+			checkf(Environment.GetSharedMassCommandBuffer()->HasPendingCommands() == false,
+				TEXT("Shared Mass command buffer used by the editor data storage wasn't cleared after flushing."));
+		}
 	}
 	else
 	{
@@ -1711,6 +1726,7 @@ UE::Editor::DataStorage::FQueryResult FTypedElementQueryProcessorData::Execute(
 			Result.Count += Context.GetNumEntities();
 		}
 	);
+
 	return Result;
 }
 
diff --git a/Engine/Plugins/Experimental/EditorDataStorage/Source/TedsCore/Private/Processors/TypedElementProcessorAdaptors.h b/Engine/Plugins/Experimental/EditorDataStorage/Source/TedsCore/Private/Processors/TypedElementProcessorAdaptors.h
--- a/Engine/Plugins/Experimental/EditorDataStorage/Source/TedsCore/Private/Processors/TypedElementProcessorAdaptors.h
+++ b/Engine/Plugins/Experimental/EditorDataStorage/Source/TedsCore/Private/Processors/TypedElementProcessorAdaptors.h
@@ -66,7 +66,7 @@ struct FTypedElementQueryProcessorData
 		FMassEntityQuery& NativeQuery,
 		FMassEntityManager& EntityManager,
 		UE::Editor::DataStorage::FEnvironment& Environment,
-		UE::Editor::DataStorage::EDirectQueryExecutionFlags ExecutionFlags);
+		UE::Editor::DataStorage::ERunQueryFlags ExecutionFlags);
 	static UE::Editor::DataStorage::FQueryResult Execute(
 		UE::Editor::DataStorage::SubqueryCallbackRef& Callback,
 		UE::Editor::DataStorage::FQueryDescription& Description,
diff --git a/Engine/Plugins/Experimental/EditorDataStorage/Source/TedsCore/Private/Queries/TypedElementExtendedQueryStore.cpp b/Engine/Plugins/Experimental/EditorDataStorage/Source/TedsCore/Private/Queries/TypedElementExtendedQueryStore.cpp
--- a/Engine/Plugins/Experimental/EditorDataStorage/Source/TedsCore/Private/Queries/TypedElementExtendedQueryStore.cpp
+++ b/Engine/Plugins/Experimental/EditorDataStorage/Source/TedsCore/Private/Queries/TypedElementExtendedQueryStore.cpp
@@ -378,13 +378,13 @@ namespace UE::Editor::DataStorage
 		return Result;
 	}
 
-	template<typename CallbackReference>
+	template<typename CallbackReference, typename ExecutionFlagsType>
 	FQueryResult FExtendedQueryStore::RunQueryCallbackCommon(
 		FMassEntityManager& EntityManager,
 		FEnvironment& Environment,
 		FMassExecutionContext* ParentContext,
 		Handle Query,
-		EDirectQueryExecutionFlags DirectExecutionFlags,
+		ExecutionFlagsType ExecutionFlags,
 		CallbackReference Callback)
 	{
 		FQueryResult Result;
@@ -400,15 +400,19 @@ namespace UE::Editor::DataStorage
 				{
 					if constexpr (std::is_same_v<CallbackReference, DirectQueryCallbackRef>)
 					{
+						static_assert(std::is_same_v<ExecutionFlagsType, EDirectQueryExecutionFlags>, 
+							"Running a direct query requires 'EDirectQueryExecutionFlags' flags.");
 						Result = FTypedElementQueryProcessorData::Execute(
-							Callback, QueryData->Description, QueryData->NativeQuery, EntityManager, Environment, DirectExecutionFlags);
+							Callback, QueryData->Description, QueryData->NativeQuery, EntityManager, Environment, ExecutionFlags);
 					}
 					else if constexpr (std::is_same_v<CallbackReference, Queries::TQueryFunction<void>>)
 					{
 						if (CheckCompatibility(QueryData->Description, Callback))
 						{
+							static_assert(std::is_same_v<ExecutionFlagsType, ERunQueryFlags>, 
+								"Calling run query requires 'ERunQueryFlags' flags.");
 							Result = FTypedElementQueryProcessorData::Execute(
-								Callback, QueryData->Description, QueryData->NativeQuery, EntityManager, Environment, DirectExecutionFlags);
+								Callback, QueryData->Description, QueryData->NativeQuery, EntityManager, Environment, ExecutionFlags);
 						}
 						else
 						{
@@ -455,7 +459,7 @@ namespace UE::Editor::DataStorage
 		FEnvironment& Environment, Handle Query, ERunQueryFlags Flags,
 		const Queries::TQueryFunction<void>& Callback)
 	{
-		return RunQueryCallbackCommon(EntityManager, Environment, nullptr, Query, static_cast<EDirectQueryExecutionFlags>(Flags), Callback);
+		return RunQueryCallbackCommon(EntityManager, Environment, nullptr, Query, Flags, Callback);
 	}
 	
 	FQueryResult FExtendedQueryStore::RunQuery(FMassEntityManager& EntityManager,
diff --git a/Engine/Plugins/Experimental/EditorDataStorage/Source/TedsCore/Private/Queries/TypedElementExtendedQueryStore.h b/Engine/Plugins/Experimental/EditorDataStorage/Source/TedsCore/Private/Queries/TypedElementExtendedQueryStore.h
--- a/Engine/Plugins/Experimental/EditorDataStorage/Source/TedsCore/Private/Queries/TypedElementExtendedQueryStore.h
+++ b/Engine/Plugins/Experimental/EditorDataStorage/Source/TedsCore/Private/Queries/TypedElementExtendedQueryStore.h
@@ -187,13 +187,13 @@ namespace UE::Editor::DataStorage
 			EExecutionMode ExecutionMode = EExecutionMode::Default;
 		};
 
-		template<typename CallbackReference>
+		template<typename CallbackReference, typename ExecutionFlagsType>
 		FQueryResult RunQueryCallbackCommon(
 			FMassEntityManager& EntityManager,
 			FEnvironment& Environment,
 			FMassExecutionContext* ParentContext,
 			Handle Query,
-			EDirectQueryExecutionFlags ExecutionFlags,
+			ExecutionFlagsType ExecutionFlags,
 			CallbackReference Callback);
 
 		FMassEntityQuery& SetupNativeQuery(FQueryDescription& Query, FExtendedQuery& StoredQuery, FEnvironment& Environment);
diff --git a/Engine/Plugins/Experimental/EditorDataStorage/Source/TedsCore/Private/TypedElementDatabase.cpp b/Engine/Plugins/Experimental/EditorDataStorage/Source/TedsCore/Private/TypedElementDatabase.cpp
--- a/Engine/Plugins/Experimental/EditorDataStorage/Source/TedsCore/Private/TypedElementDatabase.cpp
+++ b/Engine/Plugins/Experimental/EditorDataStorage/Source/TedsCore/Private/TypedElementDatabase.cpp
@@ -1332,7 +1332,7 @@ UE::Editor::DataStorage::FQueryResult UEditorDataStorage::RunQuery(
 	if (ActiveEditorEntityManager)
 	{
 		return Environment->GetQueryStore().RunQuery(*ActiveEditorEntityManager, *Environment, 
-			FExtendedQueryStore::Handle(Query), InFlags, Callback);
+			FExtendedQueryStore::Handle(Query), InFlags | ERunQueryFlags::ForceCommandBufferFlush, Callback);
 	}
 	else
 	{
diff --git a/Engine/Plugins/Experimental/EditorDataStorage/Source/TedsCore/Private/TypedElementDatabaseEnvironment.cpp b/Engine/Plugins/Experimental/EditorDataStorage/Source/TedsCore/Private/TypedElementDatabaseEnvironment.cpp
--- a/Engine/Plugins/Experimental/EditorDataStorage/Source/TedsCore/Private/TypedElementDatabaseEnvironment.cpp
+++ b/Engine/Plugins/Experimental/EditorDataStorage/Source/TedsCore/Private/TypedElementDatabaseEnvironment.cpp
@@ -1,6 +1,8 @@
 // Copyright Epic Games, Inc. All Rights Reserved.
 
 #include "TypedElementDatabaseEnvironment.h"
+
+#include "MassCommandBuffer.h"
 #include "TypedElementDatabase.h"
 
 namespace UE::Editor::DataStorage
@@ -15,6 +17,7 @@ namespace UE::Editor::DataStorage
 		, ValueTagManager(DynamicColumnGenerator)
 		, MassEntityManager(InMassEntityManager)
 		, MassPhaseManager(InMassPhaseManager)
+		, SharedMassCommandBuffer(MakeShared<FMassCommandBuffer>())
 	{
 		DynamicColumnGenerator.SetQueryStore(Queries);
 	}
@@ -109,6 +112,16 @@ namespace UE::Editor::DataStorage
 		return MassPhaseManager;
 	}
 
+	TSharedPtr<FMassCommandBuffer> FEnvironment::GetSharedMassCommandBuffer()
+	{
+		return SharedMassCommandBuffer;
+	}
+
+	TSharedPtr<const FMassCommandBuffer> FEnvironment::GetSharedMassCommandBuffer() const
+	{
+		return SharedMassCommandBuffer;
+	}
+
 	const UScriptStruct* FEnvironment::FindDynamicColumn(const UScriptStruct& Template, const FName& Identifier) const
 	{
 		return DynamicColumnGenerator.FindByTemplateId(Template, Identifier);
diff --git a/Engine/Plugins/Experimental/EditorDataStorage/Source/TedsCore/Private/TypedElementDatabaseEnvironment.h b/Engine/Plugins/Experimental/EditorDataStorage/Source/TedsCore/Private/TypedElementDatabaseEnvironment.h
--- a/Engine/Plugins/Experimental/EditorDataStorage/Source/TedsCore/Private/TypedElementDatabaseEnvironment.h
+++ b/Engine/Plugins/Experimental/EditorDataStorage/Source/TedsCore/Private/TypedElementDatabaseEnvironment.h
@@ -6,6 +6,7 @@
 #include "DynamicColumnGenerator.h"
 #include "MassEntityManager.h"
 #include "MassProcessingPhaseManager.h"
+#include "Templates/SharedPointer.h"
 #include "TypedElementDatabaseCommandBuffer.h"
 #include "TypedElementDatabaseScratchBuffer.h"
 #include "TypedElementDatabaseIndexTable.h"
@@ -14,6 +15,7 @@
 #include "Queries/TypedElementExtendedQueryStore.h"
 
 class UEditorDataStorage;
+struct FMassCommandBuffer;
 
 namespace UE::Editor::DataStorage
 {
@@ -51,6 +53,9 @@ namespace UE::Editor::DataStorage
 		FMassProcessingPhaseManager& GetMassPhaseManager();
 		const FMassProcessingPhaseManager& GetMassPhaseManager() const;
 
+		TSharedPtr<FMassCommandBuffer> GetSharedMassCommandBuffer();
+		TSharedPtr<const FMassCommandBuffer> GetSharedMassCommandBuffer() const;
+
 		// Finds the type information for the dynamic column
 		// Dynamic columns are specified by a template layout and a FName Identifier
 		const UScriptStruct* FindDynamicColumn(const UScriptStruct& Template, const FName& Identifier) const;
@@ -98,6 +103,7 @@ namespace UE::Editor::DataStorage
 
 		FMassEntityManager& MassEntityManager;
 		FMassProcessingPhaseManager& MassPhaseManager;
+		TSharedPtr<FMassCommandBuffer> SharedMassCommandBuffer;
 
 		uint64 UpdateCycleId = 0;
 	};
diff --git a/Engine/Source/Runtime/TypedElementFramework/Public/Elements/Interfaces/TypedElementQueryStorageInterfaces.h b/Engine/Source/Runtime/TypedElementFramework/Public/Elements/Interfaces/TypedElementQueryStorageInterfaces.h
--- a/Engine/Source/Runtime/TypedElementFramework/Public/Elements/Interfaces/TypedElementQueryStorageInterfaces.h
+++ b/Engine/Source/Runtime/TypedElementFramework/Public/Elements/Interfaces/TypedElementQueryStorageInterfaces.h
@@ -343,7 +343,8 @@ namespace UE::Editor::DataStorage
 	ENUM_CLASS_FLAGS(EDirectQueryExecutionFlags);
 
 	// For now this is an exact mirror of EDirectQueryExecutionFlags (except deprecated flags) in order to separate the old and the new 
-	// query callbacks/context. Please keep both in sync until the old query callbacks/context has been fully deprecated.
+	// query callbacks/context. Please keep both in sync, except for new entries, until the old query callbacks/context has been fully 
+	// deprecated.
 	enum class ERunQueryFlags : uint32
 	{
 		Default = 0, //< No settings, use the default behavior.
@@ -356,6 +357,8 @@ namespace UE::Editor::DataStorage
 		AutoBalanceParallelChunkProcessing = 1 << 1,
 		IgnoreActiveState = 1 << 2, //< If set, a direct call will not check activatable queries if they're set.
 		AllowBoundQueries = 1 << 3, //< Stops checking if queries are bound when set, otherwise bound queries will assert.
+		/** Force commands added to the query callback's context or parent context to execute immediately after completing. */
+		ForceCommandBufferFlush = 1 << 4,
 	};
 	ENUM_CLASS_FLAGS(ERunQueryFlags);
 
