[TEDS] Mapping performance update.

At large scales (100k+) the reverse lookup in TEDS Compat started to have performance issues due to the multimap lookups. This revision makes two changes:
- Only use a multimap if there's a hash conflict. A multimap is slower than a regular map, but the multimap is only needed if there's a hash collision which isn't the case in the vast majority of cases. The updated version uses a regular map unless there's a collision, in which case it puts a breadcrumb in the regular map and uses the multimap.
- Shard the map. Since TMap has a flat array as part of its implementation resizing starts to be a significant cost at large scales so the array has been split into 256 smaller shards. The 8 least significant bits are used to find the appropriate map and put the values there which distributes the cost and reduces the time needed for resizing.

#Synced-CL 49564377

diff --git a/Engine/Plugins/Experimental/EditorDataStorage/Source/TedsCore/Private/TypedElementDatabaseIndexTable.cpp b/Engine/Plugins/Experimental/EditorDataStorage/Source/TedsCore/Private/TypedElementDatabaseIndexTable.cpp
--- a/Engine/Plugins/Experimental/EditorDataStorage/Source/TedsCore/Private/TypedElementDatabaseIndexTable.cpp
+++ b/Engine/Plugins/Experimental/EditorDataStorage/Source/TedsCore/Private/TypedElementDatabaseIndexTable.cpp
@@ -64,7 +64,7 @@ namespace UE::Editor::DataStorage
 		FScopedExclusiveLock Lock(LockScope);
 
 		FDomain& DomainInfo = Domains.FindOrAdd(Domain);
-		DomainInfo.IndexLookupMap.Reserve(DomainInfo.IndexLookupMap.Num() + MapRowPairs.Num());
+		DomainInfo.IndexLookupMap.ReserveAdditional(MapRowPairs);
 
 		checkf(DomainInfo.Rows.Num() >= DomainInfo.FreeList.Num(),
 			TEXT("There can't be less rows than there are rows stored in the free list as free list is a subset of rows."));
@@ -84,17 +84,7 @@ namespace UE::Editor::DataStorage
 
 		if (FDomain* DomainInfo = Domains.Find(Domain))
 		{
-			uint64 Hash = OriginalKey.CalculateHash();
-			for (FDomain::IndexLookupMapType::TKeyIterator It = DomainInfo->IndexLookupMap.CreateKeyIterator(Hash); It; ++It)
-			{
-				if (DomainInfo->Keys[It->Value] == OriginalKey)
-				{
-					int32 Index = It->Value;
-					It.RemoveCurrent();
-					DomainInfo->IndexLookupMap.Add(NewKey.CalculateHash(), Index);
-					DomainInfo->Keys[Index] = MoveTemp(NewKey);
-				}
-			}
+			DomainInfo->RemapUnguarded(DataStorage, OriginalKey, MoveTemp(NewKey), IsDirty());
 		}
 	}
 
@@ -104,18 +94,7 @@ namespace UE::Editor::DataStorage
 		
 		if (FDomain* DomainInfo = Domains.Find(Domain))
 		{
-			uint64 Hash = Key.CalculateHash();
-			for (FDomain::IndexLookupMapType::TKeyIterator It = DomainInfo->IndexLookupMap.CreateKeyIterator(Hash); It; ++It)
-			{
-				if (DomainInfo->Keys[It->Value] == Key)
-				{
-					int32 Index = It->Value;
-					DomainInfo->Rows[Index] = InvalidRowHandle;
-					DomainInfo->Keys[Index].Clear();
-					DomainInfo->FreeList.PushFirst(Index);
-					It.RemoveCurrent();
-				}
-			}
+			DomainInfo->RemoveUnguarded(Key);
 		}
 	}
 
@@ -239,38 +218,270 @@ namespace UE::Editor::DataStorage
 		return RemainingUnverifiedDomains > 0 || DirtyDueToRemoval;
 	}
 
-	int32 FMappingTable::FDomain::FindIndexUnguarded(const FMapKey& Key) const
+
+	//
+	// FCollisionSafeHashMap
+	//
+
+	
+	template<typename KeyType>
+	int32* FMappingTable::FCollisionSafeHashMap::FindTemplate(uint64 Hash, const KeyType& Key, const TArray<FMapKey>& Keys)
+	{
+		if (int32* Index = Map.Find(Hash))
+		{
+			if (*Index != HashCollisionId)
+			{
+				return Index;
+			}
+			else
+			{
+				for (IndexLookupMultiMapType::TKeyIterator It = MapBackup.CreateKeyIterator(Hash); It; ++It)
+				{
+					if (Keys[It->Value] == Key)
+					{
+						return &(It->Value);
+					}
+				}
+			}
+		}
+		return nullptr;
+	}
+
+	const int32* FMappingTable::FCollisionSafeHashMap::Find(uint64 Hash, const FMapKey& Key, const TArray<FMapKey>& Keys) const
 	{
-		return FindIndexUnguarded(Key.CalculateHash(), Key);
+		return const_cast<FMappingTable::FCollisionSafeHashMap*>(this)->FindTemplate(Hash, Key, Keys);
 	}
 
-	int32 FMappingTable::FDomain::FindIndexUnguarded(uint64 Hash, const FMapKey& Key) const
+	const int32* FMappingTable::FCollisionSafeHashMap::Find(uint64 Hash, const FMapKeyView& Key, const TArray<FMapKey>& Keys) const
 	{
-		for (IndexLookupMapType::TConstKeyIterator It = IndexLookupMap.CreateConstKeyIterator(Hash); It; ++It)
+		return const_cast<FMappingTable::FCollisionSafeHashMap*>(this)->FindTemplate(Hash, Key, Keys);
+	}
+
+	int32* FMappingTable::FCollisionSafeHashMap::Find(uint64 Hash, const FMapKey& Key, const TArray<FMapKey>& Keys)
+	{
+		return FindTemplate(Hash, Key, Keys);
+	}
+
+	int32* FMappingTable::FCollisionSafeHashMap::Find(uint64 Hash, const FMapKeyView& Key, const TArray<FMapKey>& Keys)
+	{
+		return FindTemplate(Hash, Key, Keys);
+	}
+
+	void FMappingTable::FCollisionSafeHashMap::ReserveAdditional(int32 Count)
+	{
+		// Collisions are rare so assume that all requested entries go into the main map.
+		Map.Reserve(Map.Num() + Count);
+	}
+
+	void FMappingTable::FCollisionSafeHashMap::Add(uint64 Hash, int32 Index)
+	{
+		int32* StoredIndex = Map.Find(Hash);
+		if (StoredIndex == nullptr)
 		{
-			if (Keys[It->Value] == Key)
+			Map.Add(Hash, Index);
+		}
+		else
+		{
+			MapBackup.Add(Hash, Index);
+			// If this is the first conflict, move the index from the main map into the backup map and
+			// leave a breadcrumb to indicate the backup map needs to be searched.
+			if (*StoredIndex != HashCollisionId)
 			{
-				return It->Value;
+				MapBackup.Add(Hash, *StoredIndex);
+				*StoredIndex = HashCollisionId;
 			}
 		}
-		return -1;
 	}
 
-	int32 FMappingTable::FDomain::FindIndexUnguarded(const FMapKeyView& Key) const
+	template<typename KeyType>
+	int32 FMappingTable::FCollisionSafeHashMap::RemoveTemplate(uint64 Hash, const KeyType& Key, const TArray<FMapKey>& Keys)
 	{
-		return FindIndexUnguarded(Key.CalculateHash(), Key);
+		if (IndexLookupMapType::TKeyIterator It = Map.CreateKeyIterator(Hash))
+		{
+			if (It->Value != HashCollisionId)
+			{
+				int32 Index = It->Value;
+				It.RemoveCurrent();
+				return Index;
+			}
+			else
+			{
+				int32 Index = InvalidId;
+				int32 EntriesLeft = 0;
+				for (IndexLookupMultiMapType::TKeyIterator ItBackup = MapBackup.CreateKeyIterator(Hash);
+					ItBackup; ++ItBackup)
+				{
+					if (Keys[ItBackup->Value] == Key)
+					{
+						Index = ItBackup->Value;
+						ItBackup.RemoveCurrent();
+					}
+					else
+					{
+						++EntriesLeft;
+					}
+				}
+				// If there's exactly one key left there are no more collisions so switch back to the main map instead of the backup
+				// multi-map.
+				checkf(EntriesLeft > 0, TEXT("Only one entry can be removed and there should never be just one entry in the backup map."));
+				if (EntriesLeft == 1)
+				{
+					It->Value = MapBackup.FindChecked(Hash);
+					MapBackup.Remove(Hash);
+				}
+				return Index;
+			}
+		}
+		return InvalidId;
 	}
 
-	int32 FMappingTable::FDomain::FindIndexUnguarded(uint64 Hash, const FMapKeyView& Key) const
+	int32 FMappingTable::FCollisionSafeHashMap::Remove(uint64 Hash, const FMapKey& Key, const TArray<FMapKey>& Keys)
+	{
+		return RemoveTemplate(Hash, Key, Keys);
+	}
+
+	int32 FMappingTable::FCollisionSafeHashMap::Remove(uint64 Hash, const FMapKeyView& Key, const TArray<FMapKey>& Keys)
+	{
+		return RemoveTemplate(Hash, Key, Keys);
+	}
+
+	void FMappingTable::FCollisionSafeHashMap::Remove(uint64 Hash, int32 Index)
+	{
+		if (IndexLookupMapType::TKeyIterator It = Map.CreateKeyIterator(Hash))
+		{
+			if (It->Value != HashCollisionId)
+			{
+				It.RemoveCurrent();
+			}
+			else
+			{
+				int32 EntriesLeft = 0;
+				for (IndexLookupMultiMapType::TKeyIterator ItBackup = MapBackup.CreateKeyIterator(Hash);
+					ItBackup; ++ItBackup)
+				{
+					if (ItBackup->Value == Index)
+					{
+						ItBackup.RemoveCurrent();
+					}
+					else
+					{
+						++EntriesLeft;
+					}
+				}
+				// If there's exactly one key left there are no more collisions so switch back to the main map instead of the backup
+				// multi-map.
+				checkf(EntriesLeft > 0, TEXT("Only one entry can be removed and there should never be just one entry in the backup map."));
+				if (EntriesLeft == 1)
+				{
+					It->Value = MapBackup.FindChecked(Hash);
+					MapBackup.Remove(Hash);
+				}
+			}
+		}
+	}
+
+
+	// 
+	// FShardedMap
+	//
+
+	int32 FMappingTable::FShardedMap::GetShard(uint64 Hash) const
+	{
+		return Hash & ((1 << ShardBits) - 1);
+	}
+
+	uint64 FMappingTable::FShardedMap::ModifyHash(uint64 Hash) const
+	{
+		return Hash >> ShardBits;
+	}
+
+	const int32* FMappingTable::FShardedMap::Find(uint64 Hash, const FMapKey& Key, const TArray<FMapKey>& Keys) const
+	{
+		return Shards[GetShard(Hash)].Find(ModifyHash(Hash), Key, Keys);
+	}
+
+	const int32* FMappingTable::FShardedMap::Find(uint64 Hash, const FMapKeyView& Key, const TArray<FMapKey>& Keys) const
+	{
+		return Shards[GetShard(Hash)].Find(ModifyHash(Hash), Key, Keys);
+	}
+
+	int32* FMappingTable::FShardedMap::Find(uint64 Hash, const FMapKey& Key, const TArray<FMapKey>& Keys)
+	{
+		return Shards[GetShard(Hash)].Find(ModifyHash(Hash), Key, Keys);
+	}
+
+	int32* FMappingTable::FShardedMap::Find(uint64 Hash, const FMapKeyView& Key, const TArray<FMapKey>& Keys)
+	{
+		return Shards[GetShard(Hash)].Find(ModifyHash(Hash), Key, Keys);
+	}
+
+	void FMappingTable::FShardedMap::ReserveAdditional(TArrayView<TPair<FMapKey, RowHandle>> MapRowPairs)
 	{
-		for (IndexLookupMapType::TConstKeyIterator It = IndexLookupMap.CreateConstKeyIterator(Hash); It; ++It)
+		int32 Sizes[1 << ShardBits] = {};
+
+		for (TPair<FMapKey, RowHandle>& It : MapRowPairs)
 		{
-			if (Keys[It->Value] == Key)
+			Sizes[GetShard(It.Key.CalculateHash())]++;
+		}
+
+		for (int32 ShardSizeIndex = 0; ShardSizeIndex < (1 << ShardBits); ++ShardSizeIndex)
+		{
+			int32 Count = Sizes[ShardSizeIndex];
+			if (Count > 0)
 			{
-				return It->Value;
+				Shards[ShardSizeIndex].ReserveAdditional(Count);
 			}
 		}
-		return -1;
+	}
+
+	void FMappingTable::FShardedMap::Add(uint64 Hash, int32 Index)
+	{
+		Shards[GetShard(Hash)].Add(ModifyHash(Hash), Index);
+	}
+
+	int32 FMappingTable::FShardedMap::Remove(uint64 Hash, const FMapKey& Key, const TArray<FMapKey>& Keys)
+	{
+		return Shards[GetShard(Hash)].Remove(ModifyHash(Hash), Key, Keys);
+	}
+
+	int32 FMappingTable::FShardedMap::Remove(uint64 Hash, const FMapKeyView& Key, const TArray<FMapKey>& Keys)
+	{
+		return Shards[GetShard(Hash)].Remove(ModifyHash(Hash), Key, Keys);
+	}
+
+	void FMappingTable::FShardedMap::Remove(uint64 Hash, int32 Index)
+	{
+		return Shards[GetShard(Hash)].Remove(ModifyHash(Hash), Index);
+	}
+
+
+	//
+	// FDomain
+	//
+
+
+	int32 FMappingTable::FDomain::FindIndexUnguarded(const FMapKey& Key) const
+	{
+		const int32* Result = IndexLookupMap.Find(Key.CalculateHash(), Key, Keys);
+		return Result ? *Result : InvalidId;
+	}
+
+	int32 FMappingTable::FDomain::FindIndexUnguarded(uint64 Hash, const FMapKey& Key) const
+	{
+		const int32* Result = IndexLookupMap.Find(Hash, Key, Keys);
+		return Result ? *Result : InvalidId;
+	}
+
+	int32 FMappingTable::FDomain::FindIndexUnguarded(const FMapKeyView& Key) const
+	{
+		const int32* Result = IndexLookupMap.Find(Key.CalculateHash(), Key, Keys);
+		return Result ? *Result : InvalidId;
+	}
+
+	int32 FMappingTable::FDomain::FindIndexUnguarded(uint64 Hash, const FMapKeyView& Key) const
+	{
+		const int32* Result = IndexLookupMap.Find(Hash, Key, Keys);
+		return Result ? *Result : InvalidId;
 	}
 
 	void FMappingTable::FDomain::IndexRowUnguarded(UEditorDataStorage& InDataStorage, FMapKey&& Key, RowHandle Row, bool bIsDirty)
@@ -300,9 +511,35 @@ namespace UE::Editor::DataStorage
 			Rows[RowIndex] = Row;
 			Keys[RowIndex] = MoveTemp(Key);
 		}
+		
 		IndexLookupMap.Add(Hash, RowIndex);
 	}
 
+	int32 FMappingTable::FDomain::RemoveUnguarded(const FMapKeyView& Key)
+	{
+		uint64 Hash = Key.CalculateHash();
+		if (int32 Index = IndexLookupMap.Remove(Hash, Key, Keys); Index != InvalidId)
+		{
+			Rows[Index] = InvalidRowHandle;
+			Keys[Index].Clear();
+			FreeList.PushFirst(Index);
+			return Index;
+		}
+		return InvalidId;
+	}
+
+	void FMappingTable::FDomain::RemapUnguarded(UEditorDataStorage& InDataStorage, const FMapKeyView& OriginalKey, FMapKey NewKey, bool bIsDirty)
+	{
+		if (int32 Index = RemoveUnguarded(OriginalKey); Index != InvalidId)
+		{
+			IndexRowUnguarded(InDataStorage, MoveTemp(NewKey), Index, bIsDirty);
+		}
+		else
+		{
+			checkf(false, TEXT("Unable to remove key '%s' from TEDS Mapping."), *OriginalKey.ToString());
+		}
+	}
+
 	void FMappingTable::FDomain::InspectRowBlockForCleanUp(UEditorDataStorage& InDataStorage, int32 Block)
 	{
 		TEDS_EVENT_SCOPE("Mapping inspect rows");
@@ -371,7 +608,8 @@ namespace UE::Editor::DataStorage
 	{
 		FMapKey& Key = Keys[Index];
 		uint64 Hash = Key.CalculateHash();
-		IndexLookupMap.RemoveSingle(Hash, Index);
+		
+		IndexLookupMap.Remove(Hash, Index);
 		
 		Rows[Index] = InvalidRowHandle;
 		Keys[Index].Clear();
diff --git a/Engine/Plugins/Experimental/EditorDataStorage/Source/TedsCore/Private/TypedElementDatabaseIndexTable.h b/Engine/Plugins/Experimental/EditorDataStorage/Source/TedsCore/Private/TypedElementDatabaseIndexTable.h
--- a/Engine/Plugins/Experimental/EditorDataStorage/Source/TedsCore/Private/TypedElementDatabaseIndexTable.h
+++ b/Engine/Plugins/Experimental/EditorDataStorage/Source/TedsCore/Private/TypedElementDatabaseIndexTable.h
@@ -92,25 +92,75 @@ namespace UE::Editor::DataStorage
 			// The number of jobs that were not completed in the previous frame.
 			int32 RemainingJobs = 0;
 		};
+
+		static constexpr int32 InvalidId = -1;
+		static constexpr int32 HashCollisionId = -2;
+
+		struct FCollisionSafeHashMap
+		{
+			using IndexLookupMultiMapType = TMultiMap<uint64, int32>;
+			using IndexLookupMapType = TMap<uint64, int32>;
+
+			IndexLookupMultiMapType MapBackup; // Backup map in case there's a valid hash collision.
+			IndexLookupMapType Map;
+
+			const int32* Find(uint64 Hash, const FMapKey& Key, const TArray<FMapKey>& Keys) const;
+			const int32* Find(uint64 Hash, const FMapKeyView& Key, const TArray<FMapKey>& Keys) const;
+			int32* Find(uint64 Hash, const FMapKey& Key, const TArray<FMapKey>& Keys);
+			int32* Find(uint64 Hash, const FMapKeyView& Key, const TArray<FMapKey>& Keys);
+
+			void ReserveAdditional(int32 Count);
+			void Add(uint64 Hash, int32 Index);
+			int32 Remove(uint64 Hash, const FMapKey& Key, const TArray<FMapKey>& Keys);
+			int32 Remove(uint64 Hash, const FMapKeyView& Key, const TArray<FMapKey>& Keys);
+			void Remove(uint64 Hash, int32 Index);
+
+		private:
+			template<typename KeyType>
+			int32* FindTemplate(uint64 Hash, const KeyType& Key, const TArray<FMapKey>& Keys);
+			template<typename KeyType>
+			int32 RemoveTemplate(uint64 Hash, const KeyType& Key, const TArray<FMapKey>& Keys);
+		};
+
+		struct FShardedMap final
+		{
+			static constexpr int32 ShardBits = 8;
+
+			FCollisionSafeHashMap Shards[1 << ShardBits];
+
+			const int32* Find(uint64 Hash, const FMapKey& Key, const TArray<FMapKey>& Keys) const;
+			const int32* Find(uint64 Hash, const FMapKeyView& Key, const TArray<FMapKey>& Keys) const;
+			int32* Find(uint64 Hash, const FMapKey& Key, const TArray<FMapKey>& Keys);
+			int32* Find(uint64 Hash, const FMapKeyView& Key, const TArray<FMapKey>& Keys);
+
+			void ReserveAdditional(TArrayView<TPair<FMapKey, RowHandle>> MapRowPairs);
+			void Add(uint64 Hash, int32 Index);
+			int32 Remove(uint64 Hash, const FMapKey& Key, const TArray<FMapKey>& Keys);
+			int32 Remove(uint64 Hash, const FMapKeyView& Key, const TArray<FMapKey>& Keys);
+			void Remove(uint64 Hash, int32 Index);
+
+		private:
+			int32 GetShard(uint64 Hash) const;
+			uint64 ModifyHash(uint64 Hash) const;
+		};
 		
-		struct FDomain
+		struct FDomain final
 		{
-			using IndexLookupMapType = TMultiMap<uint64, int32>;
-			
-			IndexLookupMapType IndexLookupMap;
+			FShardedMap IndexLookupMap;
 			TArray<RowHandle> Rows;
 			TArray<FMapKey> Keys;
 			TDeque<int32> FreeList;
 
 			TUniquePtr<FCleanUpInfo> CleanUpInfo = MakeUnique<FCleanUpInfo>();
 
-
 			int32 FindIndexUnguarded(const FMapKey& Key) const;
 			int32 FindIndexUnguarded(uint64 Hash, const FMapKey& Key) const;
 			int32 FindIndexUnguarded(const FMapKeyView& Key) const;
 			int32 FindIndexUnguarded(uint64 Hash, const FMapKeyView& Key) const;
 
 			void IndexRowUnguarded(UEditorDataStorage& DataStorage, FMapKey&& Key, RowHandle Row, bool bIsDirty);
+			int32 RemoveUnguarded(const FMapKeyView& Key);
+			void RemapUnguarded(UEditorDataStorage& InDataStorage, const FMapKeyView& OriginalKey, FMapKey NewKey, bool bIsDirty);
 
 			enum class ERemoveInvalidRowsResult
 			{
