Re-enable AutoRTFM for editor-only code that is hit from transactions.

#Synced-CL 49723230

diff --git a/Engine/Plugins/Editor/AssetManagerEditor/Source/AssetManagerEditor/AssetManagerEditor.Build.cs b/Engine/Plugins/Editor/AssetManagerEditor/Source/AssetManagerEditor/AssetManagerEditor.Build.cs
--- a/Engine/Plugins/Editor/AssetManagerEditor/Source/AssetManagerEditor/AssetManagerEditor.Build.cs
+++ b/Engine/Plugins/Editor/AssetManagerEditor/Source/AssetManagerEditor/AssetManagerEditor.Build.cs
@@ -56,5 +56,9 @@ public class AssetManagerEditor : ModuleRules
 				"Projects",
 			}
 		);
+
+		// This module needs AutoRTFM so disable the auto disable since its in a plugin.
+		// The delegate `PackageMarkedDirtyEvent` on `UPackage` is being used.
+		bDisableAutoRTFMInstrumentation = false;
 	}
 }
diff --git a/Engine/Plugins/Experimental/EditorDataStorage/Source/TedsCore/TedsCore.Build.cs b/Engine/Plugins/Experimental/EditorDataStorage/Source/TedsCore/TedsCore.Build.cs
--- a/Engine/Plugins/Experimental/EditorDataStorage/Source/TedsCore/TedsCore.Build.cs
+++ b/Engine/Plugins/Experimental/EditorDataStorage/Source/TedsCore/TedsCore.Build.cs
@@ -34,5 +34,9 @@ public class TedsCore : ModuleRules
 			PrivateDependencyModuleNames.AddRange(new string[] {});
 			DynamicallyLoadedModuleNames.AddRange(new string[] {});
 		}
+
+		// This module needs AutoRTFM so disable the auto disable since its in a plugin.
+		// The call to `ICompatibilityProvider::ICompatibilityProvider` calls into here.
+		bDisableAutoRTFMInstrumentation = false;
 	}
 }
diff --git a/Engine/Plugins/PCG/Source/PCGEditor/PCGEditor.Build.cs b/Engine/Plugins/PCG/Source/PCGEditor/PCGEditor.Build.cs
--- a/Engine/Plugins/PCG/Source/PCGEditor/PCGEditor.Build.cs
+++ b/Engine/Plugins/PCG/Source/PCGEditor/PCGEditor.Build.cs
@@ -78,6 +78,10 @@ namespace UnrealBuildTool.Rules
 					"GeometryCore",
 					"GeometryFramework"
 				});
+
+			// This module needs AutoRTFM so disable the auto disable since its in a plugin.
+			// The call to the delegate `OnObjectModified` calls into here.
+			bDisableAutoRTFMInstrumentation = false;
 		}
 	}
 }
