Enable AutoRTFM on more modules where it is required.

#Synced-CL 49758756

diff --git a/Engine/Plugins/Experimental/EditorDataStorage/Source/TedsCore/TedsCore.Build.cs b/Engine/Plugins/Experimental/EditorDataStorage/Source/TedsCore/TedsCore.Build.cs
--- a/Engine/Plugins/Experimental/EditorDataStorage/Source/TedsCore/TedsCore.Build.cs
+++ b/Engine/Plugins/Experimental/EditorDataStorage/Source/TedsCore/TedsCore.Build.cs
@@ -36,7 +36,7 @@ public class TedsCore : ModuleRules
 		}
 
 		// This module needs AutoRTFM so disable the auto disable since its in a plugin.
-		// The call to `ICompatibilityProvider::ICompatibilityProvider` calls into here.
+		// The call to `ICompatibilityProvider::AddCompatibleObject` calls into here.
 		bDisableAutoRTFMInstrumentation = false;
 	}
 }
