[TEDS Filters] Add new options to EFilterOptions so that a user can specify if they want to include/exclude rows that have the required columns and set FilterNodes to verify as the default

#Synced-CL 49728534

diff --git a/Engine/Plugins/Experimental/EditorDataStorage/Source/TedsCore/Private/TypedElementDatabase.cpp b/Engine/Plugins/Experimental/EditorDataStorage/Source/TedsCore/Private/TypedElementDatabase.cpp
--- a/Engine/Plugins/Experimental/EditorDataStorage/Source/TedsCore/Private/TypedElementDatabase.cpp
+++ b/Engine/Plugins/Experimental/EditorDataStorage/Source/TedsCore/Private/TypedElementDatabase.cpp
@@ -617,21 +617,27 @@ void UEditorDataStorage::FilterRowsBy(UE::Editor::DataStorage::FRowHandleArray&
 	{
 		struct FQueryResult : TResult<bool>
 		{
-			bool bResult = true;
-			virtual void Add(bool ResultValue) override
+			bool bResult;
+			explicit FQueryResult(const bool bResultValue)
+				: bResult(bResultValue)
+			{}
+
+			virtual void Add(bool bResultValue) override
 			{
-				bResult = ResultValue;
+				bResult = bResultValue;
 			}
 		};
 
 		QueryContext_DirectSingleApi Context(*Environment, *this);
 		Context.CheckCompatiblity(Filter);
 
-		bool bAcceptedResult = Options == EFilterOptions::Inclusive ? true : false;
+		const bool bVerifyColumns = !EnumHasAnyFlags(Options, EFilterOptions::IncludeRowsMissingColumns);
+		const bool bAcceptedResult = !EnumHasAnyFlags(Options, EFilterOptions::InvertFilter);
+
 		for (RowHandle Row : Input)
 		{
 			Context.GetContextImplementation().SetCurrentRow(Row);
-			FQueryResult QueryResult;
+			FQueryResult QueryResult(!bVerifyColumns);
 			Filter.Call<EFunctionCallConfig::VerifyColumns>(QueryResult, Context, Context.GetContextImplementation());
 			if (QueryResult.bResult == bAcceptedResult)
 			{
diff --git a/Engine/Plugins/Experimental/EditorDataStorageFeatures/Source/TedsOutliner/Private/TedsOutlinerModule.cpp b/Engine/Plugins/Experimental/EditorDataStorageFeatures/Source/TedsOutliner/Private/TedsOutlinerModule.cpp
--- a/Engine/Plugins/Experimental/EditorDataStorageFeatures/Source/TedsOutliner/Private/TedsOutlinerModule.cpp
+++ b/Engine/Plugins/Experimental/EditorDataStorageFeatures/Source/TedsOutliner/Private/TedsOutlinerModule.cpp
@@ -200,13 +200,9 @@ TSharedRef<SWidget> FTedsOutlinerModule::CreateLevelEditorTedsOutliner()
 		nullptr,
 		Storage->RegisterQuery(
 			Select()
-			.Where()
-				.All<FActorComponentTypeTag>()
+			.Where(TColumn<FActorComponentTypeTag>())
 			.Compile()));
-
-	// TODO: Make it possible to set EFunctionCallConfig::VerifyColumns in BuildQueryFunction as we will usually want to exclude rows that 
-	// don't have the columns we want to filter (MobilityColumn=Static will only filter on columns that have the Mobility Column and
-	// return true for rows that don't have that column), perhaps add it to EFilterOptions?
+	
 	Params.Filters.Emplace(
 		"Static",
 		LOCTEXT("StaticFilterDisplayName", "Static"), 
diff --git a/Engine/Plugins/Experimental/EditorDataStorageFeatures/Source/TedsQueryStack/Public/TedsRowFilterNode.h b/Engine/Plugins/Experimental/EditorDataStorageFeatures/Source/TedsQueryStack/Public/TedsRowFilterNode.h
--- a/Engine/Plugins/Experimental/EditorDataStorageFeatures/Source/TedsQueryStack/Public/TedsRowFilterNode.h
+++ b/Engine/Plugins/Experimental/EditorDataStorageFeatures/Source/TedsQueryStack/Public/TedsRowFilterNode.h
@@ -35,7 +35,7 @@ namespace UE::Editor::DataStorage::QueryStack
 			: FRowFilterNode(
 				Storage, 
 				ParentRowNode, 
-				ICoreProvider::EFilterOptions::Inclusive, 
+				ICoreProvider::EFilterOptions::Default, 
 				Filter)
 		{}
             
@@ -62,7 +62,7 @@ namespace UE::Editor::DataStorage::QueryStack
 			: FRowFilterNode(
 				Storage, 
 				ParentRowNode, 
-				ICoreProvider::EFilterOptions::Inclusive, 
+				ICoreProvider::EFilterOptions::Default, 
 				Queries::BuildQueryFunction<bool>(Forward<FilterFunction>(Filter)))
 		{}
 
diff --git a/Engine/Source/Runtime/TypedElementFramework/Public/Elements/Interfaces/TypedElementDataStorageInterface.h b/Engine/Source/Runtime/TypedElementFramework/Public/Elements/Interfaces/TypedElementDataStorageInterface.h
--- a/Engine/Source/Runtime/TypedElementFramework/Public/Elements/Interfaces/TypedElementDataStorageInterface.h
+++ b/Engine/Source/Runtime/TypedElementFramework/Public/Elements/Interfaces/TypedElementDataStorageInterface.h
@@ -191,10 +191,12 @@ public:
 
 	enum class EFilterOptions : uint8
 	{
-		/** Include rows that have the columns required for the filter function and the filter function returns true. */
-		Inclusive,
-		/** Exclude rows that have the columns required for the filter function and the filter function returns true. */
-		Exclusive
+		/** Only include rows in the result that have the required columns and pass the filter function */
+		Default = 0,
+		/** Include rows that don't have the required columns and allow them to bypass the filter */
+		IncludeRowsMissingColumns = 1 << 0,
+		/** Invert the filter to return rows that did not pass the filter function */
+		InvertFilter = 1 << 1,
 	};
 
 	/** Filters the provided rows using the query filter and stores the results in the provided list. */
@@ -625,6 +627,8 @@ public:
 	inline void RemoveRowMapping(const FMapKeyView& Key);
 };
 
+ENUM_CLASS_FLAGS(ICoreProvider::EFilterOptions)
+
 // Implementations
 
 template <typename FactoryT>
