TEDS UI: Fix a bug where widget constructors were allocated with FMemory::Malloc but deleted using the delete operator, causing a mismatch

- Replace TUniquePtr arguments with TSharedPtr since most widget constructors were converted to the latter anyways
- Provide a custom deleter for the shared ptr that properly calls the destructor and FMemory::Free to match the allocation
- Add an alias for TSharedPtr<FTypedElementWidgetConstructor> to make it easier to change in the future
- Deprecate old arguments and functions and replace uses
- Also do the same for decorator widgets to ensure they are deleted properly, but there is nothing to deprecate since they are created and destructed internally

#Synced-CL 49721519

diff --git a/Engine/Plugins/Experimental/EditorDataStorage/Source/TedsCore/Private/TypedElementDatabaseUI.cpp b/Engine/Plugins/Experimental/EditorDataStorage/Source/TedsCore/Private/TypedElementDatabaseUI.cpp
--- a/Engine/Plugins/Experimental/EditorDataStorage/Source/TedsCore/Private/TypedElementDatabaseUI.cpp
+++ b/Engine/Plugins/Experimental/EditorDataStorage/Source/TedsCore/Private/TypedElementDatabaseUI.cpp
@@ -53,6 +53,66 @@
 		return false;
 	}
 
+	// Custom deleter for a widget constructor since we allocate using FMemory::Malloc and then pass ownership to a TSharedPtr
+	struct FWidgetConstructorDeleter
+	{
+		FWidgetConstructorDeleter(const UScriptStruct* InConstructorType)
+			: ConstructorType(InConstructorType)
+		{}
+		
+		void operator()(FTypedElementWidgetConstructor* Constructor) const
+		{
+			if (Constructor)
+			{
+				ConstructorType->DestroyStruct(Constructor);
+				FMemory::Free(Constructor);
+			}
+		}
+
+	private:
+		const UScriptStruct* ConstructorType;
+	};
+
+	UE::Editor::DataStorage::IUiProvider::FWidgetConstructorPtr CreateWidgetConstructorInstance(const UScriptStruct* ConstructorType)
+	{
+		FTypedElementWidgetConstructor* ConstructorData = reinterpret_cast<FTypedElementWidgetConstructor*>(
+			FMemory::Malloc(ConstructorType->GetStructureSize(), ConstructorType->GetMinAlignment()));
+		
+		ConstructorType->InitializeStruct(ConstructorData);
+		
+		return MakeShareable(ConstructorData, FWidgetConstructorDeleter(ConstructorType));
+	}
+
+	// Custom deleter for a decorator widget constructor since we allocate using FMemory::Malloc and then store in a TUniquePtr
+	struct FDecoratorWidgetConstructorDeleter
+	{
+		FDecoratorWidgetConstructorDeleter(const UScriptStruct* InConstructorType)
+			: ConstructorType(InConstructorType)
+		{}
+		
+		void operator()(FTedsDecoratorWidgetConstructor* Constructor) const
+		{
+			if (Constructor)
+			{
+				ConstructorType->DestroyStruct(Constructor);
+				FMemory::Free(Constructor);
+			}
+		}
+
+	private:
+		const UScriptStruct* ConstructorType;
+	};
+
+	TSharedPtr<FTedsDecoratorWidgetConstructor> CreateDecoratorWidgetConstructorInstance(const UScriptStruct* ConstructorType)
+	{
+		FTedsDecoratorWidgetConstructor* ConstructorData = reinterpret_cast<FTedsDecoratorWidgetConstructor*>(
+			FMemory::Malloc(ConstructorType->GetStructureSize(), ConstructorType->GetMinAlignment()));
+		
+		ConstructorType->InitializeStruct(ConstructorData);
+
+		return MakeShareable(ConstructorData, FDecoratorWidgetConstructorDeleter(ConstructorType));
+	}
+
 	static const UE::Editor::DataStorage::IUiProvider::FPurposeID DefaultWidgetPurposeID(
 		UE::Editor::DataStorage::IUiProvider::FPurposeInfo("General", "Cell", "Default").GeneratePurposeID());
 	
@@ -370,7 +430,7 @@
 }
 
 void UEditorDataStorageUi::CreateWidgetConstructors(UE::Editor::DataStorage::RowHandle PurposeRow,
-	const UE::Editor::DataStorage::FMetaDataView& Arguments, const WidgetConstructorCallback& Callback)
+	const UE::Editor::DataStorage::FMetaDataView& Arguments, const WidgetConstructionCallback& Callback)
 {
 	while (Storage->HasColumns<FWidgetPurposeColumn>(PurposeRow))
 	{
@@ -406,7 +466,7 @@
 
 void UEditorDataStorageUi::CreateWidgetConstructors(UE::Editor::DataStorage::RowHandle PurposeRow, EMatchApproach MatchApproach,
 	TArray<TWeakObjectPtr<const UScriptStruct>>& Columns, const UE::Editor::DataStorage::FMetaDataView& Arguments,
-	const WidgetConstructorCallback& Callback)
+	const WidgetConstructionCallback& Callback)
 {
 	// Sort by name so that removing the matched columns can be done in a single pass.
 	// Sorting by ptr does not work because dynamic column ptrs are different than their base template
@@ -523,107 +583,19 @@
 	}
 }
 
-void UEditorDataStorageUi::RegisterWidgetPurpose(FName Purpose, EPurposeType Type, FText Description)
-{
-	RegisterWidgetPurpose(UE::Editor::DataStorage::FMapKey(Purpose), FPurposeInfo(Purpose, Type, Description));
-}
-
-bool UEditorDataStorageUi::RegisterWidgetFactory(FName Purpose, const UScriptStruct* Constructor)
-{
-	UE::Editor::DataStorage::RowHandle PurposeRow = Storage->LookupMappedRow(PurposeMappingDomain, UE::Editor::DataStorage::FMapKey(Purpose));
-
-	if (Storage->IsRowAvailable(PurposeRow))
-	{
-		return RegisterWidgetFactory(PurposeRow, Constructor);
-	}
-	
-	UE_LOG(LogEditorDataStorageUI, Warning, 
-		TEXT("Unable to register widget factory '%s' as purpose '%s' isn't registered."), *Constructor->GetName(), *Purpose.ToString());
-
-	return false;
-}
-
-bool UEditorDataStorageUi::RegisterWidgetFactory(
-	FName Purpose, const UScriptStruct* Constructor, UE::Editor::DataStorage::Queries::FConditions Columns)
-{
-	UE::Editor::DataStorage::RowHandle PurposeRow = Storage->LookupMappedRow(PurposeMappingDomain, UE::Editor::DataStorage::FMapKey(Purpose));
-	
-	if (Storage->IsRowAvailable(PurposeRow))
-	{
-		return RegisterWidgetFactory(PurposeRow, Constructor, Columns);
-	}
-	
-	UE_LOG(LogEditorDataStorageUI, Warning, 
-			TEXT("Unable to register widget factory '%s' as purpose '%s' isn't registered."), *Constructor->GetName(), *Purpose.ToString());
-
-	return false;
-}
-
-bool UEditorDataStorageUi::RegisterWidgetFactory(FName Purpose, TUniquePtr<FTypedElementWidgetConstructor>&& Constructor)
-{
-	UE::Editor::DataStorage::RowHandle PurposeRow = Storage->LookupMappedRow(PurposeMappingDomain, UE::Editor::DataStorage::FMapKey(Purpose));
-	if (Storage->IsRowAvailable(PurposeRow))
-	{
-		return RegisterWidgetFactory(PurposeRow, MoveTemp(Constructor));
-	}
-	
-	UE_LOG(LogEditorDataStorageUI, Warning, 
-			TEXT("Unable to register widget factory as purpose '%s' isn't registered."), *Purpose.ToString());
-	return false;
-}
-
-bool UEditorDataStorageUi::RegisterWidgetFactory(FName Purpose, TUniquePtr<FTypedElementWidgetConstructor>&& Constructor, 
-	UE::Editor::DataStorage::Queries::FConditions Columns)
-{
-	UE::Editor::DataStorage::RowHandle PurposeRow = Storage->LookupMappedRow(PurposeMappingDomain, UE::Editor::DataStorage::FMapKey(Purpose));
-	
-	if (Storage->IsRowAvailable(PurposeRow))
-	{
-		return RegisterWidgetFactory(PurposeRow, MoveTemp(Constructor), Columns);
-	}
-		
-	UE_LOG(LogEditorDataStorageUI, Warning, TEXT("Unable to register widget factory '%s' as purpose '%s' isn't registered."), 
-			*Constructor->GetTypeInfo()->GetName(), *Purpose.ToString());
-	return false;
-}
-
-void UEditorDataStorageUi::CreateWidgetConstructors(FName Purpose,
-	const UE::Editor::DataStorage::FMetaDataView& Arguments, const WidgetConstructorCallback& Callback)
-{
-	UE::Editor::DataStorage::RowHandle PurposeRow = Storage->LookupMappedRow(PurposeMappingDomain, UE::Editor::DataStorage::FMapKey(Purpose));
-	CreateWidgetConstructors(PurposeRow, Arguments, Callback);
-}
-
-void UEditorDataStorageUi::CreateWidgetConstructors(FName Purpose, EMatchApproach MatchApproach, 
-	TArray<TWeakObjectPtr<const UScriptStruct>>& Columns, const UE::Editor::DataStorage::FMetaDataView& Arguments,
-	const WidgetConstructorCallback& Callback)
-{
-	UE::Editor::DataStorage::RowHandle PurposeRow = Storage->LookupMappedRow(PurposeMappingDomain, UE::Editor::DataStorage::FMapKey(Purpose));
-	CreateWidgetConstructors(PurposeRow, MatchApproach, Columns, Arguments, Callback);
-}
-
-void UEditorDataStorageUi::ConstructWidgets(FName Purpose, const UE::Editor::DataStorage::FMetaDataView& Arguments,
-	const WidgetCreatedCallback& ConstructionCallback)
-{
-	UE::Editor::DataStorage::RowHandle PurposeRow = Storage->LookupMappedRow(PurposeMappingDomain, UE::Editor::DataStorage::FMapKey(Purpose));
-	ConstructWidgets(PurposeRow, Arguments, ConstructionCallback);
-}
-
 bool UEditorDataStorageUi::CreateSingleWidgetConstructor(
 	UE::Editor::DataStorage::RowHandle FactoryRow,
 	const UE::Editor::DataStorage::FMetaDataView& Arguments,
 	TArray<TWeakObjectPtr<const UScriptStruct>> MatchedColumnTypes,
-	const WidgetConstructorCallback& Callback)
+	const WidgetConstructionCallback& Callback)
 {
 	if (FWidgetFactoryConstructorTypeInfoColumn* ConstructorTypeInfoColumn = Storage->GetColumn<FWidgetFactoryConstructorTypeInfoColumn>(FactoryRow))
 	{
 		if (const UScriptStruct* Target = ConstructorTypeInfoColumn->Constructor.Get())
 		{
-			TUniquePtr<FTypedElementWidgetConstructor> Result(reinterpret_cast<FTypedElementWidgetConstructor*>(
-				FMemory::Malloc(Target->GetStructureSize(), Target->GetMinAlignment())));
+			FWidgetConstructorPtr Result(Internal::CreateWidgetConstructorInstance(Target));
 			if (Result)
 			{
-				Target->InitializeStruct(Result.Get());
 				Result->Initialize(Arguments, MoveTemp(MatchedColumnTypes), FactoryRow);
 				const TArray<TWeakObjectPtr<const UScriptStruct>>& MatchedColumns = Result->GetMatchedColumns();
 				return Callback(MoveTemp(Result), MatchedColumns);
@@ -634,18 +606,16 @@
 	return false;
 }
 
-TUniquePtr<FTedsDecoratorWidgetConstructor> UEditorDataStorageUi::CreateSingleDecoratorWidget(UE::Editor::DataStorage::RowHandle FactoryRow,
+TSharedPtr<FTedsDecoratorWidgetConstructor> UEditorDataStorageUi::CreateSingleDecoratorWidget(UE::Editor::DataStorage::RowHandle FactoryRow,
 	const UE::Editor::DataStorage::FMetaDataView& Arguments, TArray<TWeakObjectPtr<const UScriptStruct>> MatchedColumnTypes)
 {
 	if (FWidgetFactoryConstructorTypeInfoColumn* ConstructorTypeInfoColumn = Storage->GetColumn<FWidgetFactoryConstructorTypeInfoColumn>(FactoryRow))
 	{
 		if (const UScriptStruct* Target = ConstructorTypeInfoColumn->Constructor.Get())
 		{
-			TUniquePtr<FTedsDecoratorWidgetConstructor> Result(reinterpret_cast<FTedsDecoratorWidgetConstructor*>(
-				FMemory::Malloc(Target->GetStructureSize(), Target->GetMinAlignment())));
+			TSharedPtr<FTedsDecoratorWidgetConstructor> Result(Internal::CreateDecoratorWidgetConstructorInstance(Target));
 			if (Result)
 			{
-				Target->InitializeStruct(Result.Get());
 				Result->Initialize(Arguments, MoveTemp(MatchedColumnTypes), FactoryRow);
 				return MoveTemp(Result);
 			}
@@ -1170,12 +1140,12 @@
 					return Lhs->GetName() < Rhs->GetName();
 				});
 
-			TArray<TUniquePtr<FTedsDecoratorWidgetConstructor>> Decorators;
+			TArray<TSharedPtr<FTedsDecoratorWidgetConstructor>> Decorators;
 
 			FWidgetConstructionMethod ConstructionCallback = [this, Arguments, &Decorators]
 				(UE::Editor::DataStorage::RowHandle FactoryRow, TArray<TWeakObjectPtr<const UScriptStruct>> MatchedColumns)
 			{
-				TUniquePtr<FTedsDecoratorWidgetConstructor> Constructor = CreateSingleDecoratorWidget(FactoryRow, Arguments, MatchedColumns);
+				TSharedPtr<FTedsDecoratorWidgetConstructor> Constructor = CreateSingleDecoratorWidget(FactoryRow, Arguments, MatchedColumns);
 				Decorators.Add(MoveTemp(Constructor));
 				return true;
 			};
@@ -1192,7 +1162,7 @@
 			}
 
 			// Create each decorator widget, compounding the result widget of the previous operation onto the current widget
-			for (TUniquePtr<FTedsDecoratorWidgetConstructor>& ConstructorIt : Decorators)
+			for (TSharedPtr<FTedsDecoratorWidgetConstructor>& ConstructorIt : Decorators)
 			{
 				InternalWidget = ConstructorIt->CreateDecoratorWidget(InternalWidget, Storage, this, TargetRow, WidgetRowHandle, Arguments);
 			}
@@ -1249,4 +1219,129 @@
 	DecoratorObservers.Add(Storage->RegisterQuery(MoveTemp(OnRemoveObserver)));
 }
 
+// Deprecated functions use a deprecated alias as an argument so we need to disable warnings for their definitions
+PRAGMA_DISABLE_DEPRECATION_WARNINGS
+
+void UEditorDataStorageUi::CreateWidgetConstructors(UE::Editor::DataStorage::RowHandle PurposeRow,
+	const UE::Editor::DataStorage::FMetaDataView& Arguments, const WidgetConstructorCallback& Callback)
+{
+	auto CallbackWrapper = [Callback](FWidgetConstructorPtr Constructor, TConstArrayView<TWeakObjectPtr<const UScriptStruct>> MatchedColumns) -> bool
+		{
+			// This is technically safe since nothing else is using the shared ptr yet as it was just created,
+			// plus we are manually resetting the only instance which is used in this callback
+			TUniquePtr<FTypedElementWidgetConstructor> ConstructorPtr(Constructor.Get());
+			Constructor.Reset();
+
+			return Callback(MoveTemp(ConstructorPtr), MatchedColumns);
+		};
+
+	CreateWidgetConstructors(PurposeRow, Arguments, CallbackWrapper);
+}
+
+void UEditorDataStorageUi::CreateWidgetConstructors(UE::Editor::DataStorage::RowHandle PurposeRow, EMatchApproach MatchApproach,
+	TArray<TWeakObjectPtr<const UScriptStruct>>& Columns, const UE::Editor::DataStorage::FMetaDataView& Arguments,
+	const WidgetConstructorCallback& Callback)
+{
+	auto CallbackWrapper = [Callback](FWidgetConstructorPtr Constructor, TConstArrayView<TWeakObjectPtr<const UScriptStruct>> MatchedColumns) -> bool
+		{
+			// This is technically safe since nothing else is using the shared ptr yet as it was just created,
+			// plus we are manually resetting the only instance which is used in this callback
+			TUniquePtr<FTypedElementWidgetConstructor> ConstructorPtr(Constructor.Get());
+			Constructor.Reset();
+
+			return Callback(MoveTemp(ConstructorPtr), MatchedColumns);
+		};
+
+	CreateWidgetConstructors(PurposeRow, MatchApproach, Columns, Arguments, CallbackWrapper);
+}
+
+void UEditorDataStorageUi::RegisterWidgetPurpose(FName Purpose, EPurposeType Type, FText Description)
+{
+	RegisterWidgetPurpose(UE::Editor::DataStorage::FMapKey(Purpose), FPurposeInfo(Purpose, Type, Description));
+}
+
+bool UEditorDataStorageUi::RegisterWidgetFactory(FName Purpose, const UScriptStruct* Constructor)
+{
+	UE::Editor::DataStorage::RowHandle PurposeRow = Storage->LookupMappedRow(PurposeMappingDomain, UE::Editor::DataStorage::FMapKey(Purpose));
+
+	if (Storage->IsRowAvailable(PurposeRow))
+	{
+		return RegisterWidgetFactory(PurposeRow, Constructor);
+	}
+	
+	UE_LOG(LogEditorDataStorageUI, Warning, 
+		TEXT("Unable to register widget factory '%s' as purpose '%s' isn't registered."), *Constructor->GetName(), *Purpose.ToString());
+
+	return false;
+}
+
+bool UEditorDataStorageUi::RegisterWidgetFactory(
+	FName Purpose, const UScriptStruct* Constructor, UE::Editor::DataStorage::Queries::FConditions Columns)
+{
+	UE::Editor::DataStorage::RowHandle PurposeRow = Storage->LookupMappedRow(PurposeMappingDomain, UE::Editor::DataStorage::FMapKey(Purpose));
+	
+	if (Storage->IsRowAvailable(PurposeRow))
+	{
+		return RegisterWidgetFactory(PurposeRow, Constructor, Columns);
+	}
+	
+	UE_LOG(LogEditorDataStorageUI, Warning, 
+			TEXT("Unable to register widget factory '%s' as purpose '%s' isn't registered."), *Constructor->GetName(), *Purpose.ToString());
+
+	return false;
+}
+
+bool UEditorDataStorageUi::RegisterWidgetFactory(FName Purpose, TUniquePtr<FTypedElementWidgetConstructor>&& Constructor)
+{
+	UE::Editor::DataStorage::RowHandle PurposeRow = Storage->LookupMappedRow(PurposeMappingDomain, UE::Editor::DataStorage::FMapKey(Purpose));
+	if (Storage->IsRowAvailable(PurposeRow))
+	{
+		return RegisterWidgetFactory(PurposeRow, MoveTemp(Constructor));
+	}
+	
+	UE_LOG(LogEditorDataStorageUI, Warning, 
+			TEXT("Unable to register widget factory as purpose '%s' isn't registered."), *Purpose.ToString());
+	return false;
+}
+
+bool UEditorDataStorageUi::RegisterWidgetFactory(FName Purpose, TUniquePtr<FTypedElementWidgetConstructor>&& Constructor, 
+	UE::Editor::DataStorage::Queries::FConditions Columns)
+{
+	UE::Editor::DataStorage::RowHandle PurposeRow = Storage->LookupMappedRow(PurposeMappingDomain, UE::Editor::DataStorage::FMapKey(Purpose));
+	
+	if (Storage->IsRowAvailable(PurposeRow))
+	{
+		return RegisterWidgetFactory(PurposeRow, MoveTemp(Constructor), Columns);
+	}
+		
+	UE_LOG(LogEditorDataStorageUI, Warning, TEXT("Unable to register widget factory '%s' as purpose '%s' isn't registered."), 
+			*Constructor->GetTypeInfo()->GetName(), *Purpose.ToString());
+	return false;
+}
+
+void UEditorDataStorageUi::CreateWidgetConstructors(FName Purpose,
+	const UE::Editor::DataStorage::FMetaDataView& Arguments, const WidgetConstructorCallback& Callback)
+{
+	UE::Editor::DataStorage::RowHandle PurposeRow = Storage->LookupMappedRow(PurposeMappingDomain, UE::Editor::DataStorage::FMapKey(Purpose));
+	CreateWidgetConstructors(PurposeRow, Arguments, Callback);
+}
+
+void UEditorDataStorageUi::CreateWidgetConstructors(FName Purpose, EMatchApproach MatchApproach, 
+	TArray<TWeakObjectPtr<const UScriptStruct>>& Columns, const UE::Editor::DataStorage::FMetaDataView& Arguments,
+	const WidgetConstructorCallback& Callback)
+{
+	UE::Editor::DataStorage::RowHandle PurposeRow = Storage->LookupMappedRow(PurposeMappingDomain, UE::Editor::DataStorage::FMapKey(Purpose));
+	CreateWidgetConstructors(PurposeRow, MatchApproach, Columns, Arguments, Callback);
+}
+
+void UEditorDataStorageUi::ConstructWidgets(FName Purpose, const UE::Editor::DataStorage::FMetaDataView& Arguments,
+	const WidgetCreatedCallback& ConstructionCallback)
+{
+	UE::Editor::DataStorage::RowHandle PurposeRow = Storage->LookupMappedRow(PurposeMappingDomain, UE::Editor::DataStorage::FMapKey(Purpose));
+	ConstructWidgets(PurposeRow, Arguments, ConstructionCallback);
+}
+
+// Deprecated functions end here
+PRAGMA_ENABLE_DEPRECATION_WARNINGS
+
 #undef LOCTEXT_NAMESPACE // "TypedElementDatabaseUI"
diff --git a/Engine/Plugins/Experimental/EditorDataStorage/Source/TedsCore/Private/TypedElementDatabaseUI.h b/Engine/Plugins/Experimental/EditorDataStorage/Source/TedsCore/Private/TypedElementDatabaseUI.h
--- a/Engine/Plugins/Experimental/EditorDataStorage/Source/TedsCore/Private/TypedElementDatabaseUI.h
+++ b/Engine/Plugins/Experimental/EditorDataStorage/Source/TedsCore/Private/TypedElementDatabaseUI.h
@@ -51,13 +51,18 @@ public:
 		const UScriptStruct* Column) override;
 	
 	UE_API virtual void CreateWidgetConstructors(UE::Editor::DataStorage::RowHandle PurposeRow, 
-		const UE::Editor::DataStorage::FMetaDataView& Arguments, const WidgetConstructorCallback& Callback) override;
+		const UE::Editor::DataStorage::FMetaDataView& Arguments, const WidgetConstructionCallback& Callback) override;
 	UE_API virtual void CreateWidgetConstructors(UE::Editor::DataStorage::RowHandle PurposeRow, EMatchApproach MatchApproach, TArray<TWeakObjectPtr<const UScriptStruct>>& Columns,
-		const UE::Editor::DataStorage::FMetaDataView& Arguments, const WidgetConstructorCallback& Callback) override;
+		const UE::Editor::DataStorage::FMetaDataView& Arguments, const WidgetConstructionCallback& Callback) override;
 	UE_API virtual void ConstructWidgets(UE::Editor::DataStorage::RowHandle PurposeRow, const UE::Editor::DataStorage::FMetaDataView& Arguments,
 			const WidgetCreatedCallback& ConstructionCallback) override;
 
 	// Deprecated functions
+	PRAGMA_DISABLE_DEPRECATION_WARNINGS
+	UE_API virtual void CreateWidgetConstructors(UE::Editor::DataStorage::RowHandle PurposeRow, 
+	const UE::Editor::DataStorage::FMetaDataView& Arguments, const WidgetConstructorCallback& Callback) override;
+	UE_API virtual void CreateWidgetConstructors(UE::Editor::DataStorage::RowHandle PurposeRow, EMatchApproach MatchApproach, TArray<TWeakObjectPtr<const UScriptStruct>>& Columns,
+		const UE::Editor::DataStorage::FMetaDataView& Arguments, const WidgetConstructorCallback& Callback) override;
 	UE_API virtual bool RegisterWidgetFactory(UE::Editor::DataStorage::RowHandle PurposeRow,
 	TUniquePtr<FTypedElementWidgetConstructor>&& Constructor) override;
 	UE_API virtual bool RegisterWidgetFactory(UE::Editor::DataStorage::RowHandle PurposeRow, TUniquePtr<FTypedElementWidgetConstructor>&& Constructor, 
@@ -75,6 +80,7 @@ public:
 		const UE::Editor::DataStorage::FMetaDataView& Arguments, const WidgetConstructorCallback& Callback) override;
 	UE_API void ConstructWidgets(FName Purpose, const UE::Editor::DataStorage::FMetaDataView& Arguments,
 		const WidgetCreatedCallback& ConstructionCallback) override;
+	PRAGMA_ENABLE_DEPRECATION_WARNINGS
 	// ~Deprecated functions
 	
 	UE_API TSharedPtr<SWidget> ConstructWidget(UE::Editor::DataStorage::RowHandle Row, FTypedElementWidgetConstructor& Constructor,
@@ -122,7 +128,7 @@ private:
 		UE::Editor::DataStorage::RowHandle FactoryRow,
 		const UE::Editor::DataStorage::FMetaDataView& Arguments,
 		TArray<TWeakObjectPtr<const UScriptStruct>> MatchedColumnTypes,
-		const WidgetConstructorCallback& Callback);
+		const WidgetConstructionCallback& Callback);
 
 	UE_API void CreateWidgetInstance(
 		FTypedElementWidgetConstructor& Constructor,
@@ -149,7 +155,7 @@ private:
 		TArray<TWeakObjectPtr<const UScriptStruct>>& Columns,
 		FWidgetConstructionMethod ConstructionFunction);
 
-	UE_API TUniquePtr<FTedsDecoratorWidgetConstructor> CreateSingleDecoratorWidget(
+	UE_API TSharedPtr<FTedsDecoratorWidgetConstructor> CreateSingleDecoratorWidget(
 		UE::Editor::DataStorage::RowHandle FactoryRow,
 		const UE::Editor::DataStorage::FMetaDataView& Arguments,
 		TArray<TWeakObjectPtr<const UScriptStruct>> MatchedColumnTypes);
diff --git a/Engine/Plugins/Experimental/EditorDataStorage/Source/TedsCore/Private/Widgets/STedsWidget.cpp b/Engine/Plugins/Experimental/EditorDataStorage/Source/TedsCore/Private/Widgets/STedsWidget.cpp
--- a/Engine/Plugins/Experimental/EditorDataStorage/Source/TedsCore/Private/Widgets/STedsWidget.cpp
+++ b/Engine/Plugins/Experimental/EditorDataStorage/Source/TedsCore/Private/Widgets/STedsWidget.cpp
@@ -86,7 +86,7 @@ namespace UE::Editor::DataStorage::Widgets
 				// First we set the widget reference on the column
 				WidgetReferenceColumn->Widget = InContent;
 
-				if (TSharedPtr<FTypedElementWidgetConstructor> Constructor = WidgetReferenceColumn->WidgetConstructor.Pin())
+				if (IUiProvider::FWidgetConstructorPtr Constructor = WidgetReferenceColumn->WidgetConstructor.Pin())
 				{
 					TConstArrayView<const UScriptStruct*> AdditionalColumns = Constructor->GetAdditionalColumnsList();
 					
diff --git a/Engine/Plugins/Experimental/EditorDataStorageFeatures/Source/TedsAssetData/Private/Widgets/AssetDataLabelWidget.cpp b/Engine/Plugins/Experimental/EditorDataStorageFeatures/Source/TedsAssetData/Private/Widgets/AssetDataLabelWidget.cpp
--- a/Engine/Plugins/Experimental/EditorDataStorageFeatures/Source/TedsAssetData/Private/Widgets/AssetDataLabelWidget.cpp
+++ b/Engine/Plugins/Experimental/EditorDataStorageFeatures/Source/TedsAssetData/Private/Widgets/AssetDataLabelWidget.cpp
@@ -56,11 +56,11 @@ TSharedPtr<SWidget> FAssetDataLabelWidgetConstructor::CreateWidget(UE::Editor::D
 		// Create the LabelWidget through its General Widget
 		UE::Editor::DataStorage::RowHandle LabelPurposeName = DataStorageUi->FindPurpose(DataStorageUi->GetGeneralWidgetPurposeID());
 
-		TSharedPtr<FTypedElementWidgetConstructor> OutLabelWidgetConstructorPtr;
+		IUiProvider::FWidgetConstructorPtr OutLabelWidgetConstructorPtr;
 
-		auto AssignLabelWidgetToColumn = [&OutLabelWidgetConstructorPtr] (TUniquePtr<FTypedElementWidgetConstructor> Constructor, TConstArrayView<TWeakObjectPtr<const UScriptStruct>>)
+		auto AssignLabelWidgetToColumn = [&OutLabelWidgetConstructorPtr] (IUiProvider::FWidgetConstructorPtr WidgetConstructor, TConstArrayView<TWeakObjectPtr<const UScriptStruct>>)
 		{
-			OutLabelWidgetConstructorPtr = TSharedPtr<FTypedElementWidgetConstructor>(Constructor.Release());
+			OutLabelWidgetConstructorPtr = MoveTemp(WidgetConstructor);
 			return false;
 		};
 
diff --git a/Engine/Plugins/Experimental/EditorDataStorageFeatures/Source/TedsAssetData/Private/Widgets/AssetPreview/STedsAssetPreviewWidget.cpp b/Engine/Plugins/Experimental/EditorDataStorageFeatures/Source/TedsAssetData/Private/Widgets/AssetPreview/STedsAssetPreviewWidget.cpp
--- a/Engine/Plugins/Experimental/EditorDataStorageFeatures/Source/TedsAssetData/Private/Widgets/AssetPreview/STedsAssetPreviewWidget.cpp
+++ b/Engine/Plugins/Experimental/EditorDataStorageFeatures/Source/TedsAssetData/Private/Widgets/AssetPreview/STedsAssetPreviewWidget.cpp
@@ -101,10 +101,10 @@ void STedsAssetPreviewWidget::ConstructWidget()
 
 	// 1) Header
 	{
-		TSharedPtr<FTypedElementWidgetConstructor> OutHeaderButtonsWidgetConstructorPtr;
-		auto AssignHeaderButtonsWidgetToColumn = [&OutHeaderButtonsWidgetConstructorPtr] (TUniquePtr<FTypedElementWidgetConstructor> Constructor, TConstArrayView<TWeakObjectPtr<const UScriptStruct>>)
+		IUiProvider::FWidgetConstructorPtr OutHeaderButtonsWidgetConstructorPtr;
+		auto AssignHeaderButtonsWidgetToColumn = [&OutHeaderButtonsWidgetConstructorPtr] (IUiProvider::FWidgetConstructorPtr WidgetConstructor, TConstArrayView<TWeakObjectPtr<const UScriptStruct>>)
 		{
-			OutHeaderButtonsWidgetConstructorPtr = TSharedPtr<FTypedElementWidgetConstructor>(Constructor.Release());
+			OutHeaderButtonsWidgetConstructorPtr = MoveTemp(WidgetConstructor);;
 			return false;
 		};
 
@@ -163,10 +163,10 @@ void STedsAssetPreviewWidget::ConstructWidget()
 
 	// 2) Thumbnail
 	{
-		TSharedPtr<FTypedElementWidgetConstructor> OutThumbnailWidgetConstructorPtr;
-		auto AssignThumbnailWidgetToColumn = [&OutThumbnailWidgetConstructorPtr] (TUniquePtr<FTypedElementWidgetConstructor> Constructor, TConstArrayView<TWeakObjectPtr<const UScriptStruct>>)
+		IUiProvider::FWidgetConstructorPtr OutThumbnailWidgetConstructorPtr;
+		auto AssignThumbnailWidgetToColumn = [&OutThumbnailWidgetConstructorPtr] (IUiProvider::FWidgetConstructorPtr WidgetConstructor, TConstArrayView<TWeakObjectPtr<const UScriptStruct>>)
 		{
-			OutThumbnailWidgetConstructorPtr = TSharedPtr<FTypedElementWidgetConstructor>(Constructor.Release());
+			OutThumbnailWidgetConstructorPtr = MoveTemp(WidgetConstructor);;
 			return false;
 		};
 
@@ -205,10 +205,10 @@ void STedsAssetPreviewWidget::ConstructWidget()
 
 	// 3) Asset Name and Basic info
 	{
-		TSharedPtr<FTypedElementWidgetConstructor> OutBasicInfoWidgetConstructorPtr;
-		auto AssignBasicInfoWidgetToColumn = [&OutBasicInfoWidgetConstructorPtr] (TUniquePtr<FTypedElementWidgetConstructor> Constructor, TConstArrayView<TWeakObjectPtr<const UScriptStruct>>)
+		IUiProvider::FWidgetConstructorPtr OutBasicInfoWidgetConstructorPtr;
+		auto AssignBasicInfoWidgetToColumn = [&OutBasicInfoWidgetConstructorPtr] (IUiProvider::FWidgetConstructorPtr WidgetConstructor, TConstArrayView<TWeakObjectPtr<const UScriptStruct>>)
 		{
-			OutBasicInfoWidgetConstructorPtr = TSharedPtr<FTypedElementWidgetConstructor>(Constructor.Release());
+			OutBasicInfoWidgetConstructorPtr = MoveTemp(WidgetConstructor);;
 			return false;
 		};
 
@@ -247,10 +247,10 @@ void STedsAssetPreviewWidget::ConstructWidget()
 
 	// 4) Full Details
 	{
-		TSharedPtr<FTypedElementWidgetConstructor> OutFullInfoWidgetConstructorPtr;
-		auto AssignFullInfoWidgetToColumn = [&OutFullInfoWidgetConstructorPtr] (TUniquePtr<FTypedElementWidgetConstructor> Constructor, TConstArrayView<TWeakObjectPtr<const UScriptStruct>>)
+		IUiProvider::FWidgetConstructorPtr OutFullInfoWidgetConstructorPtr;
+		auto AssignFullInfoWidgetToColumn = [&OutFullInfoWidgetConstructorPtr] (IUiProvider::FWidgetConstructorPtr WidgetConstructor, TConstArrayView<TWeakObjectPtr<const UScriptStruct>>)
 		{
-			OutFullInfoWidgetConstructorPtr = TSharedPtr<FTypedElementWidgetConstructor>(Constructor.Release());
+			OutFullInfoWidgetConstructorPtr = MoveTemp(WidgetConstructor);;
 			return false;
 		};
 
diff --git a/Engine/Plugins/Experimental/EditorDataStorageFeatures/Source/TedsAssetData/Private/Widgets/AssetPreview/Widgets/BaseInfo/AssetPreviewBaseInfoWidget.cpp b/Engine/Plugins/Experimental/EditorDataStorageFeatures/Source/TedsAssetData/Private/Widgets/AssetPreview/Widgets/BaseInfo/AssetPreviewBaseInfoWidget.cpp
--- a/Engine/Plugins/Experimental/EditorDataStorageFeatures/Source/TedsAssetData/Private/Widgets/AssetPreview/Widgets/BaseInfo/AssetPreviewBaseInfoWidget.cpp
+++ b/Engine/Plugins/Experimental/EditorDataStorageFeatures/Source/TedsAssetData/Private/Widgets/AssetPreview/Widgets/BaseInfo/AssetPreviewBaseInfoWidget.cpp
@@ -49,10 +49,10 @@ TSharedPtr<SWidget> FAssetPreviewBaseInfoWidgetConstructor::CreateWidget(UE::Edi
 	{
 		const FName NamePurposeTableName = TEXT("Editor_WidgetTable");
 
-		TSharedPtr<FTypedElementWidgetConstructor> OutNameWidgetConstructorPtr;
-		auto AssignNameWidgetToColumn = [&OutNameWidgetConstructorPtr] (TUniquePtr<FTypedElementWidgetConstructor> Constructor, TConstArrayView<TWeakObjectPtr<const UScriptStruct>>)
+		IUiProvider::FWidgetConstructorPtr OutNameWidgetConstructorPtr;
+		auto AssignNameWidgetToColumn = [&OutNameWidgetConstructorPtr] (IUiProvider::FWidgetConstructorPtr WidgetConstructor, TConstArrayView<TWeakObjectPtr<const UScriptStruct>>)
 		{
-			OutNameWidgetConstructorPtr = TSharedPtr<FTypedElementWidgetConstructor>(Constructor.Release());
+			OutNameWidgetConstructorPtr = MoveTemp(WidgetConstructor);;
 			return false;
 		};
 
@@ -90,20 +90,23 @@ TSharedPtr<SWidget> FAssetPreviewBaseInfoWidgetConstructor::CreateWidget(UE::Edi
 	FText AssetClassName = FText::GetEmpty();
 	if (DataStorage->HasColumns<FAssetTag>(TargetRow))
 	{
-		FAssetData AssetData = DataStorage->GetColumn<FAssetDataColumn_Experimental>(TargetRow)->AssetData;
-		if (const UClass* AssetClass = AssetData.GetClass())
+		if (FAssetDataColumn_Experimental* AssetDataColumn = DataStorage->GetColumn<FAssetDataColumn_Experimental>(TargetRow))
 		{
-			AssetClassName = AssetClass->GetDisplayNameText();
-			Icon = FSlateIconFinder::FindIconForClass(AssetClass);
-
-			if (DataStorage->HasColumns<FSlateColorColumn>(TargetRow))
-			{
-				IconColor = DataStorage->GetColumn<FSlateColorColumn>(TargetRow)->Color;
-			}
-			// If no Color is set use the Class specific one
-			else if (const UAssetDefinition* AssetDefinition = UAssetDefinitionRegistry::Get()->GetAssetDefinitionForClass(AssetClass))
+			FAssetData AssetData = AssetDataColumn->AssetData;
+			if (const UClass* AssetClass = AssetData.GetClass())
 			{
-				IconColor = AssetDefinition->GetAssetColor();
+				AssetClassName = AssetClass->GetDisplayNameText();
+				Icon = FSlateIconFinder::FindIconForClass(AssetClass);
+
+				if (DataStorage->HasColumns<FSlateColorColumn>(TargetRow))
+				{
+					IconColor = DataStorage->GetColumn<FSlateColorColumn>(TargetRow)->Color;
+				}
+				// If no Color is set use the Class specific one
+				else if (const UAssetDefinition* AssetDefinition = UAssetDefinitionRegistry::Get()->GetAssetDefinitionForClass(AssetClass))
+				{
+					IconColor = AssetDefinition->GetAssetColor();
+				}
 			}
 		}
 	}
diff --git a/Engine/Plugins/Experimental/EditorDataStorageFeatures/Source/TedsAssetData/Private/Widgets/AssetPreview/Widgets/Header/AssetPreviewHeaderWidget.cpp b/Engine/Plugins/Experimental/EditorDataStorageFeatures/Source/TedsAssetData/Private/Widgets/AssetPreview/Widgets/Header/AssetPreviewHeaderWidget.cpp
--- a/Engine/Plugins/Experimental/EditorDataStorageFeatures/Source/TedsAssetData/Private/Widgets/AssetPreview/Widgets/Header/AssetPreviewHeaderWidget.cpp
+++ b/Engine/Plugins/Experimental/EditorDataStorageFeatures/Source/TedsAssetData/Private/Widgets/AssetPreview/Widgets/Header/AssetPreviewHeaderWidget.cpp
@@ -54,18 +54,18 @@ TSharedPtr<SWidget> FAssetPreviewHeaderWidgetConstructor::CreateWidget(UE::Edito
 
 	const FName NamePurposeTableName = TEXT("Editor_WidgetTable");
 
-	TSharedPtr<FTypedElementWidgetConstructor> OutEditModeButtonWidgetConstructorPtr;
-	TSharedPtr<FTypedElementWidgetConstructor> OutContextMenuWidgetConstructorPtr;
+	IUiProvider::FWidgetConstructorPtr OutEditModeButtonWidgetConstructorPtr;
+	IUiProvider::FWidgetConstructorPtr OutContextMenuWidgetConstructorPtr;
 
-	auto AssignEditModeButtonWidgetToColumn = [&OutEditModeButtonWidgetConstructorPtr] (TUniquePtr<FTypedElementWidgetConstructor> Constructor, TConstArrayView<TWeakObjectPtr<const UScriptStruct>>)
+	auto AssignEditModeButtonWidgetToColumn = [&OutEditModeButtonWidgetConstructorPtr] (IUiProvider::FWidgetConstructorPtr WidgetConstructor, TConstArrayView<TWeakObjectPtr<const UScriptStruct>>)
 	{
-		OutEditModeButtonWidgetConstructorPtr = TSharedPtr<FTypedElementWidgetConstructor>(Constructor.Release());
+		OutEditModeButtonWidgetConstructorPtr = MoveTemp(WidgetConstructor);;
 		return false;
 	};
 
-	auto AssignContextMenuWidgetToColumn = [&OutContextMenuWidgetConstructorPtr] (TUniquePtr<FTypedElementWidgetConstructor> Constructor, TConstArrayView<TWeakObjectPtr<const UScriptStruct>>)
+	auto AssignContextMenuWidgetToColumn = [&OutContextMenuWidgetConstructorPtr] (IUiProvider::FWidgetConstructorPtr WidgetConstructor, TConstArrayView<TWeakObjectPtr<const UScriptStruct>>)
 	{
-		OutContextMenuWidgetConstructorPtr = TSharedPtr<FTypedElementWidgetConstructor>(Constructor.Release());
+		OutContextMenuWidgetConstructorPtr = MoveTemp(WidgetConstructor);;
 		return false;
 	};
 
diff --git a/Engine/Plugins/Experimental/EditorDataStorageFeatures/Source/TedsAssetData/Private/Widgets/AssetPreview/Widgets/Thumbnail/AssetPreviewThumbnailWidget.cpp b/Engine/Plugins/Experimental/EditorDataStorageFeatures/Source/TedsAssetData/Private/Widgets/AssetPreview/Widgets/Thumbnail/AssetPreviewThumbnailWidget.cpp
--- a/Engine/Plugins/Experimental/EditorDataStorageFeatures/Source/TedsAssetData/Private/Widgets/AssetPreview/Widgets/Thumbnail/AssetPreviewThumbnailWidget.cpp
+++ b/Engine/Plugins/Experimental/EditorDataStorageFeatures/Source/TedsAssetData/Private/Widgets/AssetPreview/Widgets/Thumbnail/AssetPreviewThumbnailWidget.cpp
@@ -65,10 +65,10 @@ TSharedPtr<SWidget> FAssetPreviewThumbnailWidgetConstructor::CreateWidget(UE::Ed
 		SizeValueColumn->SizeValue = bIsThumbnailExpandedGlobal ? 256.f : 128.f;
 	}
 
-	TSharedPtr<FTypedElementWidgetConstructor> OutThumbnailWidgetConstructorPtr;
-    auto AssignThumbnailWidgetToColumn = [&OutThumbnailWidgetConstructorPtr] (TUniquePtr<FTypedElementWidgetConstructor> Constructor, TConstArrayView<TWeakObjectPtr<const UScriptStruct>>)
+	IUiProvider::FWidgetConstructorPtr OutThumbnailWidgetConstructorPtr;
+    auto AssignThumbnailWidgetToColumn = [&OutThumbnailWidgetConstructorPtr] (IUiProvider::FWidgetConstructorPtr WidgetConstructor, TConstArrayView<TWeakObjectPtr<const UScriptStruct>>)
     {
-    	OutThumbnailWidgetConstructorPtr = TSharedPtr<FTypedElementWidgetConstructor>(Constructor.Release());
+    	OutThumbnailWidgetConstructorPtr = MoveTemp(WidgetConstructor);;
     	return false;
     };
 
diff --git a/Engine/Plugins/Experimental/EditorDataStorageFeatures/Source/TedsAssetData/Private/Widgets/ContentBrowserListViewNameWidget.cpp b/Engine/Plugins/Experimental/EditorDataStorageFeatures/Source/TedsAssetData/Private/Widgets/ContentBrowserListViewNameWidget.cpp
--- a/Engine/Plugins/Experimental/EditorDataStorageFeatures/Source/TedsAssetData/Private/Widgets/ContentBrowserListViewNameWidget.cpp
+++ b/Engine/Plugins/Experimental/EditorDataStorageFeatures/Source/TedsAssetData/Private/Widgets/ContentBrowserListViewNameWidget.cpp
@@ -69,18 +69,18 @@ TSharedPtr<SWidget> FContentBrowserListViewNameWidgetConstructor::CreateWidget(U
 			ParentWidgetRowHandle = ParentWidgetRow->Parent;
 		}
 
-		TSharedPtr<FTypedElementWidgetConstructor> OutThumbnailWidgetConstructorPtr;
-		TSharedPtr<FTypedElementWidgetConstructor> OutLabelWidgetConstructorPtr;
+		IUiProvider::FWidgetConstructorPtr OutThumbnailWidgetConstructorPtr;
+		IUiProvider::FWidgetConstructorPtr OutLabelWidgetConstructorPtr;
 
-		auto AssignThumbnailWidgetToColumn = [&OutThumbnailWidgetConstructorPtr] (TUniquePtr<FTypedElementWidgetConstructor> Constructor, TConstArrayView<TWeakObjectPtr<const UScriptStruct>>)
+		auto AssignThumbnailWidgetToColumn = [&OutThumbnailWidgetConstructorPtr] (IUiProvider::FWidgetConstructorPtr WidgetConstructor, TConstArrayView<TWeakObjectPtr<const UScriptStruct>>)
 		{
-			OutThumbnailWidgetConstructorPtr = TSharedPtr<FTypedElementWidgetConstructor>(Constructor.Release());
+			OutThumbnailWidgetConstructorPtr = MoveTemp(WidgetConstructor);;
 			return false;
 		};
 
-		auto AssignLabelWidgetToColumn = [&OutLabelWidgetConstructorPtr] (TUniquePtr<FTypedElementWidgetConstructor> Constructor, TConstArrayView<TWeakObjectPtr<const UScriptStruct>>)
+		auto AssignLabelWidgetToColumn = [&OutLabelWidgetConstructorPtr] (IUiProvider::FWidgetConstructorPtr WidgetConstructor, TConstArrayView<TWeakObjectPtr<const UScriptStruct>>)
 		{
-			OutLabelWidgetConstructorPtr = TSharedPtr<FTypedElementWidgetConstructor>(Constructor.Release());
+			OutLabelWidgetConstructorPtr = MoveTemp(WidgetConstructor);;
 			return false;
 		};
 
diff --git a/Engine/Plugins/Experimental/EditorDataStorageFeatures/Source/TedsContentBrowser/Private/Widgets/ContentBrowserTileViewWidget.cpp b/Engine/Plugins/Experimental/EditorDataStorageFeatures/Source/TedsContentBrowser/Private/Widgets/ContentBrowserTileViewWidget.cpp
--- a/Engine/Plugins/Experimental/EditorDataStorageFeatures/Source/TedsContentBrowser/Private/Widgets/ContentBrowserTileViewWidget.cpp
+++ b/Engine/Plugins/Experimental/EditorDataStorageFeatures/Source/TedsContentBrowser/Private/Widgets/ContentBrowserTileViewWidget.cpp
@@ -496,11 +496,11 @@ TSharedPtr<SWidget> FContentBrowserTileViewWidgetConstructor::CreateThumbnailWid
 {
 	using namespace UE::Editor::DataStorage;
 
-	TSharedPtr<FTypedElementWidgetConstructor> OutThumbnailWidgetConstructorPtr;
+	IUiProvider::FWidgetConstructorPtr OutThumbnailWidgetConstructorPtr;
 
-	auto AssignThumbnailWidgetToColumn = [&OutThumbnailWidgetConstructorPtr] (TUniquePtr<FTypedElementWidgetConstructor> Constructor, TConstArrayView<TWeakObjectPtr<const UScriptStruct>>)
+	auto AssignThumbnailWidgetToColumn = [&OutThumbnailWidgetConstructorPtr] (IUiProvider::FWidgetConstructorPtr WidgetConstructor, TConstArrayView<TWeakObjectPtr<const UScriptStruct>>)
 	{
-		OutThumbnailWidgetConstructorPtr = TSharedPtr<FTypedElementWidgetConstructor>(Constructor.Release());
+		OutThumbnailWidgetConstructorPtr = WidgetConstructor;
 		return false;
 	};
 	
@@ -573,11 +573,11 @@ TSharedPtr<SWidget> FContentBrowserTileViewWidgetConstructor::CreateItemTypeWidg
 {
 	using namespace UE::Editor::DataStorage;
 
-	TSharedPtr<FTypedElementWidgetConstructor> OutItemTypeWidgetConstructorPtr;
+	IUiProvider::FWidgetConstructorPtr OutItemTypeWidgetConstructorPtr;
 
-	auto AssignItemTypeWidgetToColumn = [&OutItemTypeWidgetConstructorPtr] (TUniquePtr<FTypedElementWidgetConstructor> Constructor, TConstArrayView<TWeakObjectPtr<const UScriptStruct>>)
+	auto AssignItemTypeWidgetToColumn = [&OutItemTypeWidgetConstructorPtr] (IUiProvider::FWidgetConstructorPtr WidgetConstructor, TConstArrayView<TWeakObjectPtr<const UScriptStruct>>)
 	{
-		OutItemTypeWidgetConstructorPtr = TSharedPtr<FTypedElementWidgetConstructor>(Constructor.Release());
+		OutItemTypeWidgetConstructorPtr = MoveTemp(WidgetConstructor);;
 		return false;
 	};
 
diff --git a/Engine/Plugins/Experimental/EditorDataStorageFeatures/Source/TedsDebugger/Private/Widgets/QueryEditor/TedsQueryEditorResultsView.cpp b/Engine/Plugins/Experimental/EditorDataStorageFeatures/Source/TedsDebugger/Private/Widgets/QueryEditor/TedsQueryEditorResultsView.cpp
--- a/Engine/Plugins/Experimental/EditorDataStorageFeatures/Source/TedsDebugger/Private/Widgets/QueryEditor/TedsQueryEditorResultsView.cpp
+++ b/Engine/Plugins/Experimental/EditorDataStorageFeatures/Source/TedsDebugger/Private/Widgets/QueryEditor/TedsQueryEditorResultsView.cpp
@@ -207,9 +207,8 @@ namespace UE::Editor::DataStorage::Debug::QueryEditor
 
 	void SResultsView::CreateRowHandleColumn()
 	{
-		auto AssignWidgetToColumn = [this](TUniquePtr<FTypedElementWidgetConstructor> Constructor, TConstArrayView<TWeakObjectPtr<const UScriptStruct>>)
+		auto AssignWidgetToColumn = [this](IUiProvider::FWidgetConstructorPtr WidgetConstructor, TConstArrayView<TWeakObjectPtr<const UScriptStruct>>)
 		{
-			TSharedPtr<FTypedElementWidgetConstructor> WidgetConstructor(Constructor.Release());
 			RowHandleColumn = MakeShared<FTedsTableViewerColumn>(TEXT("Row Handle"), WidgetConstructor);
 			return false;
 		};
diff --git a/Engine/Plugins/Experimental/EditorDataStorageFeatures/Source/TedsDebugger/Private/Widgets/STedsDiscoveryDebugger.cpp b/Engine/Plugins/Experimental/EditorDataStorageFeatures/Source/TedsDebugger/Private/Widgets/STedsDiscoveryDebugger.cpp
--- a/Engine/Plugins/Experimental/EditorDataStorageFeatures/Source/TedsDebugger/Private/Widgets/STedsDiscoveryDebugger.cpp
+++ b/Engine/Plugins/Experimental/EditorDataStorageFeatures/Source/TedsDebugger/Private/Widgets/STedsDiscoveryDebugger.cpp
@@ -223,9 +223,8 @@ namespace UE::Editor::DataStorage::Debug
 
 	void SDiscoveryDebuggerTab::CreateRowHandleColumn()
 	{
-		auto AssignWidgetToColumn = [this](TUniquePtr<FTypedElementWidgetConstructor> Constructor, TConstArrayView<TWeakObjectPtr<const UScriptStruct>>)
+		auto AssignWidgetToColumn = [this](IUiProvider::FWidgetConstructorPtr WidgetConstructor, TConstArrayView<TWeakObjectPtr<const UScriptStruct>>)
 			{
-				TSharedPtr<FTypedElementWidgetConstructor> WidgetConstructor(Constructor.Release());
 				RowHandleColumn = MakeShared<FTedsTableViewerColumn>(TEXT("Row Handle"), WidgetConstructor);
 				return false;
 			};
diff --git a/Engine/Plugins/Experimental/EditorDataStorageFeatures/Source/TedsOutliner/Private/Compatibility/SceneOutlinerRowHandleColumn.cpp b/Engine/Plugins/Experimental/EditorDataStorageFeatures/Source/TedsOutliner/Private/Compatibility/SceneOutlinerRowHandleColumn.cpp
--- a/Engine/Plugins/Experimental/EditorDataStorageFeatures/Source/TedsOutliner/Private/Compatibility/SceneOutlinerRowHandleColumn.cpp
+++ b/Engine/Plugins/Experimental/EditorDataStorageFeatures/Source/TedsOutliner/Private/Compatibility/SceneOutlinerRowHandleColumn.cpp
@@ -27,9 +27,8 @@ FSceneOutlinerRowHandleColumn::FSceneOutlinerRowHandleColumn(ISceneOutliner& Sce
 void FSceneOutlinerRowHandleColumn::CreateWidgetConstructor()
 {
 	using namespace UE::Editor::DataStorage;
-	auto AssignWidgetToColumn = [this](TUniquePtr<FTypedElementWidgetConstructor> Constructor, TConstArrayView<TWeakObjectPtr<const UScriptStruct>>)
+	auto AssignWidgetToColumn = [this](IUiProvider::FWidgetConstructorPtr WidgetConstructor, TConstArrayView<TWeakObjectPtr<const UScriptStruct>>)
 	{
-		TSharedPtr<FTypedElementWidgetConstructor> WidgetConstructor(Constructor.Release());
 		TableViewerColumn = MakeShared<FTedsTableViewerColumn>(TEXT("Row Handle"), WidgetConstructor);
 		return false;
 	};
diff --git a/Engine/Plugins/Experimental/EditorDataStorageFeatures/Source/TedsOutliner/Private/TedsOutlinerColumn.h b/Engine/Plugins/Experimental/EditorDataStorageFeatures/Source/TedsOutliner/Private/TedsOutlinerColumn.h
--- a/Engine/Plugins/Experimental/EditorDataStorageFeatures/Source/TedsOutliner/Private/TedsOutlinerColumn.h
+++ b/Engine/Plugins/Experimental/EditorDataStorageFeatures/Source/TedsOutliner/Private/TedsOutlinerColumn.h
@@ -30,8 +30,8 @@ namespace UE::Editor::Outliner
 
 		// TEDS UI info
 		TArray<TWeakObjectPtr<const UScriptStruct>> ColumnTypes;
-		TSharedPtr<FTypedElementWidgetConstructor> HeaderWidgetConstructor;
-		TSharedPtr<FTypedElementWidgetConstructor> CellWidgetConstructor;
+		DataStorage::IUiProvider::FWidgetConstructorPtr HeaderWidgetConstructor;
+		DataStorage::IUiProvider::FWidgetConstructorPtr CellWidgetConstructor;
 
 		// Fallbacks to work with legacy columns
 		TSharedPtr<ISceneOutlinerColumn> FallbackColumn;
diff --git a/Engine/Plugins/Experimental/EditorDataStorageFeatures/Source/TedsOutliner/Private/TedsOutlinerImpl.cpp b/Engine/Plugins/Experimental/EditorDataStorageFeatures/Source/TedsOutliner/Private/TedsOutlinerImpl.cpp
--- a/Engine/Plugins/Experimental/EditorDataStorageFeatures/Source/TedsOutliner/Private/TedsOutlinerImpl.cpp
+++ b/Engine/Plugins/Experimental/EditorDataStorageFeatures/Source/TedsOutliner/Private/TedsOutlinerImpl.cpp
@@ -108,7 +108,7 @@ namespace QueryUtils
 
 namespace LabelWidgetUtils
 {
-	static TUniquePtr<FTypedElementWidgetConstructor> CreateLabelWidgetConstructorInternal(
+	static IUiProvider::FWidgetConstructorPtr CreateLabelWidgetConstructorInternal(
 			const DataStorage::ICoreProvider& Storage,
 			DataStorage::IUiProvider& StorageUi,
 			const DataStorage::IUiProvider::FPurposeID& Purpose,
@@ -121,11 +121,11 @@ namespace LabelWidgetUtils
 			ColumnTypes.Add(&ColumnType);
 		});
 	
-		TUniquePtr<FTypedElementWidgetConstructor> WidgetConstructor;
+		IUiProvider::FWidgetConstructorPtr WidgetConstructor;
 		StorageUi.CreateWidgetConstructors(StorageUi.FindPurpose(Purpose), DataStorage::IUiProvider::EMatchApproach::LongestMatch, ColumnTypes, {},
-			[&WidgetConstructor](TUniquePtr<FTypedElementWidgetConstructor> CreatedConstructor, TConstArrayView<TWeakObjectPtr<const UScriptStruct>>)
+			[&WidgetConstructor](IUiProvider::FWidgetConstructorPtr InWidgetConstructor, TConstArrayView<TWeakObjectPtr<const UScriptStruct>>)
 			{
-				WidgetConstructor = MoveTemp(CreatedConstructor);
+				WidgetConstructor = MoveTemp(InWidgetConstructor);
 				// Either this was the exact match so no need to search further or the longest possible chain didn't match so the next ones will 
 				// always be shorter in both cases just return.
 				return false;
@@ -404,7 +404,7 @@ TSharedRef<SWidget> FTedsOutlinerImpl::CreateLabelWidget(
 	const STableRow<FSceneOutlinerTreeItemPtr>& RowItem,
 	bool bInteractable)
 {
-	TUniquePtr<FTypedElementWidgetConstructor> WidgetConstructor = LabelWidgetUtils::CreateLabelWidgetConstructorInternal(Storage, StorageUi, Purpose, TargetRow);
+	IUiProvider::FWidgetConstructorPtr WidgetConstructor = LabelWidgetUtils::CreateLabelWidgetConstructorInternal(Storage, StorageUi, Purpose, TargetRow);
 	if (!WidgetConstructor)
 	{
 		return SNullWidget::NullWidget;
@@ -572,10 +572,8 @@ void FTedsOutlinerImpl::RegenerateColumns()
 
 	int32 IndexOffset = 0;
 	auto ColumnConstructor = [this, ColumnQueryColumn, TedsOutlinerColumn, MetaDataView, &IndexOffset, HeaderPurposeRow, Dealiaser](
-		TUniquePtr<FTypedElementWidgetConstructor> Constructor, TConstArrayView<TWeakObjectPtr<const UScriptStruct>> ColumnTypes)
+		IUiProvider::FWidgetConstructorPtr CellConstructor, TConstArrayView<TWeakObjectPtr<const UScriptStruct>> ColumnTypes)
 		{
-			TSharedPtr<FTypedElementWidgetConstructor> CellConstructor(Constructor.Release());
-
 			/* If we have a fallback column for this query, remove it, take over it's priority and 
 			 * replace it with the TEDS column. But also allow the TEDS-Outliner column to fallback to it for
 			 * data not in TEDS yet.
@@ -603,7 +601,7 @@ void FTedsOutlinerImpl::RegenerateColumns()
 					FCreateSceneOutlinerColumn::CreateLambda(
 						[this, ColumnQueryColumn, TedsOutlinerColumn, MetaDataView, NameId, &ColumnTypes, CellConstructor, FallbackColumnPtr, HeaderPurposeRow, Dealiaser](ISceneOutliner&)
 						{
-							TSharedPtr<FTypedElementWidgetConstructor> HeaderConstructor = 
+							IUiProvider::FWidgetConstructorPtr HeaderConstructor = 
 								UE::Editor::DataStorage::TableViewerUtils::CreateHeaderWidgetConstructor(*StorageUi, MetaDataView, ColumnTypes, HeaderPurposeRow);
 
 							DataStorage::FTedsTableViewerColumn::FSortHandler SortDelegates
diff --git a/Engine/Plugins/Experimental/EditorDataStorageFeatures/Source/TedsOutliner/Private/Widgets/OutlinerLabelWidget.cpp b/Engine/Plugins/Experimental/EditorDataStorageFeatures/Source/TedsOutliner/Private/Widgets/OutlinerLabelWidget.cpp
--- a/Engine/Plugins/Experimental/EditorDataStorageFeatures/Source/TedsOutliner/Private/Widgets/OutlinerLabelWidget.cpp
+++ b/Engine/Plugins/Experimental/EditorDataStorageFeatures/Source/TedsOutliner/Private/Widgets/OutlinerLabelWidget.cpp
@@ -61,18 +61,18 @@ TSharedPtr<SWidget> FOutlinerLabelWidgetConstructor::CreateWidget(UE::Editor::Da
 {
 	if (DataStorage->IsRowAvailable(TargetRow))
 	{
-		TSharedPtr<FTypedElementWidgetConstructor> OutIconWidgetConstructorPtr;
-		TSharedPtr<FTypedElementWidgetConstructor> OutTextWidgetConstructorPtr;
+		UE::Editor::DataStorage::IUiProvider::FWidgetConstructorPtr OutIconWidgetConstructorPtr;
+		UE::Editor::DataStorage::IUiProvider::FWidgetConstructorPtr OutTextWidgetConstructorPtr;
 
-		auto AssignIconWidgetToColumn = [&OutIconWidgetConstructorPtr](TUniquePtr<FTypedElementWidgetConstructor> Constructor, TConstArrayView<TWeakObjectPtr<const UScriptStruct>>)
+		auto AssignIconWidgetToColumn = [&OutIconWidgetConstructorPtr](UE::Editor::DataStorage::IUiProvider::FWidgetConstructorPtr WidgetConstructor, TConstArrayView<TWeakObjectPtr<const UScriptStruct>>)
 			{
-				OutIconWidgetConstructorPtr = TSharedPtr<FTypedElementWidgetConstructor>(Constructor.Release());
+				OutIconWidgetConstructorPtr = MoveTemp(WidgetConstructor);
 				return false;
 			};
 
-		auto AssignTextWidgetToColumn = [&OutTextWidgetConstructorPtr](TUniquePtr<FTypedElementWidgetConstructor> Constructor, TConstArrayView<TWeakObjectPtr<const UScriptStruct>>)
+		auto AssignTextWidgetToColumn = [&OutTextWidgetConstructorPtr](UE::Editor::DataStorage::IUiProvider::FWidgetConstructorPtr WidgetConstructor, TConstArrayView<TWeakObjectPtr<const UScriptStruct>>)
 			{
-				OutTextWidgetConstructorPtr = TSharedPtr<FTypedElementWidgetConstructor>(Constructor.Release());
+				OutTextWidgetConstructorPtr = MoveTemp(WidgetConstructor);
 				return false;
 			};
 
diff --git a/Engine/Plugins/Experimental/EditorDataStorageFeatures/Source/TedsTableViewer/Private/TedsTableViewerColumn.cpp b/Engine/Plugins/Experimental/EditorDataStorageFeatures/Source/TedsTableViewer/Private/TedsTableViewerColumn.cpp
--- a/Engine/Plugins/Experimental/EditorDataStorageFeatures/Source/TedsTableViewer/Private/TedsTableViewerColumn.cpp
+++ b/Engine/Plugins/Experimental/EditorDataStorageFeatures/Source/TedsTableViewer/Private/TedsTableViewerColumn.cpp
@@ -20,8 +20,8 @@
 
 namespace UE::Editor::DataStorage
 {
-	FTedsTableViewerColumn::FTedsTableViewerColumn(const FName& InColumnName, const TSharedPtr<FTypedElementWidgetConstructor>& InCellWidgetConstructor,
-		const TArray<TWeakObjectPtr<const UScriptStruct>>& InMatchedColumns, const TSharedPtr<FTypedElementWidgetConstructor>& InHeaderWidgetConstructor,
+	FTedsTableViewerColumn::FTedsTableViewerColumn(const FName& InColumnName, const IUiProvider::FWidgetConstructorPtr& InCellWidgetConstructor,
+		const TArray<TWeakObjectPtr<const UScriptStruct>>& InMatchedColumns, const IUiProvider::FWidgetConstructorPtr& InHeaderWidgetConstructor,
 		const FMetaData& InWidgetMetaData, const TArray<TSharedPtr<FColumnMetaData>>& InColumnMetaData)
 		: ColumnName(InColumnName)
 		, CellWidgetConstructor(InCellWidgetConstructor)
@@ -64,7 +64,7 @@ namespace UE::Editor::DataStorage
 	{
 		TSharedPtr<SWidget> RowWidget;
 
-		if (Storage->IsRowAssigned(InRowHandle))
+		if (CellWidgetConstructor && Storage->IsRowAssigned(InRowHandle))
 		{
 			const RowHandle UiRowHandle = Storage->AddRow(Storage->FindTable(TableViewerUtils::GetWidgetTableName()));
 
@@ -322,6 +322,11 @@ namespace UE::Editor::DataStorage
 
 	void FTedsTableViewerColumn::UpdateWidgets()
 	{
+		if (!CellWidgetConstructor)
+		{
+			return;
+		}
+		
 		// Remove any widget rows that don't actually need an update
 		RowsToUpdate = RowsToUpdate.FilterByPredicate([this](const TPair<RowHandle, bool>& Pair) -> bool
 		{
diff --git a/Engine/Plugins/Experimental/EditorDataStorageFeatures/Source/TedsTableViewer/Private/TedsTableViewerModel.cpp b/Engine/Plugins/Experimental/EditorDataStorageFeatures/Source/TedsTableViewer/Private/TedsTableViewerModel.cpp
--- a/Engine/Plugins/Experimental/EditorDataStorageFeatures/Source/TedsTableViewer/Private/TedsTableViewerModel.cpp
+++ b/Engine/Plugins/Experimental/EditorDataStorageFeatures/Source/TedsTableViewer/Private/TedsTableViewerModel.cpp
@@ -229,11 +229,9 @@ namespace UE::Editor::DataStorage
 		
 		// Lambda to create the constructor for a given list of columns
 		auto ColumnConstructor = [this, &IndexOffset, &NewColumnMap, HeaderWidgetPurposeRow](
-			TUniquePtr<FTypedElementWidgetConstructor> Constructor, TConstArrayView<TWeakObjectPtr<const UScriptStruct>> MatchedColumns)
+			IUiProvider::FWidgetConstructorPtr CellConstructor, TConstArrayView<TWeakObjectPtr<const UScriptStruct>> MatchedColumns)
 			{
-				TSharedPtr<FTypedElementWidgetConstructor> CellConstructor(Constructor.Release());
-
-				TSharedPtr<FTypedElementWidgetConstructor> HeaderConstructor =
+				IUiProvider::FWidgetConstructorPtr HeaderConstructor =
 					TableViewerUtils::CreateHeaderWidgetConstructor(*StorageUi, FColumnsMetaDataView(ColumnMetaData), MatchedColumns, HeaderWidgetPurposeRow);
 			
 				FName NameId = TableViewerUtils::FindLongestMatchingName(MatchedColumns, IndexOffset);
diff --git a/Engine/Plugins/Experimental/EditorDataStorageFeatures/Source/TedsTableViewer/Private/TedsTableViewerUtils.cpp b/Engine/Plugins/Experimental/EditorDataStorageFeatures/Source/TedsTableViewer/Private/TedsTableViewerUtils.cpp
--- a/Engine/Plugins/Experimental/EditorDataStorageFeatures/Source/TedsTableViewer/Private/TedsTableViewerUtils.cpp
+++ b/Engine/Plugins/Experimental/EditorDataStorageFeatures/Source/TedsTableViewer/Private/TedsTableViewerUtils.cpp
@@ -105,22 +105,22 @@ namespace UE::Editor::DataStorage::TableViewerUtils
 		return VerifiedColumnTypes;
 	}
 
-	TSharedPtr<FTypedElementWidgetConstructor> CreateHeaderWidgetConstructor(IUiProvider& StorageUi,
+	IUiProvider::FWidgetConstructorPtr CreateHeaderWidgetConstructor(IUiProvider& StorageUi,
 		const FMetaDataView& InMetaData, TConstArrayView<TWeakObjectPtr<const UScriptStruct>> ColumnTypes, RowHandle PurposeRow)
 	{
 		using MatchApproach = IUiProvider::EMatchApproach;
 
 		TArray<TWeakObjectPtr<const UScriptStruct>> VerifiedColumnTypes = CreateVerifiedColumnTypeArray(ColumnTypes);
-		TSharedPtr<FTypedElementWidgetConstructor> Constructor = nullptr;
+		IUiProvider::FWidgetConstructorPtr Constructor = nullptr;
 
 		StorageUi.CreateWidgetConstructors(PurposeRow, MatchApproach::ExactMatch, VerifiedColumnTypes, InMetaData,
-				[&Constructor, ColumnTypes](
-					TUniquePtr<FTypedElementWidgetConstructor> CreatedConstructor, 
+				[&Constructor, VerifiedColumnTypes](
+					IUiProvider::FWidgetConstructorPtr WidgetConstructor, 
 					TConstArrayView<TWeakObjectPtr<const UScriptStruct>> MatchedColumnTypes)
 				{
-					if (ColumnTypes.Num() == MatchedColumnTypes.Num())
+					if (VerifiedColumnTypes.Num() == MatchedColumnTypes.Num())
 					{
-						Constructor = TSharedPtr<FTypedElementWidgetConstructor>(CreatedConstructor.Release());
+						Constructor = MoveTemp(WidgetConstructor);
 					}
 					// Either this was the exact match so no need to search further or the longest possible chain didn't match so the next ones will 
 					// always be shorter in both cases just return.
diff --git a/Engine/Plugins/Experimental/EditorDataStorageFeatures/Source/TedsTableViewer/Private/Widgets/SRowDetails.cpp b/Engine/Plugins/Experimental/EditorDataStorageFeatures/Source/TedsTableViewer/Private/Widgets/SRowDetails.cpp
--- a/Engine/Plugins/Experimental/EditorDataStorageFeatures/Source/TedsTableViewer/Private/Widgets/SRowDetails.cpp
+++ b/Engine/Plugins/Experimental/EditorDataStorageFeatures/Source/TedsTableViewer/Private/Widgets/SRowDetails.cpp
@@ -80,9 +80,9 @@ namespace UE::Editor::DataStorage
 			
 			DataStorageUi->CreateWidgetConstructors(PurposeRow, IUiProvider::EMatchApproach::LongestMatch,
 					Columns, {}, [this, Row](
-						TUniquePtr<FTypedElementWidgetConstructor> Constructor, TConstArrayView<TWeakObjectPtr<const UScriptStruct>> Columns)
+						IUiProvider::FWidgetConstructorPtr WidgetConstructor, TConstArrayView<TWeakObjectPtr<const UScriptStruct>> Columns)
 					{
-						Items.Add(MakeShared<FRowDetailsItem>(nullptr, MoveTemp(Constructor), Row));
+						Items.Add(MakeShared<FRowDetailsItem>(nullptr, MoveTemp(WidgetConstructor), Row));
 						return true;
 					});
 			
@@ -93,7 +93,7 @@ namespace UE::Editor::DataStorage
 				{
 					DataStorageUi->CreateWidgetConstructors(DataStorageUi->FindPurpose(DataStorageUi->GetDefaultWidgetPurposeID()), {},
 						[this, Column, Row](
-							TUniquePtr<FTypedElementWidgetConstructor> Constructor, TConstArrayView<TWeakObjectPtr<const UScriptStruct>> Columns)
+							IUiProvider::FWidgetConstructorPtr Constructor, TConstArrayView<TWeakObjectPtr<const UScriptStruct>> Columns)
 						{
 							Items.Add(MakeShared<FRowDetailsItem>(Column, MoveTemp(Constructor), Row));
 							return true;
@@ -176,7 +176,7 @@ namespace UE::Editor::DataStorage
 
 	// FRowDetailsItem
 	FRowDetailsItem::FRowDetailsItem(const TWeakObjectPtr<const UScriptStruct>& InColumnType,
-		TUniquePtr<FTypedElementWidgetConstructor> InWidgetConstructor, RowHandle InRow)
+		IUiProvider::FWidgetConstructorPtr InWidgetConstructor, RowHandle InRow)
 		: ColumnType(InColumnType)
 		, WidgetConstructor(MoveTemp(InWidgetConstructor))
 		, Row(InRow)
diff --git a/Engine/Plugins/Experimental/EditorDataStorageFeatures/Source/TedsTableViewer/Private/Widgets/STedsTileViewer.cpp b/Engine/Plugins/Experimental/EditorDataStorageFeatures/Source/TedsTableViewer/Private/Widgets/STedsTileViewer.cpp
--- a/Engine/Plugins/Experimental/EditorDataStorageFeatures/Source/TedsTableViewer/Private/Widgets/STedsTileViewer.cpp
+++ b/Engine/Plugins/Experimental/EditorDataStorageFeatures/Source/TedsTableViewer/Private/Widgets/STedsTileViewer.cpp
@@ -101,10 +101,9 @@ namespace UE::Editor::DataStorage
 		// of columns
 		TSharedPtr<FTedsTableViewerColumn> Column;
 		
-		auto AssignWidgetToColumn = [&Column](TUniquePtr<FTypedElementWidgetConstructor> Constructor, TConstArrayView<TWeakObjectPtr<const UScriptStruct>> InMatchedColumns)
+		auto AssignWidgetToColumn = [&Column](IUiProvider::FWidgetConstructorPtr WidgetConstructor, TConstArrayView<TWeakObjectPtr<const UScriptStruct>> InMatchedColumns)
 		{
 			TArray<TWeakObjectPtr<const UScriptStruct>> MatchedColumnsCopy(InMatchedColumns);
-			TSharedPtr<FTypedElementWidgetConstructor> WidgetConstructor(Constructor.Release());
 			Column = MakeShared<FTedsTableViewerColumn>(TEXT("TileView"), WidgetConstructor, MatchedColumnsCopy);
 			return false;
 		};
diff --git a/Engine/Plugins/Experimental/EditorDataStorageFeatures/Source/TedsTableViewer/Public/TedsTableViewerColumn.h b/Engine/Plugins/Experimental/EditorDataStorageFeatures/Source/TedsTableViewer/Public/TedsTableViewerColumn.h
--- a/Engine/Plugins/Experimental/EditorDataStorageFeatures/Source/TedsTableViewer/Public/TedsTableViewerColumn.h
+++ b/Engine/Plugins/Experimental/EditorDataStorageFeatures/Source/TedsTableViewer/Public/TedsTableViewerColumn.h
@@ -31,9 +31,9 @@ namespace UE::Editor::DataStorage
 
 		TEDSTABLEVIEWER_API FTedsTableViewerColumn(
 			const FName& ColumnName, // The unique ID of this column
-			const TSharedPtr<FTypedElementWidgetConstructor>& InCellWidgetConstructor, // The widget constructor to use for this column
+			const IUiProvider::FWidgetConstructorPtr& InCellWidgetConstructor, // The widget constructor to use for this column
 			const TArray<TWeakObjectPtr<const UScriptStruct>>& InMatchedColumns = {}, // An optional list of matched Teds columns 
-			const TSharedPtr<FTypedElementWidgetConstructor>& InHeaderWidgetConstructor = nullptr, // Optional constructor to use for the header widget
+			const IUiProvider::FWidgetConstructorPtr& InHeaderWidgetConstructor = nullptr, // Optional constructor to use for the header widget
 			const FMetaData& MetaData = FMetaData(), // Optional generic metadata to use when constructing widgets
 			const TArray<TSharedPtr<FColumnMetaData>>& ColumnMetaData = {});  // Optional metadata specific to the matched columns
 
@@ -92,8 +92,8 @@ namespace UE::Editor::DataStorage
 		FName ColumnName;
 
 		// Widget Constructors
-		TSharedPtr<FTypedElementWidgetConstructor> CellWidgetConstructor;
-		TSharedPtr<FTypedElementWidgetConstructor> HeaderWidgetConstructor;
+		IUiProvider::FWidgetConstructorPtr CellWidgetConstructor;
+		IUiProvider::FWidgetConstructorPtr HeaderWidgetConstructor;
 
 		// Teds Columns this widget constructor matched with
 		TArray<TWeakObjectPtr<const UScriptStruct>> MatchedColumns;
diff --git a/Engine/Plugins/Experimental/EditorDataStorageFeatures/Source/TedsTableViewer/Public/TedsTableViewerUtils.h b/Engine/Plugins/Experimental/EditorDataStorageFeatures/Source/TedsTableViewer/Public/TedsTableViewerUtils.h
--- a/Engine/Plugins/Experimental/EditorDataStorageFeatures/Source/TedsTableViewer/Public/TedsTableViewerUtils.h
+++ b/Engine/Plugins/Experimental/EditorDataStorageFeatures/Source/TedsTableViewer/Public/TedsTableViewerUtils.h
@@ -35,7 +35,7 @@ namespace UE::Editor::DataStorage::TableViewerUtils
 	TEDSTABLEVIEWER_API FName FindLongestMatchingName(TConstArrayView<TWeakObjectPtr<const UScriptStruct>> ColumnTypes, int32 DefaultNameIndex);
 
 	// Create a header widget constructor for the given columns
-	TEDSTABLEVIEWER_API TSharedPtr<FTypedElementWidgetConstructor> CreateHeaderWidgetConstructor(IUiProvider& StorageUi, 
+	TEDSTABLEVIEWER_API IUiProvider::FWidgetConstructorPtr CreateHeaderWidgetConstructor(IUiProvider& StorageUi, 
 		const FMetaDataView& InMetaData, TConstArrayView<TWeakObjectPtr<const UScriptStruct>> ColumnTypes, RowHandle PurposeRow);
 
 	// Create a copy of the provided column types array after discarding invalid entries
diff --git a/Engine/Plugins/Experimental/EditorDataStorageFeatures/Source/TedsTableViewer/Public/Widgets/SRowDetails.h b/Engine/Plugins/Experimental/EditorDataStorageFeatures/Source/TedsTableViewer/Public/Widgets/SRowDetails.h
--- a/Engine/Plugins/Experimental/EditorDataStorageFeatures/Source/TedsTableViewer/Public/Widgets/SRowDetails.h
+++ b/Engine/Plugins/Experimental/EditorDataStorageFeatures/Source/TedsTableViewer/Public/Widgets/SRowDetails.h
@@ -23,12 +23,12 @@ namespace UE::Editor::DataStorage
 		TWeakObjectPtr<const UScriptStruct> ColumnType;
 
 		// Widget for the column
-		TUniquePtr<FTypedElementWidgetConstructor> WidgetConstructor;
+		IUiProvider::FWidgetConstructorPtr WidgetConstructor;
 		
 		RowHandle Row = InvalidRowHandle;
 		RowHandle WidgetRow = InvalidRowHandle;
 
-		FRowDetailsItem(const TWeakObjectPtr<const UScriptStruct>& InColumnType, TUniquePtr<FTypedElementWidgetConstructor> InWidgetConstructor,
+		FRowDetailsItem(const TWeakObjectPtr<const UScriptStruct>& InColumnType, IUiProvider::FWidgetConstructorPtr InWidgetConstructor,
 			RowHandle InRow);
 	};
 	
diff --git a/Engine/Source/Editor/ContentBrowser/Private/ContentSources/Widgets/SContentSource.cpp b/Engine/Source/Editor/ContentBrowser/Private/ContentSources/Widgets/SContentSource.cpp
--- a/Engine/Source/Editor/ContentBrowser/Private/ContentSources/Widgets/SContentSource.cpp
+++ b/Engine/Source/Editor/ContentBrowser/Private/ContentSources/Widgets/SContentSource.cpp
@@ -61,11 +61,11 @@ namespace UE::Editor::ContentBrowser
 		using namespace UE::Editor::DataStorage;
 
 		// Otherwise we display any widget constructors registered for the "ContentBrowser.AssetView" purpose currently
-		TSharedPtr<FTypedElementWidgetConstructor> AssetViewWidgetConstructor;
+		IUiProvider::FWidgetConstructorPtr AssetViewWidgetConstructor;
 
-		auto AssignWidgetToColumn = [&AssetViewWidgetConstructor] (TUniquePtr<FTypedElementWidgetConstructor> Constructor, TConstArrayView<TWeakObjectPtr<const UScriptStruct>>)
+		auto AssignWidgetToColumn = [&AssetViewWidgetConstructor] (IUiProvider::FWidgetConstructorPtr Constructor, TConstArrayView<TWeakObjectPtr<const UScriptStruct>>)
 		{
-			AssetViewWidgetConstructor = TSharedPtr<FTypedElementWidgetConstructor>(Constructor.Release());
+			AssetViewWidgetConstructor = MoveTemp(Constructor);;
 			return false;
 		};
 
diff --git a/Engine/Source/Runtime/TypedElementFramework/Public/Elements/Interfaces/TypedElementDataStorageUiInterface.h b/Engine/Source/Runtime/TypedElementFramework/Public/Elements/Interfaces/TypedElementDataStorageUiInterface.h
--- a/Engine/Source/Runtime/TypedElementFramework/Public/Elements/Interfaces/TypedElementDataStorageUiInterface.h
+++ b/Engine/Source/Runtime/TypedElementFramework/Public/Elements/Interfaces/TypedElementDataStorageUiInterface.h
@@ -381,8 +381,9 @@ namespace UE::Editor::DataStorage
 			TYPEDELEMENTFRAMEWORK_API FName ToString() const;
 		};
 
+		using FWidgetConstructorPtr = TSharedPtr<FTypedElementWidgetConstructor>;
 		using WidgetCreatedCallback = TFunctionRef<void(const TSharedRef<SWidget>& NewWidget, UE::Editor::DataStorage::RowHandle Row)>;
-		using WidgetConstructorCallback = TFunctionRef<bool(TUniquePtr<FTypedElementWidgetConstructor>, TConstArrayView<TWeakObjectPtr<const UScriptStruct>>)>;
+		using WidgetConstructionCallback = TFunctionRef<bool(FWidgetConstructorPtr, TConstArrayView<TWeakObjectPtr<const UScriptStruct>>)>;
 		using WidgetPurposeCallback = TFunctionRef<void(FName, EPurposeType, const FText&)>;
 		using PropertySorterConstructorCallback = 
 			TFunction<TSharedPtr<const FColumnSorterInterface>(TWeakObjectPtr<const UScriptStruct> ColumnType, const FProperty& Property)>;
@@ -501,7 +502,7 @@ namespace UE::Editor::DataStorage
 		 * If no factories were found for the requested purpose, the purpose's parent chain will be traversed to try and find an ancestor purpose with factories
 		 */
 		virtual void CreateWidgetConstructors(RowHandle PurposeRow, 
-			const FMetaDataView& Arguments, const WidgetConstructorCallback& Callback) = 0;
+			const FMetaDataView& Arguments, const WidgetConstructionCallback& Callback) = 0;
 	
 		/** 
 		 * Finds matching widget constructors for provided columns, preferring longer matches over shorter matches.
@@ -513,7 +514,7 @@ namespace UE::Editor::DataStorage
 		 * the remaining columns
 		 */
 		virtual void CreateWidgetConstructors(RowHandle PurposeRow, EMatchApproach MatchApproach, TArray<TWeakObjectPtr<const UScriptStruct>>& Columns,
-			const FMetaDataView& Arguments, const WidgetConstructorCallback& Callback) = 0;
+			const FMetaDataView& Arguments, const WidgetConstructionCallback& Callback) = 0;
 	
 		/**
 		 * Creates all the widgets registered under the provided purpose. This may be a large number of widgets for a wide name
@@ -590,6 +591,20 @@ namespace UE::Editor::DataStorage
 		 * Please replace use of the functions below with the recommended replacements.
 		 */
 
+		// Old widget constructor callback that takes in a UniquePtr, all new code should use the new WidgetConstructionCallback instead
+		using WidgetConstructorCallback UE_DEPRECATED(5.8, "Use the new WidgetConstructionCallback alias instead") =
+			TFunctionRef<bool(TUniquePtr<FTypedElementWidgetConstructor>, TConstArrayView<TWeakObjectPtr<const UScriptStruct>>)>;
+
+		PRAGMA_DISABLE_DEPRECATION_WARNINGS
+		
+		UE_DEPRECATED(5.8, "Use the overload that uses the new WidgetConstructionCallback instead")
+		virtual void CreateWidgetConstructors(RowHandle PurposeRow, 
+			const FMetaDataView& Arguments, const WidgetConstructorCallback& Callback) = 0;
+
+		UE_DEPRECATED(5.8, "Use the overload that uses the new WidgetConstructionCallback instead")
+		virtual void CreateWidgetConstructors(RowHandle PurposeRow, EMatchApproach MatchApproach, TArray<TWeakObjectPtr<const UScriptStruct>>& Columns,
+			const FMetaDataView& Arguments, const WidgetConstructorCallback& Callback) = 0;
+		
 		UE_DEPRECATED(5.7, "Registering Widget Factories by instance is not supported anymore, register them by type instead")
 		virtual bool RegisterWidgetFactory(RowHandle PurposeRow, TUniquePtr<FTypedElementWidgetConstructor>&& Constructor) = 0;
 	
@@ -630,7 +645,8 @@ namespace UE::Editor::DataStorage
 
 		UE_DEPRECATED(5.6, "Use the version that takes in PurposeRowHandle instead. FindPurpose() to lookup row handle by PurposeID")
 		virtual void ConstructWidgets(FName Purpose, const FMetaDataView& Arguments, const WidgetCreatedCallback& ConstructionCallback) = 0;
-
+		
+		PRAGMA_ENABLE_DEPRECATION_WARNINGS
 	};
 
 	using FWidgetConstructor = FTypedElementWidgetConstructor;
