Teds memento system:
Add functionality to add a column to a row when a memento is restored.

Teds Visibility:
Rework the visibility column to use a new FVisibilityChangeTag to sync changes back to the world instead of the FTypedElementSyncBackToWorldTag.
Add a memento translater for the visibility column to make persist across re-instancing, also add FVisibilityChangeTag on memento restore.

These changes allow us to sync visibility changes back to the world after a memento restoration without adding the FTypedElementSyncBackToWorldTag to the row which can cause a lot of other data to get synced for a lot of UObjects since FVisibleinEditorColumn exists on most actors and other objects.

#Synced-CL 49950745

diff --git a/Engine/Plugins/ChaosVD/Source/ChaosVD/Private/TEDS/ChaosVDDataStorageVisibilityQueries.cpp b/Engine/Plugins/ChaosVD/Source/ChaosVD/Private/TEDS/ChaosVDDataStorageVisibilityQueries.cpp
--- a/Engine/Plugins/ChaosVD/Source/ChaosVD/Private/TEDS/ChaosVDDataStorageVisibilityQueries.cpp
+++ b/Engine/Plugins/ChaosVD/Source/ChaosVD/Private/TEDS/ChaosVDDataStorageVisibilityQueries.cpp
@@ -116,7 +116,7 @@ void UChaosVDDataStorageVisibilityQueries::RegisterVisibilityColumnToParticleQue
 				}
 			})
 		.Where()
-			.All<FChaosVDObjectDataTag, FTypedElementSyncBackToWorldTag>()
+			.All<FChaosVDObjectDataTag, FVisibilityChangedTag>()
 		.Compile()
 	);
 }
diff --git a/Engine/Plugins/Experimental/Dataflow/Source/DataflowEditor/Private/Dataflow/DataflowOutlinerFactory.cpp b/Engine/Plugins/Experimental/Dataflow/Source/DataflowEditor/Private/Dataflow/DataflowOutlinerFactory.cpp
--- a/Engine/Plugins/Experimental/Dataflow/Source/DataflowEditor/Private/Dataflow/DataflowOutlinerFactory.cpp
+++ b/Engine/Plugins/Experimental/Dataflow/Source/DataflowEditor/Private/Dataflow/DataflowOutlinerFactory.cpp
@@ -249,7 +249,7 @@ void UDataflowObjectFactory::RegisterVisibilityQueries(UE::Editor::DataStorage::
 					});
 				}
 			})
-		.Where(TColumn<FDataflowSceneObjectTag>() && TColumn<FTypedElementSyncBackToWorldTag>())
+		.Where(TColumn<FDataflowSceneObjectTag>() && TColumn<FVisibilityChangedTag>())
 		.Compile()
 	);
 
@@ -269,7 +269,7 @@ void UDataflowObjectFactory::RegisterVisibilityQueries(UE::Editor::DataStorage::
 					});
 				}
 			})
-		.Where(TColumn<FDataflowSceneStructTag>() && TColumn<FTypedElementSyncBackToWorldTag>())
+		.Where(TColumn<FDataflowSceneStructTag>() && TColumn<FVisibilityChangedTag>())
 		.Compile()
 	);
 }
diff --git a/Engine/Plugins/Experimental/Dataflow/Source/DataflowEditor/Private/Dataflow/DataflowOutlinerMode.cpp b/Engine/Plugins/Experimental/Dataflow/Source/DataflowEditor/Private/Dataflow/DataflowOutlinerMode.cpp
--- a/Engine/Plugins/Experimental/Dataflow/Source/DataflowEditor/Private/Dataflow/DataflowOutlinerMode.cpp
+++ b/Engine/Plugins/Experimental/Dataflow/Source/DataflowEditor/Private/Dataflow/DataflowOutlinerMode.cpp
@@ -83,9 +83,11 @@ namespace UE::Dataflow::Private
 				ICoreProvider* Storage = GetMutableDataStorageFeature<ICoreProvider>(StorageFeatureName);
 				if (TEDSItem->IsValid() && Storage)
 				{
-					if(FVisibleInEditorColumn* VisibilityColumn = Storage->GetColumn<FVisibleInEditorColumn>(TEDSItem->GetRowHandle()))
+					RowHandle Row = TEDSItem->GetRowHandle();
+					if(FVisibleInEditorColumn* VisibilityColumn = Storage->GetColumn<FVisibleInEditorColumn>(Row))
 					{
 						VisibilityColumn->bIsVisibleInEditor = bIsVisible;
+						Storage->AddColumn<FVisibilityChangedTag>(Row);
 					}
 				}
 			}
diff --git a/Engine/Plugins/Experimental/EditorDataStorage/Source/TedsCore/Private/Memento/TypedElementMementoSystem.cpp b/Engine/Plugins/Experimental/EditorDataStorage/Source/TedsCore/Private/Memento/TypedElementMementoSystem.cpp
--- a/Engine/Plugins/Experimental/EditorDataStorage/Source/TedsCore/Private/Memento/TypedElementMementoSystem.cpp
+++ b/Engine/Plugins/Experimental/EditorDataStorage/Source/TedsCore/Private/Memento/TypedElementMementoSystem.cpp
@@ -100,6 +100,11 @@ namespace UE::Editor::DataStorage
 					{
 						ColumnType.CopyScriptStruct(Destination, Source);
 					});
+				
+				if (const UScriptStruct* Tag = Translator->GetColumnToAddOnRestore())
+				{
+					DataStorage.AddColumn(TargetRow, Tag);
+				}
 
 				UE_LOG(LogEditorDataStorage, VeryVerbose, 
 					TEXT("Memento->Column: %llu -> %llu"), MementoRow, TargetRow);
diff --git a/Engine/Plugins/Experimental/EditorDataStorage/Source/TedsCore/Private/Memento/TypedElementMementoTranslators.cpp b/Engine/Plugins/Experimental/EditorDataStorage/Source/TedsCore/Private/Memento/TypedElementMementoTranslators.cpp
--- a/Engine/Plugins/Experimental/EditorDataStorage/Source/TedsCore/Private/Memento/TypedElementMementoTranslators.cpp
+++ b/Engine/Plugins/Experimental/EditorDataStorage/Source/TedsCore/Private/Memento/TypedElementMementoTranslators.cpp
@@ -7,6 +7,10 @@
 
 #include UE_INLINE_GENERATED_CPP_BY_NAME(TypedElementMementoTranslators)
 
+const UScriptStruct* UTedsMementoTranslatorBase::GetColumnToAddOnRestore() const
+{
+	return ColumnToAddOnRestore;
+}
 
 const UScriptStruct* UTedsDefaultMementoTranslator::GetMementoType() const
 {
diff --git a/Engine/Plugins/Experimental/EditorDataStorage/Source/TedsCore/Private/Memento/TypedElementMementoizedColumns .cpp b/Engine/Plugins/Experimental/EditorDataStorage/Source/TedsCore/Private/Memento/TypedElementMementoizedColumns .cpp
new file mode 100644
--- /dev/null
+++ b/Engine/Plugins/Experimental/EditorDataStorage/Source/TedsCore/Private/Memento/TypedElementMementoizedColumns .cpp	
@@ -0,0 +1,17 @@
+// Copyright Epic Games, Inc. All Rights Reserved.
+
+#include "TypedElementMementoizedColumns.h"
+
+#include "Elements/Columns/TypedElementVisibilityColumns.h"
+
+
+UTedsVisibilityColumnMementoTranslator::UTedsVisibilityColumnMementoTranslator()
+{
+	// To sync any visibility changes via re-instance back to the world we add the FVisibilityChangedTag to the row
+	ColumnToAddOnRestore = FVisibilityChangedTag::StaticStruct();
+}
+
+const UScriptStruct* UTedsVisibilityColumnMementoTranslator::GetColumnType() const
+{
+	return FVisibleInEditorColumn::StaticStruct();
+}
diff --git a/Engine/Plugins/Experimental/EditorDataStorage/Source/TedsCore/Private/Memento/TypedElementMementoizedColumns.h b/Engine/Plugins/Experimental/EditorDataStorage/Source/TedsCore/Private/Memento/TypedElementMementoizedColumns.h
--- a/Engine/Plugins/Experimental/EditorDataStorage/Source/TedsCore/Private/Memento/TypedElementMementoizedColumns.h
+++ b/Engine/Plugins/Experimental/EditorDataStorage/Source/TedsCore/Private/Memento/TypedElementMementoizedColumns.h
@@ -2,6 +2,7 @@
 
 #pragma once
 
+#include "Elements/Columns/TypedElementMiscColumns.h"
 #include "Elements/Columns/TypedElementSelectionColumns.h"
 #include "Memento/TypedElementMementoTranslators.h"
 #include "TypedElementMementoizedColumns.generated.h"
@@ -15,4 +16,18 @@ class UTedsSelectionColumnMementoTranslator final : public UTedsDefaultMementoTr
 	GENERATED_BODY()
 public:
 	virtual const UScriptStruct* GetColumnType() const override { return FTypedElementSelectionColumn::StaticStruct(); }
-};
\ No newline at end of file
+};
+
+/**
+ * Enable Visibility to be mementoized
+ */
+UCLASS()
+class UTedsVisibilityColumnMementoTranslator final : public UTedsDefaultMementoTranslator
+{
+	GENERATED_BODY()
+public:
+
+	UTedsVisibilityColumnMementoTranslator();
+	virtual const UScriptStruct* GetColumnType() const override;
+};
+
diff --git a/Engine/Plugins/Experimental/EditorDataStorage/Source/TedsCore/Private/Queries/TedsVisibilityQueries.cpp b/Engine/Plugins/Experimental/EditorDataStorage/Source/TedsCore/Private/Queries/TedsVisibilityQueries.cpp
new file mode 100644
--- /dev/null
+++ b/Engine/Plugins/Experimental/EditorDataStorage/Source/TedsCore/Private/Queries/TedsVisibilityQueries.cpp
@@ -0,0 +1,24 @@
+// Copyright Epic Games, Inc. All Rights Reserved.
+
+#include "TedsVisibilityQueries.h"
+
+#include "Elements/Columns/TypedElementVisibilityColumns.h"
+#include "Elements/Framework/TypedElementQueryBuilder.h"
+
+void UTedsVisibilityFactory::RegisterQueries(UE::Editor::DataStorage::ICoreProvider& DataStorage)
+{
+	using namespace UE::Editor::DataStorage::Queries;
+
+	DataStorage.RegisterQuery(
+		Select(
+			TEXT("Remove 'visibility change' tag"),
+			FProcessor(EQueryTickPhase::FrameEnd, DataStorage.GetQueryTickGroupName(EQueryTickGroups::SyncDataStorageToExternal)),
+			[](IQueryContext& Context, const RowHandle* Rows)
+			{
+				Context.RemoveColumns<FVisibilityChangedTag>(TConstArrayView<RowHandle>(Rows, Context.GetRowCount()));
+			}
+		)
+		.Where()
+			.All<FVisibilityChangedTag>()
+		.Compile());
+}
diff --git a/Engine/Plugins/Experimental/EditorDataStorage/Source/TedsCore/Private/Queries/TedsVisibilityQueries.h b/Engine/Plugins/Experimental/EditorDataStorage/Source/TedsCore/Private/Queries/TedsVisibilityQueries.h
new file mode 100644
--- /dev/null
+++ b/Engine/Plugins/Experimental/EditorDataStorage/Source/TedsCore/Private/Queries/TedsVisibilityQueries.h
@@ -0,0 +1,22 @@
+// Copyright Epic Games, Inc. All Rights Reserved.
+
+#pragma once
+
+#include "Elements/Interfaces/TypedElementDataStorageFactory.h"
+#include "UObject/ObjectMacros.h"
+
+#include "TedsVisibilityQueries.generated.h"
+
+/**
+ * Removes all FVisibilityChangeTag at the end of an update cycle.
+ */
+UCLASS()
+class UTedsVisibilityFactory : public UEditorDataStorageFactory
+{
+	GENERATED_BODY()
+
+public:
+	~UTedsVisibilityFactory() override = default;
+
+	void RegisterQueries(UE::Editor::DataStorage::ICoreProvider& DataStorage) override;
+};
diff --git a/Engine/Plugins/Experimental/EditorDataStorage/Source/TedsCore/Public/Memento/TypedElementMementoTranslators.h b/Engine/Plugins/Experimental/EditorDataStorage/Source/TedsCore/Public/Memento/TypedElementMementoTranslators.h
--- a/Engine/Plugins/Experimental/EditorDataStorage/Source/TedsCore/Public/Memento/TypedElementMementoTranslators.h
+++ b/Engine/Plugins/Experimental/EditorDataStorage/Source/TedsCore/Public/Memento/TypedElementMementoTranslators.h
@@ -40,6 +40,17 @@ public:
 	 * Implements the logic to convert a memento to a column
 	 */
 	virtual void TranslateMementoToColumn(const void* TypeErasedMemento, void* TypeErasedColumn) const PURE_VIRTUAL(UTedsMementoTranslatorBase, )
+	
+	/**
+	 * This column (if specified) will be added to the object row after the memento is restored
+	 */
+	const UScriptStruct* GetColumnToAddOnRestore() const;
+	
+protected:
+	
+	// Optional column that will be added to the object row after the memento is restored
+	UPROPERTY()
+	TObjectPtr<const UScriptStruct> ColumnToAddOnRestore = nullptr;
 };
 
 /**
diff --git a/Engine/Plugins/Experimental/EditorDataStorageFeatures/Source/TedsActorCompatibility/Private/Compatibility/TedsActorVisibilityQueries.cpp b/Engine/Plugins/Experimental/EditorDataStorageFeatures/Source/TedsActorCompatibility/Private/Compatibility/TedsActorVisibilityQueries.cpp
--- a/Engine/Plugins/Experimental/EditorDataStorageFeatures/Source/TedsActorCompatibility/Private/Compatibility/TedsActorVisibilityQueries.cpp
+++ b/Engine/Plugins/Experimental/EditorDataStorageFeatures/Source/TedsActorCompatibility/Private/Compatibility/TedsActorVisibilityQueries.cpp
@@ -95,7 +95,7 @@ void UActorVisibilityDataStorageFactory::RegisterVisibilityColumnToActorQuery(UE
 				}
 			})
 		.Where()
-			.All<FTypedElementActorTag, FTypedElementSyncBackToWorldTag>()
+			.All<FTypedElementActorTag, FVisibilityChangedTag>()
 		.Compile()
 	);
 }
diff --git a/Engine/Plugins/Experimental/EditorDataStorageFeatures/Source/TedsOutliner/Private/Widgets/OutlinerVisibilityWidget.cpp b/Engine/Plugins/Experimental/EditorDataStorageFeatures/Source/TedsOutliner/Private/Widgets/OutlinerVisibilityWidget.cpp
--- a/Engine/Plugins/Experimental/EditorDataStorageFeatures/Source/TedsOutliner/Private/Widgets/OutlinerVisibilityWidget.cpp
+++ b/Engine/Plugins/Experimental/EditorDataStorageFeatures/Source/TedsOutliner/Private/Widgets/OutlinerVisibilityWidget.cpp
@@ -329,7 +329,7 @@ void STedsVisibilityWidget::CommitVisibility(UE::Editor::DataStorage::ICoreProvi
 	{
 		VisibileInEditorColumn->bIsVisibleInEditor = bVisible;
 	}
-	DataStorage.AddColumn<FTypedElementSyncBackToWorldTag>(Row);
+	DataStorage.AddColumn<FVisibilityChangedTag>(Row);
 }
 
 void STedsVisibilityWidget::SetVisibility_Recursive(UE::Editor::DataStorage::ICoreProvider& DataStorage, RowHandle Row,
diff --git a/Engine/Source/Runtime/TypedElementFramework/Public/Elements/Columns/TypedElementVisibilityColumns.h b/Engine/Source/Runtime/TypedElementFramework/Public/Elements/Columns/TypedElementVisibilityColumns.h
--- a/Engine/Source/Runtime/TypedElementFramework/Public/Elements/Columns/TypedElementVisibilityColumns.h
+++ b/Engine/Source/Runtime/TypedElementFramework/Public/Elements/Columns/TypedElementVisibilityColumns.h
@@ -17,4 +17,15 @@ struct FVisibleInEditorColumn final : public FEditorDataStorageColumn
 public:
 	UPROPERTY()
 	bool bIsVisibleInEditor = true;
+};
+
+/**
+ * Tag added to rows for a frame when their visibility value changes. 
+ * This tag is required to be added for any changes you make to the TEDS column to be propagated back to the world, and is then removed at the end
+ * of the frame.
+ */
+USTRUCT(meta = (DisplayName = "Visibility Changed"))
+struct FVisibilityChangedTag final : public FEditorDataStorageTag
+{
+	GENERATED_BODY()
 };
\ No newline at end of file
