[TEDS] Fixed bug in BatchAddRemoveColumns.

`BatchAddRemoveColumns` used the original Mass implementation that required to keep fragments and tags separated. As an optimization the archetype list that goes into both calls is only calculated once, but if the same row gets a fragment and tag added that archetype list is invalid and needs to be rebuild. To avoid this, the fix is to update to the new element based alternative in Mass that can do both fragments and tags in a single call. There is however no single row equivalent for that yet so the command buffer had to be updated to support both single row commands and batch row commands. The side effect of this is that calling `BatchAddRemoveColumns` is now faster when called on queued rows.

#Synced-CL 50174701

diff --git a/Engine/Plugins/Experimental/EditorDataStorage/Source/TedsCore/Private/TypedElementDatabase.cpp b/Engine/Plugins/Experimental/EditorDataStorage/Source/TedsCore/Private/TypedElementDatabase.cpp
--- a/Engine/Plugins/Experimental/EditorDataStorage/Source/TedsCore/Private/TypedElementDatabase.cpp
+++ b/Engine/Plugins/Experimental/EditorDataStorage/Source/TedsCore/Private/TypedElementDatabase.cpp
@@ -113,6 +113,22 @@ namespace UE::Editor::DataStorage::Private
 		
 		bool MustUpdate() const { return bMustUpdateFragments || bMustUpdateTags; }
 	};
+
+	void ColumnsToBitSets(TConstArrayView<const UScriptStruct*> Columns, FMassElementBitSet& Elements)
+	{
+		for (const UScriptStruct* ColumnType : Columns)
+		{
+			Elements.Add(*ColumnType);
+		}
+	}
+
+	FMassElementBitSet ColumnsToBitSets(TConstArrayView<const UScriptStruct*> Columns)
+	{
+		FMassElementBitSet Result;
+		ColumnsToBitSets(Columns, Result);
+		return Result;
+	}
+
 	ColumnsToBitSetsResult ColumnsToBitSets(TConstArrayView<const UScriptStruct*> Columns, FMassFragmentBitSet& Fragments, FMassTagBitSet& Tags)
 	{
 		ColumnsToBitSetsResult Result;
@@ -850,51 +866,41 @@ void UEditorDataStorage::AddRemoveColumns(UE::Editor::DataStorage::RowHandle Row
 void UEditorDataStorage::BatchAddRemoveColumns(TConstArrayView<UE::Editor::DataStorage::RowHandle> Rows,
 	TConstArrayView<const UScriptStruct*> ColumnsToAdd, TConstArrayView<const UScriptStruct*> ColumnsToRemove)
 {	
-	if (ActiveEditorEntityManager)
-	{
-		FMassFragmentBitSet FragmentsToAdd;
-		FMassFragmentBitSet FragmentsToRemove;
-
-		FMassTagBitSet TagsToAdd;
-		FMassTagBitSet TagsToRemove;
-
-		using namespace UE::Editor::DataStorage;
-
-		Private::ColumnsToBitSetsResult AddResult = Private::ColumnsToBitSets(ColumnsToAdd, FragmentsToAdd, TagsToAdd);
-		Private::ColumnsToBitSetsResult RemoveResult = Private::ColumnsToBitSets(ColumnsToRemove, FragmentsToRemove, TagsToRemove);
-		
-		if (AddResult.MustUpdate() || RemoveResult.MustUpdate())
-		{
-			using EntityHandleArray = TArray<FMassEntityHandle, TInlineAllocator<32>>;
-			using EntityArchetypeLookup = TMap<FMassArchetypeHandle, EntityHandleArray, TInlineSetAllocator<32>>;
-			using ArchetypeEntityArray = TArray<FMassArchetypeEntityCollection, TInlineAllocator<32>>;
-
-			Legacy::FCommandBuffer& CommandBuffer = Environment->GetDirectDeferredCommands();
+	using namespace UE::Editor::DataStorage;
+	
+	using RowHandleArray = TArray<RowHandle, TInlineAllocator<32>>;
+	using EntityHandleArray = TArray<FMassEntityHandle, TInlineAllocator<32>>;
+	using EntityArchetypeLookup = TMap<FMassArchetypeHandle, EntityHandleArray, TInlineSetAllocator<32>>;
+	using ArchetypeEntityArray = TArray<FMassArchetypeEntityCollection, TInlineAllocator<32>>;
+
+	if (ActiveEditorEntityManager && !Rows.IsEmpty() && (!ColumnsToAdd.IsEmpty() || !ColumnsToRemove.IsEmpty()))
+	{	
+		FMassElementBitSet ElementsToAdd = Private::ColumnsToBitSets(ColumnsToAdd);
+		FMassElementBitSet ElementsToRemove = Private::ColumnsToBitSets(ColumnsToRemove);
 			
-			// Sort rows (entities) into to matching table (archetype) bucket.
-			EntityArchetypeLookup LookupTable;
-			for (RowHandle EntityId : Rows)
+		Legacy::FCommandBuffer& CommandBuffer = Environment->GetDirectDeferredCommands();
+			
+		// Sort rows (entities) into to matching table (archetype) bucket. If a row hasn't been constructed yet, but
+		// only reserved, then queue the command for later execution.
+		EntityArchetypeLookup LookupTable;
+		RowHandleArray DelayedRows;
+		for (RowHandle EntityId : Rows)
+		{
+			FMassEntityHandle Entity = FMassEntityHandle::FromNumber(EntityId);
+			if (ActiveEditorEntityManager->IsEntityActive(Entity))
 			{
-				FMassEntityHandle Entity = FMassEntityHandle::FromNumber(EntityId);
-				if (ActiveEditorEntityManager->IsEntityActive(Entity))
-				{
-					FMassArchetypeHandle Archetype = ActiveEditorEntityManager->GetArchetypeForEntity(Entity);
-					EntityHandleArray& EntityCollection = LookupTable.FindOrAdd(Archetype);
-					EntityCollection.Add(Entity);
-				}
-				else
-				{
-					if (AddResult.MustUpdate())
-					{
-						CommandBuffer.Queue_AddColumnsCommand(EntityId, FragmentsToAdd, TagsToAdd);
-					}
-					if (RemoveResult.MustUpdate())
-					{
-						CommandBuffer.Queue_RemoveColumnsCommand(EntityId, FragmentsToRemove, TagsToRemove);
-					}
-				}
+				FMassArchetypeHandle Archetype = ActiveEditorEntityManager->GetArchetypeForEntity(Entity);
+				EntityHandleArray& EntityCollection = LookupTable.FindOrAdd(Archetype);
+				EntityCollection.Add(Entity);
 			}
+			else
+			{
+				DelayedRows.Add(EntityId);
+			}
+		}
 		
+		if (!LookupTable.IsEmpty())
+		{
 			// Construct table (archetype) specific row (entity) collections.
 			ArchetypeEntityArray EntityCollections;
 			EntityCollections.Reserve(LookupTable.Num());
@@ -903,15 +909,12 @@ void UEditorDataStorage::BatchAddRemoveColumns(TConstArrayView<UE::Editor::DataS
 				EntityCollections.Emplace(It.Key(), It.Value(), FMassArchetypeEntityCollection::EDuplicatesHandling::FoldDuplicates);
 			}
 
-			// Batch update using the appropriate fragment/bit sets.
-			if (AddResult.bMustUpdateFragments || RemoveResult.bMustUpdateFragments)
-			{
-				ActiveEditorEntityManager->BatchChangeFragmentCompositionForEntities(EntityCollections, FragmentsToAdd, FragmentsToRemove);
-			}
-			if (AddResult.bMustUpdateTags || RemoveResult.bMustUpdateTags)
-			{
-				ActiveEditorEntityManager->BatchChangeTagsForEntities(EntityCollections, TagsToAdd, TagsToRemove);
-			}
+			ActiveEditorEntityManager->BatchChangeCompositionForEntities(EntityCollections, ElementsToAdd, ElementsToRemove);
+		}
+
+		if (!DelayedRows.IsEmpty())
+		{
+			CommandBuffer.Queue_AddRemoveColumnsCommand(DelayedRows, ElementsToAdd, ElementsToRemove);
 		}
 	}
 }
diff --git a/Engine/Plugins/Experimental/EditorDataStorage/Source/TedsCore/Private/TypedElementDatabaseCommandBuffer.cpp b/Engine/Plugins/Experimental/EditorDataStorage/Source/TedsCore/Private/TypedElementDatabaseCommandBuffer.cpp
--- a/Engine/Plugins/Experimental/EditorDataStorage/Source/TedsCore/Private/TypedElementDatabaseCommandBuffer.cpp
+++ b/Engine/Plugins/Experimental/EditorDataStorage/Source/TedsCore/Private/TypedElementDatabaseCommandBuffer.cpp
@@ -83,6 +83,48 @@ namespace UE::Editor::DataStorage::Legacy
 		}
 	}
 
+	template <typename T>
+	concept HasRow = requires(T t) { t.Row; } && 
+		std::is_same_v<std::remove_cvref_t<decltype(std::declval<T>().Row)>, RowHandle>;
+
+	template <typename T>
+	concept HasRowsView = requires(T t) { t.Rows; } && 
+		(
+			std::is_same_v<std::remove_cvref_t<decltype(std::declval<T>().Rows)>, TArrayView<RowHandle>> ||
+			std::is_same_v<std::remove_cvref_t<decltype(std::declval<T>().Rows)>, TArrayView<const RowHandle>>
+		);
+
+	struct FCheckForRowVisitor
+	{
+		RowHandle Row;
+		
+		template<typename T>
+		bool operator()(T& Command) requires (!HasRowsView<T> && HasRow<T>)
+		{
+			return Command.Row == Row;
+		}
+
+		template<typename T>
+		bool operator()(T& Command) requires (HasRowsView<T> && !HasRow<T>)
+		{
+			int32 Index;
+			if (Command.Rows.Find(Row, Index))
+			{
+				// Swap the index and the last, then reduce the size of the array by one to remove the row without
+				// resizing the array as that's a fixed allocation plus it's short lived so not worth resizing.
+				Command.Rows[Index] = Command.Rows.Last();
+				Command.Rows.LeftChopInline(1);
+			}
+			return Command.Rows.IsEmpty();
+		}
+
+		template<typename T>
+		bool operator()(T& Command)
+		{
+			static_assert(sizeof(T) == 0, "Command requires either `RowHandle Row` or `TArrayView<RowHandle> Rows`.");
+		}
+	};
+
 	void FCommandBuffer::Clear(RowHandle Row)
 	{
 		for (TMap<PendingColumnMappingKey, void*>::TIterator It = PendingColumns.CreateIterator(); It; ++It)
@@ -93,11 +135,13 @@ namespace UE::Editor::DataStorage::Legacy
 			}
 		}
 
-		for (TArray<FCommand>::TIterator It = Commands.CreateIterator(); It; ++It)
+		FCheckForRowVisitor Visitor;
+		Visitor.Row = Row;
+		for (TArray<CommandData>::TIterator It = Commands.CreateIterator(); It; ++It)
 		{
-			if (It->Row == Row)
+			if (Visit(Visitor, *It))
 			{
-				It.RemoveCurrent(); // Don't swap as the order of commands is important.
+				It.RemoveCurrent();
 			}
 		}
 	}
@@ -116,8 +160,9 @@ namespace UE::Editor::DataStorage::Legacy
 
 	void FCommandBuffer::Queue_AddColumnCommand(RowHandle Row, const UScriptStruct* ColumnType)
 	{
-		AddCommand(Row, FAddColumnCommand
+		AddCommand(FAddColumnCommand
 			{
+				.Row = Row,
 				.ColumnType = ColumnType
 			});
 		PendingColumns.FindOrAdd(PendingColumnMappingKey(Row, ColumnType), nullptr);
@@ -135,8 +180,9 @@ namespace UE::Editor::DataStorage::Legacy
 			// Initialize to zero to replicate the default from Mass.
 			void* Data = Environment.GetScratchBuffer().AllocateZeroInitialized(ColumnType->GetStructureSize(), ColumnType->GetMinAlignment());
 			PendingColumns.Add(PendingColumnMappingKey(Row, ColumnType), Data);
-			AddCommand(Row, FAddDataColumnCommand
+			AddCommand(FAddDataColumnCommand
 				{
+					.Row = Row,
 					.ColumnType = ColumnType,
 					.Relocator = Relocator,
 					.Data = Data
@@ -160,8 +206,9 @@ namespace UE::Editor::DataStorage::Legacy
 		FragmentsToAdd.ExportTypes(AddColumn);
 		TagsToAdd.ExportTypes(AddColumn);
 
-		AddCommand(Row, FAddColumnsCommand
+		AddCommand(FAddColumnsCommand
 			{
+				.Row = Row,
 				.FragmentsToAdd = MoveTemp(FragmentsToAdd),
 				.TagsToAdd = MoveTemp(TagsToAdd)
 			});
@@ -169,15 +216,15 @@ namespace UE::Editor::DataStorage::Legacy
 
 	void FCommandBuffer::Queue_RemoveColumnCommand(RowHandle Row, const UScriptStruct* ColumnType)
 	{
-		AddCommand(Row, FRemoveColumnCommand
+		AddCommand(FRemoveColumnCommand
 			{
+				.Row = Row,
 				.ColumnType = ColumnType
 			});
 		PendingColumns.Remove(PendingColumnMappingKey(Row, ColumnType));
 	}
 
-	void FCommandBuffer::Queue_RemoveColumnsCommand(RowHandle Row,
-		FMassFragmentBitSet FragmentsToRemove, FMassTagBitSet TagsToRemove)
+	void FCommandBuffer::Queue_RemoveColumnsCommand(RowHandle Row, FMassFragmentBitSet FragmentsToRemove, FMassTagBitSet TagsToRemove)
 	{
 		auto RemoveColumn = [this, Row](const UScriptStruct* ColumnType)
 		{
@@ -187,13 +234,45 @@ namespace UE::Editor::DataStorage::Legacy
 		FragmentsToRemove.ExportTypes(RemoveColumn);
 		TagsToRemove.ExportTypes(RemoveColumn);
 
-		AddCommand(Row, FRemoveColumnsCommand
+		AddCommand(FRemoveColumnsCommand
 			{
+				.Row = Row,
 				.FragmentsToRemove = MoveTemp(FragmentsToRemove),
 				.TagsToRemove = MoveTemp(TagsToRemove)
 			});
 	}
 
+	void FCommandBuffer::Queue_AddRemoveColumnsCommand(TConstArrayView<RowHandle> Rows, FMassElementBitSet ColumnsToAdd, FMassElementBitSet ColumnsToRemove)
+	{
+		for (RowHandle Row : Rows)
+		{
+			auto AddColumn = [this, Row](const UScriptStruct* ColumnType)
+				{
+					PendingColumns.FindOrAdd(PendingColumnMappingKey(Row, ColumnType), nullptr);
+					return true;
+				};
+			ColumnsToAdd.ExportTypes(AddColumn);
+
+			auto RemoveColumn = [this, Row](const UScriptStruct* ColumnType)
+				{
+					PendingColumns.Remove(PendingColumnMappingKey(Row, ColumnType));
+					return true;
+				};
+			ColumnsToRemove.ExportTypes(RemoveColumn);
+		}
+
+		int32 RowCount = Rows.Num();
+		TArrayView<RowHandle> LocalRows = Environment.GetScratchBuffer().AllocateUninitializedArray<RowHandle>(RowCount);
+		FMemory::Memcpy(LocalRows.GetData(), Rows.GetData(), RowCount * sizeof(RowHandle));
+
+		AddCommand(FAddRemoveColumnsCommand
+			{
+				.Rows = LocalRows,
+				.ElementsToAdd = MoveTemp(ColumnsToAdd),
+				.ElementsToRemove = MoveTemp(ColumnsToRemove)
+			});
+	}
+
 	//
 	// Execute section
 	//
@@ -213,7 +292,7 @@ namespace UE::Editor::DataStorage::Legacy
 		RowHandle Row,
 		const UScriptStruct* ColumnType)
 	{
-		if (ColumnType)
+		if (ColumnType && Execute_IsRowAssigned(MassEntityManager, Row))
 		{
 			FMassEntityHandle Entity = FMassEntityHandle::FromNumber(Row);
 			if (UE::Mass::IsA<FMassTag>(ColumnType))
@@ -236,7 +315,7 @@ namespace UE::Editor::DataStorage::Legacy
 		FMassEntityManager& MassEntityManager, RowHandle Row, const UScriptStruct* ColumnType,
 		void* Data, ColumnCopyOrMoveCallback Relocator)
 	{
-		if (ColumnType)
+		if (ColumnType && Execute_IsRowAssigned(MassEntityManager, Row))
 		{
 			checkf(UE::Mass::IsA<FMassFragment>(ColumnType),
 				TEXT("Trying to create a data column for '%s' from a deferred command that isn't a data column."), *ColumnType->GetName());
@@ -261,7 +340,7 @@ namespace UE::Editor::DataStorage::Legacy
 
 	void FCommandBuffer::Execute_AddSharedColumnCommand(FMassEntityManager& MassEntityManager, RowHandle Row, const FConstSharedStruct& SharedColumn)
 	{
-		if (SharedColumn.IsValid())
+		if (SharedColumn.IsValid() && Execute_IsRowAssigned(MassEntityManager, Row))
 		{
 			FMassEntityHandle Entity = FMassEntityHandle::FromNumber(Row);
 			MassEntityManager.AddConstSharedFragmentToEntity(Entity, SharedColumn);
@@ -270,7 +349,7 @@ namespace UE::Editor::DataStorage::Legacy
 
 	void FCommandBuffer::Execute_RemoveSharedColumnCommand(FMassEntityManager& MassEntityManager, RowHandle Row, const UScriptStruct& ColumnType)
 	{
-		if (ColumnType.IsChildOf(FTedsSharedColumn::StaticStruct()))
+		if (ColumnType.IsChildOf(FTedsSharedColumn::StaticStruct()) && Execute_IsRowAssigned(MassEntityManager, Row))
 		{
 			FMassEntityHandle Entity = FMassEntityHandle::FromNumber(Row);
 			MassEntityManager.RemoveConstSharedFragmentFromEntity(Entity, ColumnType);
@@ -283,9 +362,12 @@ namespace UE::Editor::DataStorage::Legacy
 		FMassFragmentBitSet FragmentsToAdd,
 		FMassTagBitSet TagsToAdd)
 	{
-		FMassArchetypeCompositionDescriptor AddComposition(
-			MoveTemp(FragmentsToAdd), MoveTemp(TagsToAdd), FMassChunkFragmentBitSet(), FMassSharedFragmentBitSet(), FMassConstSharedFragmentBitSet());
-		MassEntityManager.AddCompositionToEntity_GetDelta(FMassEntityHandle::FromNumber(Row), AddComposition);
+		if (Execute_IsRowAssigned(MassEntityManager, Row))
+		{
+			FMassArchetypeCompositionDescriptor AddComposition(
+				MoveTemp(FragmentsToAdd), MoveTemp(TagsToAdd), FMassChunkFragmentBitSet(), FMassSharedFragmentBitSet(), FMassConstSharedFragmentBitSet());
+			MassEntityManager.AddCompositionToEntity_GetDelta(FMassEntityHandle::FromNumber(Row), AddComposition);
+		}
 	}
 
 	void FCommandBuffer::Execute_RemoveColumnCommand(
@@ -293,7 +375,7 @@ namespace UE::Editor::DataStorage::Legacy
 		RowHandle Row,
 		const UScriptStruct* ColumnType)
 	{
-		if (ColumnType)
+		if (ColumnType && Execute_IsRowAssigned(MassEntityManager, Row))
 		{
 			FMassEntityHandle Entity = FMassEntityHandle::FromNumber(Row);
 			if (UE::Mass::IsA<FMassTag>(ColumnType))
@@ -313,9 +395,45 @@ namespace UE::Editor::DataStorage::Legacy
 		FMassFragmentBitSet FragmentsToRemove,
 		FMassTagBitSet TagsToRemove)
 	{
-		FMassArchetypeCompositionDescriptor RemoveComposition(
-			MoveTemp(FragmentsToRemove), MoveTemp(TagsToRemove), FMassChunkFragmentBitSet(), FMassSharedFragmentBitSet(), FMassConstSharedFragmentBitSet());
-		MassEntityManager.RemoveCompositionFromEntity(FMassEntityHandle::FromNumber(Row), RemoveComposition);
+		if (Execute_IsRowAssigned(MassEntityManager, Row))
+		{
+			FMassArchetypeCompositionDescriptor RemoveComposition(
+				MoveTemp(FragmentsToRemove), MoveTemp(TagsToRemove), FMassChunkFragmentBitSet(), FMassSharedFragmentBitSet(), FMassConstSharedFragmentBitSet());
+			MassEntityManager.RemoveCompositionFromEntity(FMassEntityHandle::FromNumber(Row), RemoveComposition);
+		}
+	}
+
+	void FCommandBuffer::Execute_AddRemoveColumnsCommand(FMassEntityManager& MassEntityManager, TArrayView<RowHandle> Rows,
+		FMassElementBitSet ElementsToAdd, FMassElementBitSet ElementsToRemove)
+	{
+		using EntityHandleArray = TArray<FMassEntityHandle, TInlineAllocator<32>>;
+		using EntityArchetypeLookup = TMap<FMassArchetypeHandle, EntityHandleArray, TInlineSetAllocator<32>>;
+		using ArchetypeEntityArray = TArray<FMassArchetypeEntityCollection, TInlineAllocator<32>>;
+
+		EntityArchetypeLookup LookupTable;
+		for (RowHandle EntityId : Rows)
+		{
+			FMassEntityHandle Entity = FMassEntityHandle::FromNumber(EntityId);
+			if (MassEntityManager.IsEntityActive(Entity))
+			{
+				FMassArchetypeHandle Archetype = MassEntityManager.GetArchetypeForEntity(Entity);
+				EntityHandleArray& EntityCollection = LookupTable.FindOrAdd(Archetype);
+				EntityCollection.Add(Entity);
+			}
+		}
+
+		if (!LookupTable.IsEmpty())
+		{
+			// Construct table (archetype) specific row (entity) collections.
+			ArchetypeEntityArray EntityCollections;
+			EntityCollections.Reserve(LookupTable.Num());
+			for (auto It = LookupTable.CreateConstIterator(); It; ++It)
+			{
+				EntityCollections.Emplace(It.Key(), It.Value(), FMassArchetypeEntityCollection::EDuplicatesHandling::FoldDuplicates);
+			}
+
+			MassEntityManager.BatchChangeCompositionForEntities(EntityCollections, ElementsToAdd, ElementsToRemove);
+		}
 	}
 
 	void FCommandBuffer::ProcessCommands()
@@ -323,23 +441,37 @@ namespace UE::Editor::DataStorage::Legacy
 		struct FProcessor
 		{
 			FMassEntityManager& EntityManager;
-			RowHandle Row;
-			void operator()(const FAddColumnCommand& Command) { Execute_AddColumnCommand(EntityManager, Row, Command.ColumnType.Get()); }
-			void operator()(const FAddDataColumnCommand& Command) { Execute_AddDataColumnCommand(EntityManager, Row, Command.ColumnType.Get(), Command.Data, Command.Relocator); }
-			void operator()(const FAddColumnsCommand& Command) { Execute_AddColumnsCommand(EntityManager, Row, Command.FragmentsToAdd, Command.TagsToAdd); }
-			void operator()(const FRemoveColumnCommand& Command) { Execute_RemoveColumnCommand(EntityManager, Row, Command.ColumnType.Get()); }
-			void operator()(const FRemoveColumnsCommand& Command) { Execute_RemoveColumnsCommand(EntityManager, Row, Command.FragmentsToRemove, Command.TagsToRemove); }
+			void operator()(FAddColumnCommand&& Command) 
+			{ 
+				Execute_AddColumnCommand(EntityManager, Command.Row, Command.ColumnType.Get());
+			}
+			void operator()(FAddDataColumnCommand&& Command) 
+			{ 
+				Execute_AddDataColumnCommand(EntityManager, Command.Row, Command.ColumnType.Get(), Command.Data, Command.Relocator);
+			}
+			void operator()(FAddColumnsCommand&& Command)
+			{ 
+				Execute_AddColumnsCommand(EntityManager, Command.Row, MoveTemp(Command.FragmentsToAdd), MoveTemp(Command.TagsToAdd));
+			}
+			void operator()(FRemoveColumnCommand&& Command) 
+			{ 
+				Execute_RemoveColumnCommand(EntityManager, Command.Row, Command.ColumnType.Get());
+			}
+			void operator()(FRemoveColumnsCommand&& Command) 
+			{ 
+				Execute_RemoveColumnsCommand(EntityManager, Command.Row, MoveTemp(Command.FragmentsToRemove), MoveTemp(Command.TagsToRemove));
+			}
+			void operator()(FAddRemoveColumnsCommand&& Command) 
+			{ 
+				Execute_AddRemoveColumnsCommand(EntityManager, Command.Rows, MoveTemp(Command.ElementsToAdd), MoveTemp(Command.ElementsToRemove));
+			}
 		};
 
 		FMassEntityManager& EntityManager = Environment.GetMassEntityManager();
 		FProcessor Processor{ .EntityManager = EntityManager };
-		for (FCommand& Command : Commands)
+		for (CommandData& Command : Commands)
 		{
-			if (Execute_IsRowAssigned(EntityManager, Command.Row))
-			{
-				Processor.Row = Command.Row;
-				Visit(Processor, Command.Data);
-			}
+			Visit(Processor, MoveTemp(Command));
 		}
 
 		PendingColumns.Reset();
@@ -357,11 +489,10 @@ namespace UE::Editor::DataStorage::Legacy
 	//
 
 	template<typename T>
-	void FCommandBuffer::AddCommand(RowHandle Row, T&& Args)
+	void FCommandBuffer::AddCommand(T&& Args)
 	{
-		FCommand Command;
-		Command.Row = Row;
-		Command.Data.Emplace<T>(Forward<T>(Args));
+		CommandData Command;
+		Command.Emplace<T>(Forward<T>(Args));
 		Commands.Add(MoveTemp(Command));
 	}
 } // namespace UE::Editor::DataStorage::Legacy
diff --git a/Engine/Plugins/Experimental/EditorDataStorage/Source/TedsCore/Private/TypedElementDatabaseCommandBuffer.h b/Engine/Plugins/Experimental/EditorDataStorage/Source/TedsCore/Private/TypedElementDatabaseCommandBuffer.h
--- a/Engine/Plugins/Experimental/EditorDataStorage/Source/TedsCore/Private/TypedElementDatabaseCommandBuffer.h
+++ b/Engine/Plugins/Experimental/EditorDataStorage/Source/TedsCore/Private/TypedElementDatabaseCommandBuffer.h
@@ -46,6 +46,8 @@ namespace UE::Editor::DataStorage
 	
 			void Queue_RemoveColumnCommand(RowHandle Row, const UScriptStruct* ColumnType);
 			void Queue_RemoveColumnsCommand(RowHandle Row, FMassFragmentBitSet FragmentsToRemove, FMassTagBitSet TagsToRemove);
+
+			void Queue_AddRemoveColumnsCommand(TConstArrayView<RowHandle> Rows, FMassElementBitSet ColumnsToAdd, FMassElementBitSet ColumnsToRemove);
 	
 			/**
 			 * @section Execute_* functions directly execute a command with limited validation checks.
@@ -71,6 +73,9 @@ namespace UE::Editor::DataStorage
 			static void Execute_RemoveColumnsCommand(
 				FMassEntityManager& MassEntityManager, RowHandle Row,
 				FMassFragmentBitSet FragmentsToRemove, FMassTagBitSet TagsToRemove);
+			static void Execute_AddRemoveColumnsCommand(
+				FMassEntityManager& MassEntityManager, TArrayView<RowHandle> Rows,
+				FMassElementBitSet ElementsToAdd, FMassElementBitSet ElementsToRemove);
 	
 			void ProcessCommands();
 			void ClearCommands();
@@ -78,10 +83,12 @@ namespace UE::Editor::DataStorage
 		private:
 			struct FAddColumnCommand
 			{
+				RowHandle Row;
 				TWeakObjectPtr<const UScriptStruct> ColumnType;
 			};
 			struct FAddDataColumnCommand
 			{
+				RowHandle Row;
 				TWeakObjectPtr<const UScriptStruct> ColumnType;
 				ColumnCopyOrMoveCallback Relocator;
 				void* Data;
@@ -90,37 +97,41 @@ namespace UE::Editor::DataStorage
 			};
 			struct FAddColumnsCommand
 			{
+				RowHandle Row;
 				FMassFragmentBitSet FragmentsToAdd;
 				FMassTagBitSet TagsToAdd;
 			};
 			struct FRemoveColumnCommand
 			{
+				RowHandle Row;
 				TWeakObjectPtr<const UScriptStruct> ColumnType;
 			};
 			struct FRemoveColumnsCommand
 			{
+				RowHandle Row;
 				FMassFragmentBitSet FragmentsToRemove;
 				FMassTagBitSet TagsToRemove;
 			};
+			struct FAddRemoveColumnsCommand
+			{
+				TArrayView<RowHandle> Rows;
+				FMassElementBitSet ElementsToAdd;
+				FMassElementBitSet ElementsToRemove;
+			};
 			using CommandData = TVariant<
 				FAddColumnCommand,
 				FAddDataColumnCommand,
 				FAddColumnsCommand,
 				FRemoveColumnCommand,
-				FRemoveColumnsCommand>;
-	
-			struct FCommand
-			{
-				RowHandle Row;
-				CommandData Data;
-			};
-	
+				FRemoveColumnsCommand,
+				FAddRemoveColumnsCommand>;
+
 			template<typename T>
-			void AddCommand(RowHandle Row, T&& Args);
+			void AddCommand(T&& Args);
 	
 			using PendingColumnMappingKey = TPair<RowHandle, TWeakObjectPtr<const UScriptStruct>>;
 			TMap<PendingColumnMappingKey, void*> PendingColumns;
-			TArray<FCommand> Commands;
+			TArray<CommandData> Commands;
 			FEnvironment& Environment;
 		};
 	} // namespace Legacy
