[TEDS] Reintroduced the `TEDS.Enable` CVar.

Users have asked for the ability to turn TEDS off in certain configuration that similate the game or server while using an editor build.

#Synced-CL 50468877

diff --git a/Engine/Plugins/Experimental/EditorDataStorage/Source/TedsCore/Private/TypedElementsDataStorage.cpp b/Engine/Plugins/Experimental/EditorDataStorage/Source/TedsCore/Private/TypedElementsDataStorage.cpp
--- a/Engine/Plugins/Experimental/EditorDataStorage/Source/TedsCore/Private/TypedElementsDataStorage.cpp
+++ b/Engine/Plugins/Experimental/EditorDataStorage/Source/TedsCore/Private/TypedElementsDataStorage.cpp
@@ -20,6 +20,17 @@
 
 #define LOCTEXT_NAMESPACE "FEditorDataStorageModule"
 
+namespace Private
+{
+	static bool GEnableTEDS = true;
+	static FAutoConsoleVariableRef CVarEnableTEDS(
+		TEXT("TEDS.Enable"),
+		GEnableTEDS,
+		TEXT("Enable the TEDS Editor Data Storage - Used as the central storage for editor data."),
+		ECVF_ReadOnly
+	);
+} // namespace Private
+
 // MASS uses CDO in a few places, making it a difficult to consistently register Type Element's Columns and Tags
 // as they may have not been set up to impersonate MASS' Fragments and Tags yet. There are currently no longer
 // any cases where TEDS relies on this but may happen again the future. For the standalone version a static
@@ -78,8 +89,9 @@ void FEditorDataStorageModule::StartupModule()
 
 	// TEDS relies on ticking Mass for processors to run and to pump command queues
 	// Commandlets do not run the same loop as the editor and does not set up tickables to execute each frame.
-	if (IsRunningCommandlet())
+	if (!Private::GEnableTEDS || IsRunningCommandlet())
 	{
+		UE_LOG(LogEditorDataStorage, Log, TEXT("TEDS has been disabled."));
 		return;
 	}
 
