[TEDS] Fixed RemoveColumns called from a query callback not being called on rows outside the active query context.

This is a targeted bugfix that only addresses `RemoveColumns` version for single rows and batches of rows. The delayed equivalents of those two are not updated.

#Synced-CL 50468877

diff --git a/Engine/Plugins/Experimental/EditorDataStorage/Source/TedsCore/Private/Processors/TypedElementProcessorAdaptors.cpp b/Engine/Plugins/Experimental/EditorDataStorage/Source/TedsCore/Private/Processors/TypedElementProcessorAdaptors.cpp
--- a/Engine/Plugins/Experimental/EditorDataStorage/Source/TedsCore/Private/Processors/TypedElementProcessorAdaptors.cpp
+++ b/Engine/Plugins/Experimental/EditorDataStorage/Source/TedsCore/Private/Processors/TypedElementProcessorAdaptors.cpp
@@ -393,8 +393,8 @@ namespace UE::Editor::DataStorage
 						for (FRemovedColumns* const RemovedColumns : RemovedColumnsQueue)
 						{
 							FRowOperations& Operations = RowOperations.FindOrAdd(RemovedColumns->Entity);
-							Operations.ColumnsChange.RemovedTags += RemovedColumns->RemoveDescriptor.Get<FMassTagBitSet>();
-							Operations.ColumnsChange.RemovedFragments += RemovedColumns->RemoveDescriptor.Get<FMassFragmentBitSet>();
+							Operations.ColumnsChange.RemovedTags += RemovedColumns->RemovedColumns.Get<FMassTagBitSet>();
+							Operations.ColumnsChange.RemovedFragments += RemovedColumns->RemovedColumns.Get<FMassFragmentBitSet>();
 						}
 						RemovedColumnsQueue.Empty();
 
@@ -403,8 +403,8 @@ namespace UE::Editor::DataStorage
 							for (const FMassEntityHandle& Row : RemovedColumns->Entities)
 							{
 								FRowOperations& Operations = RowOperations.FindOrAdd(Row);
-								Operations.ColumnsChange.RemovedTags += RemovedColumns->RemoveDescriptor.Get<FMassTagBitSet>();
-								Operations.ColumnsChange.RemovedFragments += RemovedColumns->RemoveDescriptor.Get<FMassFragmentBitSet>();
+								Operations.ColumnsChange.RemovedTags += RemovedColumns->RemovedColumns.Get<FMassTagBitSet>();
+								Operations.ColumnsChange.RemovedFragments += RemovedColumns->RemovedColumns.Get<FMassFragmentBitSet>();
 							}
 						}
 						RemovedColumnsBatchQueue.Empty();
@@ -803,114 +803,86 @@ namespace UE::Editor::DataStorage
 			void RemoveColumns(RowHandle Row, TConstArrayView<const UScriptStruct*> ColumnTypes) 
 			{
 				FRemovedColumns* RemovedColumns = Environment.GetScratchBuffer().template Emplace<FRemovedColumns>();
-				TedsColumnsToMassDescriptorIfActiveTable(RemovedColumns->RemoveDescriptor, ColumnTypes);
-				if (!RemovedColumns->RemoveDescriptor.IsEmpty())
-				{
-					RemovedColumns->Entity = FMassEntityHandle::FromNumber(Row);
-
-
-					if (bBatchDeferredCommands)
-					{
-						RemovedColumnsQueue.AddElement(RemovedColumns);
-					}
-					else
-					{ 
-						this->Context.Defer().template PushCommand<FMassDeferredRemoveCommand>(
-							[RemovedColumns](FMassEntityManager& System)
-							{
-								if (System.IsEntityValid(RemovedColumns->Entity) && System.IsEntityBuilt(RemovedColumns->Entity))
-								{
-									System.RemoveCompositionFromEntity(RemovedColumns->Entity, RemovedColumns->RemoveDescriptor);
-								}
-							});
-					}
-				}
-			}
-
-			void RemoveColumns(TConstArrayView<RowHandle> Rows, TConstArrayView<const UScriptStruct*> ColumnTypes) 
-			{
-				FScratchBuffer& ScratchBuffer = Environment.GetScratchBuffer();
-				FRemovedColumnsBatch* RemovedColumns = ScratchBuffer.Emplace<FRemovedColumnsBatch>();
-				TedsColumnsToMassDescriptorIfActiveTable(RemovedColumns->RemoveDescriptor, ColumnTypes);
-
-				RemovedColumns->Entities = ScratchBuffer.EmplaceArray<FMassEntityHandle>(Rows.Num());
-				FMassEntityHandle* Entities = RemovedColumns->Entities.GetData();
-				for (RowHandle Row : Rows)
-				{
-					*Entities = FMassEntityHandle::FromNumber(Row);
-					Entities++;
-				}
+				RemovedColumns->RemovedColumns = FMassElementBitSet(ColumnTypes);
+				RemovedColumns->Entity = FMassEntityHandle::FromNumber(Row);
 
 				if (bBatchDeferredCommands)
 				{
-					RemovedColumnsBatchQueue.AddElement(RemovedColumns);
+					RemovedColumnsQueue.AddElement(RemovedColumns);
 				}
 				else
 				{
 					this->Context.Defer().template PushCommand<FMassDeferredRemoveCommand>(
 						[RemovedColumns](FMassEntityManager& System)
 						{
-							TArrayView<FMassEntityHandle> Entities = RemovedColumns->Entities;
-							
-							using EntityHandleArray = TArray<FMassEntityHandle, TInlineAllocator<32>>;
-							using EntityArchetypeLookup = TMap<FMassArchetypeHandle, EntityHandleArray, TInlineSetAllocator<32>>;
-							using ArchetypeEntityArray = TArray<FMassArchetypeEntityCollection, TInlineAllocator<32>>;
+							if (FMassArchetypeHandle Archetype = System.GetArchetypeForEntity(RemovedColumns->Entity); Archetype.IsValid())
+							{
+								FMassArchetypeEntityCollection Collection(Archetype, RemovedColumns->Entity);
+								System.BatchChangeCompositionForEntities({ Collection }, {}, RemovedColumns->RemovedColumns);
+							}
+						});
+				}
+			}
 
-							// Sort rows (entities) into to matching table (archetype) bucket.
-							EntityArchetypeLookup LookupTable;
-							ArchetypeEntityArray EntityCollections;
+			void RemoveColumns(TConstArrayView<RowHandle> Rows, TConstArrayView<const UScriptStruct*> ColumnTypes) 
+			{
+				switch (Rows.Num())
+				{
+				case 0:
+					return;
+				case 1:
+					RemoveColumns(*Rows.GetData(), ColumnTypes);
+					return;
+				default:
+				{
+					FScratchBuffer& ScratchBuffer = Environment.GetScratchBuffer();
+					FRemovedColumnsBatch* RemovedColumns = ScratchBuffer.Emplace<FRemovedColumnsBatch>();
 
+					RemovedColumns->RemovedColumns = FMassElementBitSet(ColumnTypes);
+					RemovedColumns->Entities = ScratchBuffer.EmplaceArray<FMassEntityHandle>(Rows.Num());
+					FMassEntityHandle* Entities = RemovedColumns->Entities.GetData();
+					for (RowHandle Row : Rows)
+					{
+						*Entities = FMassEntityHandle::FromNumber(Row);
+						Entities++;
+					}
 
-							// This could be improved by adding an operation that would both combine the Fragments and Tags change in one bath operation.
-							if (!RemovedColumns->RemoveDescriptor.Get<FMassFragmentBitSet>().IsEmpty())
+					if (bBatchDeferredCommands)
+					{
+						RemovedColumnsBatchQueue.AddElement(RemovedColumns);
+					}
+					else
+					{
+						this->Context.Defer().template PushCommand<FMassDeferredRemoveCommand>(
+							[RemovedColumns](FMassEntityManager& System)
 							{
-								for (FMassEntityHandle Entity : Entities)
-								{
-									if (System.IsEntityValid(Entity))
-									{
-										FMassArchetypeHandle Archetype = System.GetArchetypeForEntity(Entity);
-										EntityHandleArray& EntityCollection = LookupTable.FindOrAdd(Archetype);
-										EntityCollection.Add(Entity);
-									}
-								}
-
-								// Construct table (archetype) specific row (entity) collections.
-								EntityCollections.Reserve(LookupTable.Num());
-								for (auto It = LookupTable.CreateConstIterator(); It; ++It)
-								{
-									// Could be more effective but the previous implementation was robust when called with duplicate rows.
-									EntityCollections.Emplace(It.Key(), It.Value(), FMassArchetypeEntityCollection::EDuplicatesHandling::FoldDuplicates);
-								}
+								using EntityHandleArray = TArray<FMassEntityHandle, TInlineAllocator<32>>;
+								using EntityArchetypeLookup = TMap<FMassArchetypeHandle, EntityHandleArray, TInlineSetAllocator<32>>;
+								using EntityCollectionArray = TArray<FMassArchetypeEntityCollection, TInlineAllocator<32>>;
 
-								System.BatchChangeFragmentCompositionForEntities(EntityCollections, FMassFragmentBitSet(), RemovedColumns->RemoveDescriptor.Get<FMassFragmentBitSet>());
-							}
-							if (!RemovedColumns->RemoveDescriptor.Get<FMassTagBitSet>().IsEmpty())
-							{
-								LookupTable.Reset();
-								EntityCollections.Reset();
+								TArrayView<FMassEntityHandle> Entities = RemovedColumns->Entities;
+								EntityArchetypeLookup LookupTable;
 
 								for (FMassEntityHandle Entity : Entities)
 								{
-									if (System.IsEntityValid(Entity))
+									if (FMassArchetypeHandle Archetype = System.GetArchetypeForEntity(Entity); Archetype.IsValid())
 									{
-										FMassArchetypeHandle Archetype = System.GetArchetypeForEntity(Entity);
-										EntityHandleArray& EntityCollection = LookupTable.FindOrAdd(Archetype);
-										EntityCollection.Add(Entity);
+										LookupTable.FindOrAdd(Archetype).Add(Entity);
 									}
 								}
 
-								// Construct table (archetype) specific row (entity) collections.
+								EntityCollectionArray EntityCollections;
 								EntityCollections.Reserve(LookupTable.Num());
-								for (auto It = LookupTable.CreateConstIterator(); It; ++It)
+								for (EntityArchetypeLookup::TIterator It = LookupTable.CreateIterator(); It; ++It)
 								{
-									// Could be more effective but the previous implementation was robust when called with duplicate rows.
-									EntityCollections.Emplace(It.Key(), It.Value(), FMassArchetypeEntityCollection::EDuplicatesHandling::FoldDuplicates);
+									EntityCollections.Emplace(It.Key(), TConstArrayView<FMassEntityHandle>(It.Value()),
+										FMassArchetypeEntityCollection::EDuplicatesHandling::FoldDuplicates);
 								}
 
-								System.BatchChangeTagsForEntities(EntityCollections, FMassTagBitSet(), RemovedColumns->RemoveDescriptor.Get<FMassTagBitSet>());
-							}
-						});
+								System.BatchChangeCompositionForEntities(EntityCollections, {}, RemovedColumns->RemovedColumns);
+							});
 					}
+				}}
 			}
 	
 			RowHandle AddRow(TableHandle Table) 
@@ -1075,13 +1047,13 @@ namespace UE::Editor::DataStorage
 
 			struct FRemovedColumns
 			{
-				FMassArchetypeCompositionDescriptor RemoveDescriptor;
+				FMassElementBitSet RemovedColumns;
 				FMassEntityHandle Entity;
 			};
 
 			struct FRemovedColumnsBatch
 			{
-				FMassArchetypeCompositionDescriptor RemoveDescriptor;
+				FMassElementBitSet RemovedColumns;
 				TArrayView<FMassEntityHandle> Entities;
 			};
 
