[TEDS] Fixed commands getting stuck in the command buffer.

Mass and TEDS both have their own command buffers for different reasons. The one in Mass didn't get properly cleared when direct calls were made that trigger a table change (e.g. add/remove column) and in turn trigger an observer that queues commands. This also applied to the TEDS specific command buffer, but this was also not being cleared between phases.
This change forces Mass and TEDS to flush their command buffers after each table mutating call and in between phases.

#Synced-CL 50468877

diff --git a/Engine/Plugins/Experimental/EditorDataStorage/Source/TedsCore/Private/Processors/TypedElementProcessorAdaptors.cpp b/Engine/Plugins/Experimental/EditorDataStorage/Source/TedsCore/Private/Processors/TypedElementProcessorAdaptors.cpp
--- a/Engine/Plugins/Experimental/EditorDataStorage/Source/TedsCore/Private/Processors/TypedElementProcessorAdaptors.cpp
+++ b/Engine/Plugins/Experimental/EditorDataStorage/Source/TedsCore/Private/Processors/TypedElementProcessorAdaptors.cpp
@@ -1578,7 +1578,7 @@ UE::Editor::DataStorage::FQueryResult FTypedElementQueryProcessorData::Execute(
 	{
 		FMassExecutionContext Context(EntityManager);
 		Context.SetDeferredCommandBuffer(Environment.GetSharedMassCommandBuffer());
-		Context.SetFlushDeferredCommands(true);
+		Context.SetFlushDeferredCommands(false);
 		auto ExecuteFunction = [&Result, &Callback, &Description, &Environment](FMassExecutionContext& Context)
 			{
 				using namespace Processors::Private;
@@ -1599,9 +1599,6 @@ UE::Editor::DataStorage::FQueryResult FTypedElementQueryProcessorData::Execute(
 		{
 			NativeQuery.ForEachEntityChunk(Context, ExecuteFunction);
 		}
-		Context.FlushDeferred();
-		checkf(Environment.GetSharedMassCommandBuffer()->HasPendingCommands() == false,
-			TEXT("Shared Mass command buffer used by the editor data storage wasn't cleared after flushing."));
 	}
 	else
 	{
@@ -1627,11 +1624,8 @@ UE::Editor::DataStorage::FQueryResult FTypedElementQueryProcessorData::Execute(
 	if (EnumHasAnyFlags(ExecutionFlags, ERunQueryFlags::AllowBoundQueries) || !Description.Callback.Function)
 	{
 		FMassExecutionContext Context(EntityManager);
-		if (EnumHasAnyFlags(ExecutionFlags, ERunQueryFlags::ForceCommandBufferFlush))
-		{
-			Context.SetDeferredCommandBuffer(Environment.GetSharedMassCommandBuffer());
-			Context.SetFlushDeferredCommands(true);
-		}
+		Context.SetDeferredCommandBuffer(Environment.GetSharedMassCommandBuffer());
+		Context.SetFlushDeferredCommands(false);
 		auto ExecuteFunction = [&Result, &Callback, &Environment](FMassExecutionContext& Context)
 			{
 				using namespace Processors::Private;
@@ -1653,13 +1647,6 @@ UE::Editor::DataStorage::FQueryResult FTypedElementQueryProcessorData::Execute(
 		{
 			NativeQuery.ForEachEntityChunk(Context, ExecuteFunction);
 		}
-		
-		if (EnumHasAnyFlags(ExecutionFlags, ERunQueryFlags::ForceCommandBufferFlush))
-		{
-			Context.FlushDeferred();
-			checkf(Environment.GetSharedMassCommandBuffer()->HasPendingCommands() == false,
-				TEXT("Shared Mass command buffer used by the editor data storage wasn't cleared after flushing."));
-		}
 	}
 	else
 	{
diff --git a/Engine/Plugins/Experimental/EditorDataStorage/Source/TedsCore/Private/TypedElementDatabase.cpp b/Engine/Plugins/Experimental/EditorDataStorage/Source/TedsCore/Private/TypedElementDatabase.cpp
--- a/Engine/Plugins/Experimental/EditorDataStorage/Source/TedsCore/Private/TypedElementDatabase.cpp
+++ b/Engine/Plugins/Experimental/EditorDataStorage/Source/TedsCore/Private/TypedElementDatabase.cpp
@@ -156,6 +156,21 @@ namespace UE::Editor::DataStorage::Private
 	}
 } // namespace UE::Editor::DataStorage::Private
 
+UEditorDataStorage::FActiveCallScope::FActiveCallScope(UEditorDataStorage& Owner)
+	: Owner(Owner)
+{
+	Owner.ActiveCallsCount++;
+}
+
+UEditorDataStorage::FActiveCallScope::~FActiveCallScope()
+{
+	Owner.ActiveCallsCount--;
+	if (Owner.ActiveCallsCount == 0)
+	{
+		Owner.Environment->FlushCommands();
+	}
+}
+
 void UEditorDataStorage::Initialize()
 {
 	using namespace UE::Editor::DataStorage;
@@ -304,6 +319,7 @@ void UEditorDataStorage::OnPreMassTick(float DeltaTime)
 {
 	checkf(IsAvailable(), TEXT("Typed Element Database was ticked while it's not ready."));
 	
+	ActiveCallsCount++;
 	OnUpdateDelegate.Broadcast();
 	// Process pending commands after other systems have had a chance to update. Other systems may have executed work needed
 	// to complete pending work.
@@ -315,6 +331,8 @@ void UEditorDataStorage::OnPostMassTick(float DeltaTime)
 	checkf(IsAvailable(), TEXT("Typed Element Database was ticked while it's not ready."));
 	
 	Environment->NextUpdateCycle();
+	ActiveCallsCount--;
+
 	OnUpdateCompletedDelegate.Broadcast();
 
 	TickCooperativeTasks();
@@ -415,6 +433,8 @@ UE::Editor::DataStorage::RowHandle UEditorDataStorage::AddRow(UE::Editor::DataSt
 {
 	using namespace UE::Editor::DataStorage;
 
+	FActiveCallScope Scope(*this);
+	
 	const int32 TableHandleAsIndex = Private::ConvertTableHandleToIndex(Table);
 	checkf(TableHandleAsIndex < Tables.Num(), TEXT("Attempting to add a row to a non-existing table."));
 	return ActiveEditorEntityManager 
@@ -429,6 +449,8 @@ UE::Editor::DataStorage::RowHandle UEditorDataStorage::AddRow(UE::Editor::DataSt
 
 	if (ActiveEditorEntityManager)
 	{
+		FActiveCallScope Scope(*this);
+		
 		const int32 TableHandleAsIndex = Private::ConvertTableHandleToIndex(Table);
 		checkf(TableHandleAsIndex < Tables.Num(), TEXT("Attempting to a row to a non-existing table."));
 
@@ -454,6 +476,8 @@ bool UEditorDataStorage::AddRow(UE::Editor::DataStorage::RowHandle ReservedRow,
 	checkf(TableHandleAsIndex < Tables.Num(), TEXT("Attempting to add a row to a non-existing table."));
 	if (ActiveEditorEntityManager)
 	{
+		FActiveCallScope Scope(*this);
+		
 		ActiveEditorEntityManager->BuildEntity(FMassEntityHandle::FromNumber(ReservedRow), Tables[TableHandleAsIndex]);
 		return true;
 	}
@@ -470,6 +494,8 @@ bool UEditorDataStorage::AddRow(UE::Editor::DataStorage::RowHandle ReservedRow,
 
 	if (ActiveEditorEntityManager)
 	{
+		FActiveCallScope Scope(*this);
+		
 		const int32 TableHandleAsIndex = Private::ConvertTableHandleToIndex(Table);
 		checkf(TableHandleAsIndex < Tables.Num(), TEXT("Attempting to add a row to a non-existing table."));
 		
@@ -489,6 +515,8 @@ bool UEditorDataStorage::BatchAddRow(
 
 	if (ActiveEditorEntityManager)
 	{
+		FActiveCallScope Scope(*this);
+		
 		const int32 TableHandleAsIndex = Private::ConvertTableHandleToIndex(Table);
 		checkf(TableHandleAsIndex < Tables.Num(), TEXT("Attempting to add multiple rows to a non-existing table."));
 	
@@ -515,6 +543,8 @@ bool UEditorDataStorage::BatchAddRow(UE::Editor::DataStorage::TableHandle Table,
 
 	if (ActiveEditorEntityManager)
 	{
+		FActiveCallScope Scope(*this);
+		
 		const int32 TableHandleAsIndex = Private::ConvertTableHandleToIndex(Table);
 		checkf(TableHandleAsIndex < Tables.Num(), TEXT("Attempting to add multiple rows to a non-existing table."));
 	
@@ -539,6 +569,7 @@ void UEditorDataStorage::RemoveRow(UE::Editor::DataStorage::RowHandle Row)
 	{
 		if (ActiveEditorEntityManager->IsEntityBuilt(FMassEntityHandle::FromNumber(Row)))
 		{
+			FActiveCallScope Scope(*this);
 			ActiveEditorEntityManager->DestroyEntity(FMassEntityHandle::FromNumber(Row));
 		}
 		else
@@ -556,6 +587,7 @@ void UEditorDataStorage::BatchRemoveRows(TConstArrayView<UE::Editor::DataStorage
 
 	if (ActiveEditorEntityManager)
 	{
+		FActiveCallScope Scope(*this);
 		ActiveEditorEntityManager->BatchDestroyEntities(RowsToMassEntitiesConversion(Rows));
 		Environment->GetMappingTable().MarkDirty();
 	}
@@ -587,6 +619,8 @@ void UEditorDataStorage::RemoveAllRowsWithColumns(TConstArrayView<const UScriptS
 		
 		if (!MatchingArchetypes.IsEmpty())
 		{
+			FActiveCallScope Scope(*this);
+			
 			TArray<FMassArchetypeEntityCollection> Collections;
 			Collections.Reserve(MatchingArchetypes.Num());
 			for (FMassArchetypeHandle& Archetype : MatchingArchetypes)
@@ -668,6 +702,7 @@ void UEditorDataStorage::AddColumn(UE::Editor::DataStorage::RowHandle Row, const
 	{
 		if (IsRowAssigned(Row))
 		{
+			FActiveCallScope Scope(*this);
 			UE::Editor::DataStorage::Legacy::FCommandBuffer::Execute_AddColumnCommand(*ActiveEditorEntityManager, Row, ColumnType);
 		}
 		else
@@ -689,6 +724,7 @@ void UEditorDataStorage::AddColumnData(UE::Editor::DataStorage::RowHandle Row, c
 			FStructView Column = ActiveEditorEntityManager->GetFragmentDataStruct(Entity, ColumnType);
 			if (!Column.IsValid())
 			{
+				FActiveCallScope Scope(*this);
 				ActiveEditorEntityManager->AddFragmentToEntity(Entity, ColumnType, Initializer);
 			}
 			else
@@ -710,6 +746,7 @@ void UEditorDataStorage::RemoveColumn(UE::Editor::DataStorage::RowHandle Row, co
 	{
 		if (IsRowAssigned(Row))
 		{
+			FActiveCallScope Scope(*this);
 			UE::Editor::DataStorage::Legacy::FCommandBuffer::Execute_RemoveColumnCommand(*ActiveEditorEntityManager, Row, ColumnType);
 		}
 		else
@@ -778,6 +815,7 @@ void UEditorDataStorage::AddColumn(UE::Editor::DataStorage::RowHandle Row, const
 		FMassEntityHandle Entity = FMassEntityHandle::FromNumber(Row);
 		if (ActiveEditorEntityManager->IsEntityActive(Entity))
 		{
+			FActiveCallScope Scope(*this);
 			UE::Editor::DataStorage::Legacy::FCommandBuffer::Execute_AddSharedColumnCommand(*ActiveEditorEntityManager, Row, SharedStruct);
 		}
 	}
@@ -791,6 +829,7 @@ void UEditorDataStorage::RemoveColumn(UE::Editor::DataStorage::RowHandle Row, co
 		FMassEntityHandle Entity = FMassEntityHandle::FromNumber(Row);
 		if (ActiveEditorEntityManager->IsEntityActive(Entity))
 		{
+			FActiveCallScope Scope(*this);
 			UE::Editor::DataStorage::Legacy::FCommandBuffer::Execute_RemoveSharedColumnCommand(*ActiveEditorEntityManager, Row, *ValueTagType);
 		}
 	}
@@ -811,6 +850,7 @@ void UEditorDataStorage::RemoveColumns(UE::Editor::DataStorage::RowHandle Row, T
 		{
 			if (ActiveEditorEntityManager->IsEntityActive(Entity))
 			{
+				FActiveCallScope Scope(*this);
 				Legacy::FCommandBuffer::Execute_RemoveColumnsCommand(*ActiveEditorEntityManager, Row, FragmentsToRemove, TagsToRemove);
 			}
 			else
@@ -840,6 +880,7 @@ void UEditorDataStorage::AddRemoveColumns(UE::Editor::DataStorage::RowHandle Row
 		
 		if (ActiveEditorEntityManager->IsEntityActive(Entity))
 		{
+			FActiveCallScope Scope(*this);
 			if (bMustAddColumns)
 			{
 				Legacy::FCommandBuffer::Execute_AddColumnsCommand(*ActiveEditorEntityManager, Row, FragmentsToAdd, TagsToAdd);
@@ -901,6 +942,8 @@ void UEditorDataStorage::BatchAddRemoveColumns(TConstArrayView<UE::Editor::DataS
 		
 		if (!LookupTable.IsEmpty())
 		{
+			FActiveCallScope Scope(*this);
+
 			// Construct table (archetype) specific row (entity) collections.
 			ArchetypeEntityArray EntityCollections;
 			EntityCollections.Reserve(LookupTable.Num());
@@ -1281,6 +1324,7 @@ UE::Editor::DataStorage::FQueryResult UEditorDataStorage::RunQuery(UE::Editor::D
 
 	if (ActiveEditorEntityManager)
 	{
+		// There's no need to scope this because this version of the query isn't running anything that could change tables.
 		const FExtendedQueryStore::Handle StorageHandle(Query);
 		return Environment->GetQueryStore().RunQuery(*ActiveEditorEntityManager, StorageHandle);
 	}
@@ -1299,6 +1343,7 @@ UE::Editor::DataStorage::FQueryResult UEditorDataStorage::RunQuery(
 
 	if (ActiveEditorEntityManager)
 	{
+		FActiveCallScope Scope(*this);
 		const FExtendedQueryStore::Handle StorageHandle(Query);
 		return Environment->GetQueryStore().RunQuery(*ActiveEditorEntityManager, *Environment, StorageHandle, 
 			EDirectQueryExecutionFlags::Default, Callback);
@@ -1319,6 +1364,7 @@ UE::Editor::DataStorage::FQueryResult UEditorDataStorage::RunQuery(
 
 	if (ActiveEditorEntityManager)
 	{
+		FActiveCallScope Scope(*this);
 		const FExtendedQueryStore::Handle StorageHandle(Query);
 		return Environment->GetQueryStore().RunQuery(*ActiveEditorEntityManager, *Environment, StorageHandle, InFlags, Callback);
 	}
@@ -1339,8 +1385,9 @@ UE::Editor::DataStorage::FQueryResult UEditorDataStorage::RunQuery(
 
 	if (ActiveEditorEntityManager)
 	{
-		return Environment->GetQueryStore().RunQuery(*ActiveEditorEntityManager, *Environment, 
-			FExtendedQueryStore::Handle(Query), InFlags | ERunQueryFlags::ForceCommandBufferFlush, Callback);
+		FActiveCallScope Scope(*this); 
+		return Environment->GetQueryStore().RunQuery(*ActiveEditorEntityManager, *Environment,
+			FExtendedQueryStore::Handle(Query), InFlags, Callback);
 	}
 	else
 	{
@@ -1645,6 +1692,12 @@ void UEditorDataStorage::FinalizePhase(UE::Editor::DataStorage::EQueryTickPhase
 		// into a step that runs the queries and uses a shared lock and one that executes the command buffer with an exclusive lock.
 		FScopedExclusiveLock Lock(EGlobalLockScope::Public);
 		Environment->GetQueryStore().RunPhasePostambleQueries(*ActiveEditorEntityManager, *Environment, Phase, DeltaTime);
+
+		// Flush any custom commands that have queued up in between phases.
+		ActiveCallsCount--;
+		checkf(ActiveCallsCount == 0, TEXT("During phase finalization there should be no TEDS running."));
+		Environment->FlushCommands();
+		ActiveCallsCount++;
 	}
 }
 
diff --git a/Engine/Plugins/Experimental/EditorDataStorage/Source/TedsCore/Private/TypedElementDatabase.h b/Engine/Plugins/Experimental/EditorDataStorage/Source/TedsCore/Private/TypedElementDatabase.h
--- a/Engine/Plugins/Experimental/EditorDataStorage/Source/TedsCore/Private/TypedElementDatabase.h
+++ b/Engine/Plugins/Experimental/EditorDataStorage/Source/TedsCore/Private/TypedElementDatabase.h
@@ -236,6 +236,20 @@ private:
 		
 		TObjectPtr<UEditorDataStorageFactory> Instance;
 	};
+
+	// Wrapper around ActiveCallsCount to avoid issues with timing such as early outs.
+	struct FActiveCallScope
+	{
+		explicit FActiveCallScope(UEditorDataStorage& Owner);
+		~FActiveCallScope();
+
+		FActiveCallScope(const FActiveCallScope&) = delete;
+		FActiveCallScope(FActiveCallScope&&) = delete;
+		FActiveCallScope& operator=(const FActiveCallScope&) = delete;
+		FActiveCallScope& operator=(FActiveCallScope&&) = delete;
+
+		UEditorDataStorage& Owner;
+	};
 	
 	static UE_API const FName TickGroupName_Default;
 	static UE_API const FName TickGroupName_PreUpdate;
@@ -277,6 +291,10 @@ private:
 
 	TSharedPtr<FMassEntityManager> ActiveEditorEntityManager;
 	TSharedPtr<FMassProcessingPhaseManager> ActiveEditorPhaseManager;
+
+	// Track the number of active calls, incrementing whenever a function calls another direct API call. This prevents
+	// the command buffer from being flushed early and guarantees it only being flushed when the call chain has completed.
+	uint32 ActiveCallsCount = 0;
 };
 
 template <typename FactoryType, typename DatabaseType>
diff --git a/Engine/Plugins/Experimental/EditorDataStorage/Source/TedsCore/Private/TypedElementDatabaseEnvironment.cpp b/Engine/Plugins/Experimental/EditorDataStorage/Source/TedsCore/Private/TypedElementDatabaseEnvironment.cpp
--- a/Engine/Plugins/Experimental/EditorDataStorage/Source/TedsCore/Private/TypedElementDatabaseEnvironment.cpp
+++ b/Engine/Plugins/Experimental/EditorDataStorage/Source/TedsCore/Private/TypedElementDatabaseEnvironment.cpp
@@ -15,6 +15,7 @@ namespace UE::Editor::DataStorage
 		, Queries(DynamicColumnGenerator)
 		, MementoSystem(InDataStorage)
 		, ValueTagManager(DynamicColumnGenerator)
+		, ActiveCommandQueue(CommandQueues)
 		, MassEntityManager(InMassEntityManager)
 		, MassPhaseManager(InMassPhaseManager)
 		, SharedMassCommandBuffer(MakeShared<FMassCommandBuffer>())
@@ -170,16 +171,38 @@ namespace UE::Editor::DataStorage
 	void FEnvironment::PushCommands(TConstArrayView<const FEnvironmentCommand> Commands)
 	{
 		TUniqueLock<FMutex> Lock(CommandQueueMutex);
-
-		CommandQueue.Append(Commands);
+		ActiveCommandQueue->Append(Commands);
 	}
 
 	void FEnvironment::FlushCommands()
 	{
-		for (FEnvironmentCommand& Command : CommandQueue)
+		while (true)
 		{
-			Command.CommandFunction(Command.CommandData);
+			TSharedPtr<FMassCommandBuffer> MassCommandBuffer = GetSharedMassCommandBuffer();
+			bool bFlushMassCommandBuffer = MassCommandBuffer.IsValid() && MassCommandBuffer->HasPendingCommands();
+			
+			TArray<FEnvironmentCommand>* ProcessingQueue;
+			{
+				TUniqueLock<FMutex> Lock(CommandQueueMutex);
+				if (!bFlushMassCommandBuffer && ActiveCommandQueue->IsEmpty())
+				{
+					return;
+				}
+				ProcessingQueue = ActiveCommandQueue;
+				ActiveCommandQueue = (ActiveCommandQueue == &(CommandQueues[0]))
+					? &(CommandQueues[1])
+					: &(CommandQueues[0]);
+			}
+
+			if (bFlushMassCommandBuffer)
+			{
+				MassEntityManager.FlushCommands(GetSharedMassCommandBuffer());
+			}
+			for (FEnvironmentCommand& Command : *ProcessingQueue)
+			{
+				Command.CommandFunction(Command.CommandData);
+			}
+			ProcessingQueue->SetNum(0, EAllowShrinking::No);
 		}
-		CommandQueue.SetNum(0, EAllowShrinking::No);
 	}
 } // namespace UE::Editor::DataStorage
diff --git a/Engine/Plugins/Experimental/EditorDataStorage/Source/TedsCore/Private/TypedElementDatabaseEnvironment.h b/Engine/Plugins/Experimental/EditorDataStorage/Source/TedsCore/Private/TypedElementDatabaseEnvironment.h
--- a/Engine/Plugins/Experimental/EditorDataStorage/Source/TedsCore/Private/TypedElementDatabaseEnvironment.h
+++ b/Engine/Plugins/Experimental/EditorDataStorage/Source/TedsCore/Private/TypedElementDatabaseEnvironment.h
@@ -74,7 +74,7 @@ namespace UE::Editor::DataStorage
 
 		void NextUpdateCycle();
 		uint64 GetUpdateCycleId() const;
-
+		
 		struct FEnvironmentCommand
 		{
 			void (*CommandFunction)(void*);
@@ -82,11 +82,11 @@ namespace UE::Editor::DataStorage
 			void* CommandData = nullptr;
 		};
 
-		// Commands are flushed on NextUpdateCycle
+		// Commands are automatically flushed on NextUpdateCycle or when explicitly flushed with `FlushCommands`.
 		void PushCommands(TConstArrayView<const FEnvironmentCommand> Commands);
+		void FlushCommands();
 
 	private:
-		void FlushCommands();
 		
 		UEditorDataStorage& DataStorage;
 		Legacy::FCommandBuffer DirectDeferredCommands;
@@ -99,7 +99,8 @@ namespace UE::Editor::DataStorage
 		FTedsHierarchyRegistrar HierarchyInterface;
 
 		FMutex CommandQueueMutex;
-		TArray<FEnvironmentCommand> CommandQueue;
+		TArray<FEnvironmentCommand> CommandQueues[2];
+		TArray<FEnvironmentCommand>* ActiveCommandQueue;
 
 		FMassEntityManager& MassEntityManager;
 		FMassProcessingPhaseManager& MassPhaseManager;
diff --git a/Engine/Source/Runtime/TypedElementFramework/Public/Elements/Interfaces/TypedElementQueryStorageInterfaces.h b/Engine/Source/Runtime/TypedElementFramework/Public/Elements/Interfaces/TypedElementQueryStorageInterfaces.h
--- a/Engine/Source/Runtime/TypedElementFramework/Public/Elements/Interfaces/TypedElementQueryStorageInterfaces.h
+++ b/Engine/Source/Runtime/TypedElementFramework/Public/Elements/Interfaces/TypedElementQueryStorageInterfaces.h
@@ -357,8 +357,6 @@ namespace UE::Editor::DataStorage
 		AutoBalanceParallelChunkProcessing = 1 << 1,
 		IgnoreActiveState = 1 << 2, //< If set, a direct call will not check activatable queries if they're set.
 		AllowBoundQueries = 1 << 3, //< Stops checking if queries are bound when set, otherwise bound queries will assert.
-		/** Force commands added to the query callback's context or parent context to execute immediately after completing. */
-		ForceCommandBufferFlush = 1 << 4,
 	};
 	ENUM_CLASS_FLAGS(ERunQueryFlags);
 
