TEDS Hierarchies: Replace the ParentChangedColumn argument when registering a hierarchy with a bool to enable the feature instead, a dynamic column is now created internally and can be queried for using ICoreProvider::GetParentChangedColumn()
Enable this feature for all "editor" hierarchies to detect changes to them. E.g Actor -> Actor, Folder -> Actor, Actor -> Component

#Synced-CL 51094648

diff --git a/Engine/Plugins/Experimental/EditorDataStorage/Source/TedsCore/Private/Hierarchy/EditorDataStorageHierarchyColumns.h b/Engine/Plugins/Experimental/EditorDataStorage/Source/TedsCore/Private/Hierarchy/EditorDataStorageHierarchyColumns.h
--- a/Engine/Plugins/Experimental/EditorDataStorage/Source/TedsCore/Private/Hierarchy/EditorDataStorageHierarchyColumns.h
+++ b/Engine/Plugins/Experimental/EditorDataStorage/Source/TedsCore/Private/Hierarchy/EditorDataStorageHierarchyColumns.h
@@ -44,6 +44,13 @@ struct FEditorDataHierarchyUnresolvedParent_Template final : public FEditorDataS
 	
 };
 
+USTRUCT(meta = (DisplayName = "Parent Changed Column (Template)", EditorDataStorage_DynamicColumnTemplate))
+struct FEditorDataHierarchyParentChangedTag_Template final : public FEditorDataStorageTag
+{
+	GENERATED_BODY()
+};
+
+
 namespace UE::Editor::DataStorage
 {
 	using FHierarchyData_Template = FEditorDataHierarchyData_Template;
diff --git a/Engine/Plugins/Experimental/EditorDataStorage/Source/TedsCore/Private/Hierarchy/EditorDataStorageHierarchyImplementation.cpp b/Engine/Plugins/Experimental/EditorDataStorage/Source/TedsCore/Private/Hierarchy/EditorDataStorageHierarchyImplementation.cpp
--- a/Engine/Plugins/Experimental/EditorDataStorage/Source/TedsCore/Private/Hierarchy/EditorDataStorageHierarchyImplementation.cpp
+++ b/Engine/Plugins/Experimental/EditorDataStorage/Source/TedsCore/Private/Hierarchy/EditorDataStorageHierarchyImplementation.cpp
@@ -98,6 +98,28 @@ namespace UE::Editor::DataStorage
 			return InProvider->GenerateDynamicColumn(Desc);
 		}();
 		
+		const UScriptStruct* ParentChangedColumn = nullptr;
+		
+		// If the old deprecated ParentChangedColumn arg was provided, keep using it for backwards compatibility
+		PRAGMA_DISABLE_DEPRECATION_WARNINGS
+		if (RegistrationParams.ParentChangedColumn)
+		{
+			ParentChangedColumn = RegistrationParams.ParentChangedColumn;
+		}
+		PRAGMA_ENABLE_DEPRECATION_WARNINGS				
+		else if (RegistrationParams.bEnableParentChangedColumn)
+		{
+			ParentChangedColumn = [InProvider, &RegistrationParams]()
+			{
+				FDynamicColumnDescription Desc
+				{
+					.TemplateType = FEditorDataHierarchyParentChangedTag_Template::StaticStruct(),
+					.Identifier = RegistrationParams.Name
+				};
+				return InProvider->GenerateDynamicColumn(Desc);
+			}();
+		}
+		 
 		TUniquePtr<FTedsHierarchyAccessInterface> AccessInterface = [&]()
 		{
 			FTedsHierarchyAccessInterface::FConstructParams Params
@@ -106,7 +128,7 @@ namespace UE::Editor::DataStorage
 				.ParentTag = ParentTag,
 				.HierarchyData = HierarchyData,
 				.UnresolvedParentColumn = UnresolvedParentColumn,
-				.ParentChangedColumn = RegistrationParams.ParentChangedColumn
+				.ParentChangedColumn = ParentChangedColumn
 			};
 			return MakeUnique<FTedsHierarchyAccessInterface>(Params);
 		}();
@@ -118,7 +140,7 @@ namespace UE::Editor::DataStorage
 			.ParentTag = ParentTag,
 			.HierarchyData = HierarchyData,
 			.UnresolvedParentColumn = UnresolvedParentColumn,
-			.ParentChangedColumn = RegistrationParams.ParentChangedColumn,
+			.ParentChangedColumn = ParentChangedColumn,
 			.AccessInterface = MoveTemp(AccessInterface)
 		};
 
diff --git a/Engine/Plugins/Experimental/EditorDataStorageFeatures/Source/TedsActorCompatibility/Private/Compatibility/TedsActorComponentFactory.cpp b/Engine/Plugins/Experimental/EditorDataStorageFeatures/Source/TedsActorCompatibility/Private/Compatibility/TedsActorComponentFactory.cpp
--- a/Engine/Plugins/Experimental/EditorDataStorageFeatures/Source/TedsActorCompatibility/Private/Compatibility/TedsActorComponentFactory.cpp
+++ b/Engine/Plugins/Experimental/EditorDataStorageFeatures/Source/TedsActorCompatibility/Private/Compatibility/TedsActorComponentFactory.cpp
@@ -45,10 +45,10 @@ void UTedsActorComponentFactory::RegisterHierarchies(UE::Editor::DataStorage::IC
 	if (bEnableActorComponentCompatibility)
 	{
 		const FHierarchyRegistrationParams HierarchyParams
-					{
-						.Name = TEXT("ObjectComponent"),
-						.ParentChangedColumn = FComponentHierarchySyncTag::StaticStruct()
-					};
+			{
+				.Name = TEXT("ObjectComponent"),
+				.bEnableParentChangedColumn = true
+			};
 
 		DataStorage.RegisterHierarchy(HierarchyParams);
 	}
diff --git a/Engine/Plugins/Experimental/EditorDataStorageFeatures/Source/TedsActorCompatibility/Private/Compatibility/TedsActorParentQueries.cpp b/Engine/Plugins/Experimental/EditorDataStorageFeatures/Source/TedsActorCompatibility/Private/Compatibility/TedsActorParentQueries.cpp
--- a/Engine/Plugins/Experimental/EditorDataStorageFeatures/Source/TedsActorCompatibility/Private/Compatibility/TedsActorParentQueries.cpp
+++ b/Engine/Plugins/Experimental/EditorDataStorageFeatures/Source/TedsActorCompatibility/Private/Compatibility/TedsActorParentQueries.cpp
@@ -43,7 +43,14 @@ void UActorParentDataStorageFactory::PreRegister(UE::Editor::DataStorage::ICoreP
 
 void UActorParentDataStorageFactory::RegisterHierarchies(UE::Editor::DataStorage::ICoreProvider& InDataStorage)
 {
-	UE::Editor::DataStorage::FHierarchyRegistrationParams ActorHierarchyParams{.Name = UE::Editor::DataStorage::ActorParentQueries::Private::ActorEditorHierarchyName};
+	using namespace UE::Editor::DataStorage;
+	
+	UE::Editor::DataStorage::FHierarchyRegistrationParams ActorHierarchyParams
+		{
+			.Name = ActorParentQueries::Private::ActorEditorHierarchyName, 
+			.bEnableParentChangedColumn = true
+		};
+	
 	EditorActorHierarchyHandle = InDataStorage.RegisterHierarchy(ActorHierarchyParams);
 }
 
diff --git a/Engine/Plugins/Experimental/EditorDataStorageFeatures/Source/TedsActorCompatibility/Private/Compatibility/TedsActorParentQueries.h b/Engine/Plugins/Experimental/EditorDataStorageFeatures/Source/TedsActorCompatibility/Private/Compatibility/TedsActorParentQueries.h
--- a/Engine/Plugins/Experimental/EditorDataStorageFeatures/Source/TedsActorCompatibility/Private/Compatibility/TedsActorParentQueries.h
+++ b/Engine/Plugins/Experimental/EditorDataStorageFeatures/Source/TedsActorCompatibility/Private/Compatibility/TedsActorParentQueries.h
@@ -2,6 +2,7 @@
 
 #pragma once
 
+#include "DataStorage/CommonTypes.h"
 #include "DataStorage/Handles.h"
 #include "Elements/Interfaces/TypedElementDataStorageFactory.h"
 #include "UObject/ObjectMacros.h"
diff --git a/Engine/Plugins/Experimental/EditorDataStorageFeatures/Source/TedsActorCompatibility/Public/Columns/TedsActorComponentCompatibilityColumns.h b/Engine/Plugins/Experimental/EditorDataStorageFeatures/Source/TedsActorCompatibility/Public/Columns/TedsActorComponentCompatibilityColumns.h
--- a/Engine/Plugins/Experimental/EditorDataStorageFeatures/Source/TedsActorCompatibility/Public/Columns/TedsActorComponentCompatibilityColumns.h
+++ b/Engine/Plugins/Experimental/EditorDataStorageFeatures/Source/TedsActorCompatibility/Public/Columns/TedsActorComponentCompatibilityColumns.h
@@ -7,13 +7,6 @@
 
 USTRUCT()
 struct FActorComponentTypeTag : public FEditorDataStorageTag
-{
-	GENERATED_BODY()
-};
-
-// Tag added onto component rows when their parent changes. Removed at the end of the frame
-USTRUCT(meta = (DisplayName = "Component Hierarchy Change"))
-struct FComponentHierarchySyncTag final : public FEditorDataStorageTag
 {
 	GENERATED_BODY()
 };
\ No newline at end of file
diff --git a/Engine/Plugins/Experimental/EditorDataStorageFeatures/Source/TedsActorCompatibility/TedsActorCompatibility.Build.cs b/Engine/Plugins/Experimental/EditorDataStorageFeatures/Source/TedsActorCompatibility/TedsActorCompatibility.Build.cs
--- a/Engine/Plugins/Experimental/EditorDataStorageFeatures/Source/TedsActorCompatibility/TedsActorCompatibility.Build.cs
+++ b/Engine/Plugins/Experimental/EditorDataStorageFeatures/Source/TedsActorCompatibility/TedsActorCompatibility.Build.cs
@@ -30,7 +30,8 @@ public class TedsActorCompatibility : ModuleRules
 					"TedsTableViewer",
 					"SlateCore",
 					"Slate",
-					"LevelInstanceEditor"
+					"LevelInstanceEditor",
+					"TedsEditorCompatibility"
 				});
 
 			DynamicallyLoadedModuleNames.AddRange(new string[] { });
diff --git a/Engine/Source/Runtime/TypedElementFramework/Public/Elements/Interfaces/TypedElementDataStorageInterface.h b/Engine/Source/Runtime/TypedElementFramework/Public/Elements/Interfaces/TypedElementDataStorageInterface.h
--- a/Engine/Source/Runtime/TypedElementFramework/Public/Elements/Interfaces/TypedElementDataStorageInterface.h
+++ b/Engine/Source/Runtime/TypedElementFramework/Public/Elements/Interfaces/TypedElementDataStorageInterface.h
@@ -56,8 +56,11 @@ struct FHierarchyRegistrationParams
 {
 	FName Name;
 
-	// Optional column that is added to a row when the parent changes (useful for change detection).
+	// Optionally add a column to any row participating in the hierarchy if its parent changes (useful for change detection).
 	// This behaves similar to the sync tags and is removed at the end of the frame
+	bool bEnableParentChangedColumn = false;
+	
+	UE_DEPRECATED(5.8, "Providing an explicit column has been deprecated, use the bool argument + ICoreProvider::GetParentChangedColumn()")
 	const UScriptStruct* ParentChangedColumn = nullptr;
 };
 	
