[TEDS] Fixed crash when QueryContext::HasColumns(RowHandle, UScriptStruct*) called with invalid RowHandle

## Problem
Observer was authored which read a RowHandle from the MappingTable - the row handle returned was invalid and therfore had no table/archetpye associated with it.  The function was written to assume that all RowHandles have valid ArchetypeHandles and this was passed to the Mass API to get composition.  Mass internally asserts when getting archetype composition in this manner.

## Solution
Checking for validity of the returned ArchetypeHandle done to skirt the assertion.  This is safe on the basis that an invalid row has no associated table and therefore not columns.

#Synced-CL 51308403

diff --git a/Engine/Plugins/Experimental/EditorDataStorage/Source/TedsCore/Private/Processors/TypedElementProcessorAdaptors.cpp b/Engine/Plugins/Experimental/EditorDataStorage/Source/TedsCore/Private/Processors/TypedElementProcessorAdaptors.cpp
--- a/Engine/Plugins/Experimental/EditorDataStorage/Source/TedsCore/Private/Processors/TypedElementProcessorAdaptors.cpp
+++ b/Engine/Plugins/Experimental/EditorDataStorage/Source/TedsCore/Private/Processors/TypedElementProcessorAdaptors.cpp
@@ -102,10 +102,15 @@ namespace UE::Editor::DataStorage
 				FMassEntityHandle Entity = FMassEntityHandle::FromNumber(Row);
 				FMassEntityManager& Manager = Context.GetEntityManagerChecked();
 				FMassArchetypeHandle Archetype = Manager.GetArchetypeForEntity(Entity);
-				const FMassArchetypeCompositionDescriptor& Composition = Manager.GetArchetypeComposition(Archetype);
-
+				// Entity was invalid (either never was or was destroyed).  Invalid entities have no columns
+				if (!Archetype.IsValid())
+				{
+					return false;
+				}
+				
 				if (UE::Mass::IsA<FMassTag>(ColumnType) || UE::Mass::IsA<FMassFragment>(ColumnType))
 				{
+					const FMassArchetypeCompositionDescriptor& Composition = Manager.GetArchetypeComposition(Archetype);
 					return Composition.GetElementsBitSet().Contains(ColumnType);
 				}
 				const bool bIsTagOrFragment = false;
