[TEDS] Removed ClearRow in the command buffer.

Whenever an unassigned row is deleted, TEDS will check if there are any pending columns and/or commands that reference that row. These checks are however not strictly necessary as all commands will validate the row during processing, so the clearing has been removed to get back a sizeable amount of performance.

#Synced-CL 51279932

diff --git a/Engine/Plugins/Experimental/EditorDataStorage/Source/TedsCore/Private/TypedElementDatabase.cpp b/Engine/Plugins/Experimental/EditorDataStorage/Source/TedsCore/Private/TypedElementDatabase.cpp
--- a/Engine/Plugins/Experimental/EditorDataStorage/Source/TedsCore/Private/TypedElementDatabase.cpp
+++ b/Engine/Plugins/Experimental/EditorDataStorage/Source/TedsCore/Private/TypedElementDatabase.cpp
@@ -574,7 +574,6 @@ void UEditorDataStorage::RemoveRow(UE::Editor::DataStorage::RowHandle Row)
 		}
 		else
 		{
-			Environment->GetDirectDeferredCommands().Clear(Row);
 			ActiveEditorEntityManager->ReleaseReservedEntity(FMassEntityHandle::FromNumber(Row));
 		}
 		Environment->GetMappingTable().MarkDirty();
@@ -769,7 +768,7 @@ const void* UEditorDataStorage::GetColumnData(UE::Editor::DataStorage::RowHandle
 				return Column.GetMemory();
 			}
 		}
-		else
+		else if (ActiveEditorEntityManager->IsEntityValid(Entity) && ActiveEditorEntityManager->IsEntityReserved(Entity))
 		{
 			return Environment->GetDirectDeferredCommands().GetQueuedDataColumn(Row, ColumnType);
 		}
@@ -918,9 +917,7 @@ void UEditorDataStorage::BatchAddRemoveColumns(TConstArrayView<UE::Editor::DataS
 	{	
 		FMassElementBitSet ElementsToAdd = Private::ColumnsToBitSets(ColumnsToAdd);
 		FMassElementBitSet ElementsToRemove = Private::ColumnsToBitSets(ColumnsToRemove);
-			
-		Legacy::FCommandBuffer& CommandBuffer = Environment->GetDirectDeferredCommands();
-			
+				
 		// Sort rows (entities) into to matching table (archetype) bucket. If a row hasn't been constructed yet, but
 		// only reserved, then queue the command for later execution.
 		EntityArchetypeLookup LookupTable;
@@ -957,6 +954,7 @@ void UEditorDataStorage::BatchAddRemoveColumns(TConstArrayView<UE::Editor::DataS
 
 		if (!DelayedRows.IsEmpty())
 		{
+			Legacy::FCommandBuffer& CommandBuffer = Environment->GetDirectDeferredCommands();
 			CommandBuffer.Queue_AddRemoveColumnsCommand(DelayedRows, ElementsToAdd, ElementsToRemove);
 		}
 	}
@@ -966,11 +964,10 @@ bool UEditorDataStorage::HasColumns(UE::Editor::DataStorage::RowHandle Row, TCon
 {
 	if (ActiveEditorEntityManager)
 	{
-		bool bHasAllColumns = true;
-		
 		FMassEntityHandle Entity = FMassEntityHandle::FromNumber(Row);
 		if (ActiveEditorEntityManager->IsEntityActive(Entity))
 		{
+			bool bHasAllColumns = true;
 			FMassArchetypeHandle Archetype = ActiveEditorEntityManager->GetArchetypeForEntity(Entity);
 			const FMassArchetypeCompositionDescriptor& Composition = ActiveEditorEntityManager->GetArchetypeComposition(Archetype);
 
@@ -986,18 +983,19 @@ bool UEditorDataStorage::HasColumns(UE::Editor::DataStorage::RowHandle Row, TCon
 					return false;
 				}
 			}
+			return bHasAllColumns;
 		}
-		else
+		else if (ActiveEditorEntityManager->IsEntityValid(Entity) && ActiveEditorEntityManager->IsEntityReserved(Entity))
 		{
+			bool bHasAllColumns = true;
 			const UE::Editor::DataStorage::Legacy::FCommandBuffer& CommandBuffer = Environment->GetDirectDeferredCommands();
 			const UScriptStruct* const* ColumnTypesEnd = ColumnTypes.end();
 			for (const UScriptStruct* const* ColumnType = ColumnTypes.begin(); ColumnType != ColumnTypesEnd && bHasAllColumns; ++ColumnType)
 			{
 				bHasAllColumns = CommandBuffer.HasColumn(Row, *ColumnType);
 			}
+			return bHasAllColumns;
 		}
-
-		return bHasAllColumns;
 	}
 	return false;
 }
@@ -1006,11 +1004,10 @@ bool UEditorDataStorage::HasColumns(UE::Editor::DataStorage::RowHandle Row, TCon
 {
 	if (ActiveEditorEntityManager)
 	{
-		bool bHasAllColumns = true;
-	
 		FMassEntityHandle Entity = FMassEntityHandle::FromNumber(Row);
 		if (ActiveEditorEntityManager->IsEntityActive(Entity))
 		{
+			bool bHasAllColumns = true;
 			FMassArchetypeHandle Archetype = ActiveEditorEntityManager->GetArchetypeForEntity(Entity);
 			const FMassArchetypeCompositionDescriptor& Composition = ActiveEditorEntityManager->GetArchetypeComposition(Archetype);
 
@@ -1027,18 +1024,24 @@ bool UEditorDataStorage::HasColumns(UE::Editor::DataStorage::RowHandle Row, TCon
 				}
 				return false;
 			}
+			return bHasAllColumns;
 		}
-		else
+		else if (ActiveEditorEntityManager->IsEntityValid(Entity) && ActiveEditorEntityManager->IsEntityReserved(Entity))
 		{
+			bool bHasAllColumns = true;
 			const UE::Editor::DataStorage::Legacy::FCommandBuffer& CommandBuffer = Environment->GetDirectDeferredCommands();
 			const TWeakObjectPtr<const UScriptStruct>* ColumnTypesEnd = ColumnTypes.end();
 			for (const TWeakObjectPtr<const UScriptStruct>* ColumnType = ColumnTypes.begin(); ColumnType != ColumnTypesEnd && bHasAllColumns; ++ColumnType)
 			{
+				// It's ok if ColumnType is null since it's not de-referenced, only used as a hash.
 				bHasAllColumns = CommandBuffer.HasColumn(Row, ColumnType->Get());
 			}
+			return bHasAllColumns;
+		}
+		else
+		{
+			return false;
 		}
-
-		return bHasAllColumns;
 	}
 	return false;
 }
@@ -1063,7 +1066,7 @@ void UEditorDataStorage::ListColumns(UE::Editor::DataStorage::RowHandle Row, UE:
 				};
 			Composition.GetElementsBitSet().ExportTypes(CallbackWrapper);
 		}
-		else
+		else if (ActiveEditorEntityManager->IsEntityValid(Entity) && ActiveEditorEntityManager->IsEntityReserved(Entity))
 		{
 			Environment->GetDirectDeferredCommands().ListColumns(Row, Callback);
 		}
@@ -1101,7 +1104,7 @@ void UEditorDataStorage::ListColumns(UE::Editor::DataStorage::RowHandle Row, UE:
 				});
 
 		}
-		else
+		else if (ActiveEditorEntityManager->IsEntityValid(Entity) && ActiveEditorEntityManager->IsEntityReserved(Entity))
 		{
 			Environment->GetDirectDeferredCommands().ListColumns(Row, Callback);
 		}
diff --git a/Engine/Plugins/Experimental/EditorDataStorage/Source/TedsCore/Private/TypedElementDatabaseCommandBuffer.cpp b/Engine/Plugins/Experimental/EditorDataStorage/Source/TedsCore/Private/TypedElementDatabaseCommandBuffer.cpp
--- a/Engine/Plugins/Experimental/EditorDataStorage/Source/TedsCore/Private/TypedElementDatabaseCommandBuffer.cpp
+++ b/Engine/Plugins/Experimental/EditorDataStorage/Source/TedsCore/Private/TypedElementDatabaseCommandBuffer.cpp
@@ -83,69 +83,6 @@ namespace UE::Editor::DataStorage::Legacy
 		}
 	}
 
-	template <typename T>
-	concept HasRow = requires(T t) { t.Row; } && 
-		std::is_same_v<std::remove_cvref_t<decltype(std::declval<T>().Row)>, RowHandle>;
-
-	template <typename T>
-	concept HasRowsView = requires(T t) { t.Rows; } && 
-		(
-			std::is_same_v<std::remove_cvref_t<decltype(std::declval<T>().Rows)>, TArrayView<RowHandle>> ||
-			std::is_same_v<std::remove_cvref_t<decltype(std::declval<T>().Rows)>, TArrayView<const RowHandle>>
-		);
-
-	struct FCheckForRowVisitor
-	{
-		RowHandle Row;
-		
-		template<typename T>
-		bool operator()(T& Command) requires (!HasRowsView<T> && HasRow<T>)
-		{
-			return Command.Row == Row;
-		}
-
-		template<typename T>
-		bool operator()(T& Command) requires (HasRowsView<T> && !HasRow<T>)
-		{
-			int32 Index;
-			if (Command.Rows.Find(Row, Index))
-			{
-				// Swap the index and the last, then reduce the size of the array by one to remove the row without
-				// resizing the array as that's a fixed allocation plus it's short lived so not worth resizing.
-				Command.Rows[Index] = Command.Rows.Last();
-				Command.Rows.LeftChopInline(1);
-			}
-			return Command.Rows.IsEmpty();
-		}
-
-		template<typename T>
-		bool operator()(T& Command)
-		{
-			static_assert(sizeof(T) == 0, "Command requires either `RowHandle Row` or `TArrayView<RowHandle> Rows`.");
-		} //-V591
-	};
-
-	void FCommandBuffer::Clear(RowHandle Row)
-	{
-		for (TMap<PendingColumnMappingKey, void*>::TIterator It = PendingColumns.CreateIterator(); It; ++It)
-		{
-			if (It.Key().Key == Row)
-			{
-				It.RemoveCurrent();
-			}
-		}
-
-		FCheckForRowVisitor Visitor;
-		Visitor.Row = Row;
-		for (TArray<CommandData>::TIterator It = Commands.CreateIterator(); It; ++It)
-		{
-			if (Visit(Visitor, *It))
-			{
-				It.RemoveCurrent();
-			}
-		}
-	}
-
 	FCommandBuffer::FAddDataColumnCommand::~FAddDataColumnCommand()
 	{
 		if (ColumnType.IsValid() && (ColumnType->StructFlags & (STRUCT_IsPlainOldData | STRUCT_NoDestructor)) == 0)
diff --git a/Engine/Plugins/Experimental/EditorDataStorage/Source/TedsCore/Private/TypedElementDatabaseCommandBuffer.h b/Engine/Plugins/Experimental/EditorDataStorage/Source/TedsCore/Private/TypedElementDatabaseCommandBuffer.h
--- a/Engine/Plugins/Experimental/EditorDataStorage/Source/TedsCore/Private/TypedElementDatabaseCommandBuffer.h
+++ b/Engine/Plugins/Experimental/EditorDataStorage/Source/TedsCore/Private/TypedElementDatabaseCommandBuffer.h
@@ -32,9 +32,7 @@ namespace UE::Editor::DataStorage
 			void ListColumns(RowHandle Row, ColumnListCallbackRef Callback) const;
 			/** Returns a list of columns and their data addresses (can be null for tags) pending processing for a row.. */
 			void ListColumns(RowHandle Row, ColumnListWithDataCallbackRef Callback);
-			/** Clears all entries in the internal cache for the provided row. */
-			void Clear(RowHandle Row);
-	
+			
 			/**
 			 * @section Queue_* functions create a command that can be stored for later execution in the provided command buffer.
 			 * Typically the Execute_* counter part will be called for execution.
