Fix crash that was occurring with ValueTags due to added Mass Check

## Problem
* Mass required that ConstShared Fragments are registered before being added to entities.  TEDS has never done so.  The Mass API added a check() for this on Feb 23, 2026.
* TEDS had been tracking the const shared fragments already and passing the instance. This was removed in favor of using the implementation defined within Mass.

## Testing via CMD

Cmd: TEDS.Debug.CreateRow
LogEditorDataStorage: Warning: Added Row 4295002707
Cmd: TEDS.Debug.ValueTag.AddColumn 4295002707 Color Red
Cmd: TEDS.Debug.CreateRow
LogEditorDataStorage: Warning: Added Row 4295002708
Cmd: TEDS.Debug.ValueTag.AddColumn 4295002708 Color Blue
Cmd: TEDS.Debug.ValueTag.RunQuery Color
LogEditorDataStorage: Warning: Processed 2 rows
Cmd: TEDS.Debug.ValueTag.RunQuery Color Blue
LogEditorDataStorage: Warning: Processed 1 rows
Cmd: TEDS.Debug.ValueTag.RunQuery Color Red
LogEditorDataStorage: Warning: Processed 1 rows

#Synced-CL 51330672

diff --git a/Engine/Plugins/Experimental/EditorDataStorage/Source/TedsCore/Private/DynamicColumnGenerator.cpp b/Engine/Plugins/Experimental/EditorDataStorage/Source/TedsCore/Private/DynamicColumnGenerator.cpp
--- a/Engine/Plugins/Experimental/EditorDataStorage/Source/TedsCore/Private/DynamicColumnGenerator.cpp
+++ b/Engine/Plugins/Experimental/EditorDataStorage/Source/TedsCore/Private/DynamicColumnGenerator.cpp
@@ -9,6 +9,7 @@
 #include "Kismet2/ReloadUtilities.h"
 #include "Queries/TypedElementExtendedQueryStore.h"
 #include "UObject/MetaData.h"
+#include "MassEntityManager.h"
 
 #include UE_INLINE_GENERATED_CPP_BY_NAME(DynamicColumnGenerator)
 
@@ -189,8 +190,9 @@ namespace UE::Editor::DataStorage
 		return false;
 	}
 
-	FValueTagManager::FValueTagManager(FDynamicColumnGenerator& InColumnGenerator)
-		: ColumnGenerator(InColumnGenerator)
+	FValueTagManager::FValueTagManager(FMassEntityManager& InEntityManager, FDynamicColumnGenerator& InColumnGenerator)
+		: EntityManager(InEntityManager)
+		, ColumnGenerator(InColumnGenerator)
 	{
 	}
 	
@@ -199,30 +201,15 @@ namespace UE::Editor::DataStorage
 		TPair<FValueTag, FName> Pair(InTag, InValue);
 	
 		UE_MT_SCOPED_WRITE_ACCESS(AccessDetector);
-	
-		// Common path
-		{
-			if (FConstSharedStruct* TagStruct = ValueTagLookup.Find(Pair))
-			{
-				return *TagStruct;
-			}
-		}
-	
-		// 
+		
+		const UScriptStruct* ColumnType = GenerateColumnType(InTag);
+		const FTedsValueTagColumn Overlay
 		{
-			const UScriptStruct* ColumnType = GenerateColumnType(InTag);
-	
-			const FTedsValueTagColumn Overlay
-			{
-				.Value = InValue
-			};
-	
-			FConstSharedStruct SharedStruct = FConstSharedStruct::Make(ColumnType, reinterpret_cast<const uint8*>(&Overlay));
-			
-			ValueTagLookup.Emplace(Pair, SharedStruct);
-	
-			return SharedStruct;
-		}
+			.Value = InValue
+		};
+		
+		const FConstSharedStruct& SharedStruct = EntityManager.GetOrCreateConstSharedFragment(*ColumnType, reinterpret_cast<const uint8*>(&Overlay));
+		return SharedStruct;
 	}
 	
 	const UScriptStruct* FValueTagManager::GenerateColumnType(const FValueTag& Tag)
diff --git a/Engine/Plugins/Experimental/EditorDataStorage/Source/TedsCore/Private/DynamicColumnGenerator.h b/Engine/Plugins/Experimental/EditorDataStorage/Source/TedsCore/Private/DynamicColumnGenerator.h
--- a/Engine/Plugins/Experimental/EditorDataStorage/Source/TedsCore/Private/DynamicColumnGenerator.h
+++ b/Engine/Plugins/Experimental/EditorDataStorage/Source/TedsCore/Private/DynamicColumnGenerator.h
@@ -14,7 +14,7 @@
 
 class FName;
 class UScriptStruct;
-
+struct FMassEntityManager;
 // The template struct that is used to generate the ValueTag column
 USTRUCT(meta=(EditorDataStorage_DynamicColumnTemplate))
 struct FTedsValueTagColumn : public FTedsSharedColumn
@@ -111,15 +111,14 @@ namespace UE::Editor::DataStorage
 	class FValueTagManager
 	{
 	public:
-		explicit FValueTagManager(FDynamicColumnGenerator& InColumnGenerator);
+		explicit FValueTagManager(FMassEntityManager& InEntityManager, FDynamicColumnGenerator& InColumnGenerator);
 		FConstSharedStruct GenerateValueTag(const FValueTag& InTag, const FName& InValue);
 		const UScriptStruct* GenerateColumnType(const FValueTag& InTag);
 	private:
 
 		UE_MT_DECLARE_RW_ACCESS_DETECTOR(AccessDetector);
 
-		TMap<TPair<FValueTag, FName>, FConstSharedStruct> ValueTagLookup;
-
+		FMassEntityManager& EntityManager;
 		FDynamicColumnGenerator& ColumnGenerator;
 	};
 } // namespace UE::Editor::DataStorage
diff --git a/Engine/Plugins/Experimental/EditorDataStorage/Source/TedsCore/Private/TypedElementDatabaseEnvironment.cpp b/Engine/Plugins/Experimental/EditorDataStorage/Source/TedsCore/Private/TypedElementDatabaseEnvironment.cpp
--- a/Engine/Plugins/Experimental/EditorDataStorage/Source/TedsCore/Private/TypedElementDatabaseEnvironment.cpp
+++ b/Engine/Plugins/Experimental/EditorDataStorage/Source/TedsCore/Private/TypedElementDatabaseEnvironment.cpp
@@ -14,7 +14,7 @@ namespace UE::Editor::DataStorage
 		, MappingTable(InDataStorage)
 		, Queries(DynamicColumnGenerator)
 		, MementoSystem(InDataStorage)
-		, ValueTagManager(DynamicColumnGenerator)
+		, ValueTagManager(InMassEntityManager, DynamicColumnGenerator)
 		, ActiveCommandQueue(CommandQueues)
 		, MassEntityManager(InMassEntityManager)
 		, MassPhaseManager(InMassPhaseManager)
