TEDS Hierarchies: Make sure the parent is assigned before trying to get/add columns to it

#Synced-CL 51413783

diff --git a/Engine/Plugins/Experimental/EditorDataStorage/Source/TedsCore/Private/Hierarchy/EditorDataStorageHierarchyImplementation.cpp b/Engine/Plugins/Experimental/EditorDataStorage/Source/TedsCore/Private/Hierarchy/EditorDataStorageHierarchyImplementation.cpp
--- a/Engine/Plugins/Experimental/EditorDataStorage/Source/TedsCore/Private/Hierarchy/EditorDataStorageHierarchyImplementation.cpp
+++ b/Engine/Plugins/Experimental/EditorDataStorage/Source/TedsCore/Private/Hierarchy/EditorDataStorageHierarchyImplementation.cpp
@@ -487,7 +487,7 @@ namespace UE::Editor::DataStorage
 			PreviousParent = HierarchyDataColumn_Target->Parent;
 		}
 		
-		if (Parent != InvalidRowHandle)
+		if (CoreProvider->IsRowAssigned(Parent))
 		{
 			// No updates needed if the parent is the same as before
 			if (Parent != PreviousParent)
