[EditorDataStorage] Add better warnings when attempting to add multiple rows to non-existing table.  We're getting this warning now, need more info.

#Synced-CL 51553633

diff --git a/Engine/Plugins/Experimental/EditorDataStorage/Source/TedsCore/Private/TypedElementDatabase.cpp b/Engine/Plugins/Experimental/EditorDataStorage/Source/TedsCore/Private/TypedElementDatabase.cpp
--- a/Engine/Plugins/Experimental/EditorDataStorage/Source/TedsCore/Private/TypedElementDatabase.cpp
+++ b/Engine/Plugins/Experimental/EditorDataStorage/Source/TedsCore/Private/TypedElementDatabase.cpp
@@ -516,13 +516,24 @@ bool UEditorDataStorage::BatchAddRow(
 	if (ActiveEditorEntityManager)
 	{
 		FActiveCallScope Scope(*this);
-		
+
 		const int32 TableHandleAsIndex = Private::ConvertTableHandleToIndex(Table);
-		checkf(TableHandleAsIndex < Tables.Num(), TEXT("Attempting to add multiple rows to a non-existing table."));
-	
+
+		// checkf is not clever, so you need to wrap expensive side effect code this way.
+		if (!(TableHandleAsIndex < Tables.Num()))
+		{
+			checkf(
+				TableHandleAsIndex < Tables.Num(),
+				TEXT("Attempting to add multiple rows to a non-existing table. Index: %d, Tables.Num(): %d, TableName: %s"),
+				TableHandleAsIndex,
+				Tables.Num(),
+				*FindTableName(Table) // this is expensive.
+			);
+		}
+
 		TArray<FMassEntityHandle> Entities;
 		Entities.Reserve(Count);
-		TSharedRef<FMassEntityManager::FEntityCreationContext> Context = 
+		TSharedRef<FMassEntityManager::FEntityCreationContext> Context =
 			ActiveEditorEntityManager->BatchCreateEntities(Tables[TableHandleAsIndex], Count, Entities);
 		
 		for (FMassEntityHandle Entity : Entities)
@@ -1635,6 +1646,18 @@ void UEditorDataStorage::TickCooperativeTasks()
 	}
 }
 
+FString UEditorDataStorage::FindTableName(UE::Editor::DataStorage::TableHandle InHandle) const
+{
+	for (const TPair<FName, UE::Editor::DataStorage::TableHandle>& Pair : TableNameLookup)
+	{
+		if (Pair.Value == InHandle)
+		{
+			return Pair.Key.ToString();
+		}
+	}
+	return TEXT("<unknown>");
+};
+
 bool UEditorDataStorage::IsAvailable() const
 {
 	return bool(ActiveEditorEntityManager);
diff --git a/Engine/Plugins/Experimental/EditorDataStorage/Source/TedsCore/Private/TypedElementDatabase.h b/Engine/Plugins/Experimental/EditorDataStorage/Source/TedsCore/Private/TypedElementDatabase.h
--- a/Engine/Plugins/Experimental/EditorDataStorage/Source/TedsCore/Private/TypedElementDatabase.h
+++ b/Engine/Plugins/Experimental/EditorDataStorage/Source/TedsCore/Private/TypedElementDatabase.h
@@ -228,6 +228,8 @@ private:
 
 	int32 GetTableChunkSize(FName TableName) const;
 	void TickCooperativeTasks();
+	// Reverse-lookup name from handle (O(n), only use for diagnostics, slow!)
+	FString FindTableName(UE::Editor::DataStorage::TableHandle InHandle) const;
 
 	struct FFactoryTypePair
 	{
