FTedsHierarchyAccessInterface : use IsRowAvailable instead of IsRowAssigned to fix a regression in SetParentRow

This fixes a regression from 51462296 for the case where the mass entity is valid but not active.

#Synced-CL 51597297

diff --git a/Engine/Plugins/Experimental/EditorDataStorage/Source/TedsCore/Private/Hierarchy/EditorDataStorageHierarchyImplementation.cpp b/Engine/Plugins/Experimental/EditorDataStorage/Source/TedsCore/Private/Hierarchy/EditorDataStorageHierarchyImplementation.cpp
--- a/Engine/Plugins/Experimental/EditorDataStorage/Source/TedsCore/Private/Hierarchy/EditorDataStorageHierarchyImplementation.cpp
+++ b/Engine/Plugins/Experimental/EditorDataStorage/Source/TedsCore/Private/Hierarchy/EditorDataStorageHierarchyImplementation.cpp
@@ -487,7 +487,7 @@ namespace UE::Editor::DataStorage
 			PreviousParent = HierarchyDataColumn_Target->Parent;
 		}
 		
-		if (CoreProvider->IsRowAssigned(Parent))
+		if (CoreProvider->IsRowAvailable(Parent))
 		{
 			// No updates needed if the parent is the same as before
 			if (Parent != PreviousParent)
