[TEDS] Added additional logging to track the state of TEDS Core and Compatibility (de)initialization.

This change also includes some additional cleanup which should in theory not be necessary, but does no harm either.

#rnx
#Synced-CL 51558978

diff --git a/Engine/Plugins/Experimental/EditorDataStorage/Source/TedsCore/Private/TypedElementDatabase.cpp b/Engine/Plugins/Experimental/EditorDataStorage/Source/TedsCore/Private/TypedElementDatabase.cpp
--- a/Engine/Plugins/Experimental/EditorDataStorage/Source/TedsCore/Private/TypedElementDatabase.cpp
+++ b/Engine/Plugins/Experimental/EditorDataStorage/Source/TedsCore/Private/TypedElementDatabase.cpp
@@ -2,6 +2,7 @@
 
 #include "TypedElementDatabase.h"
 
+#include <cinttypes>
 #include "DataStorage/Features.h"
 #include "DataStorage/Debug/Log.h"
 #include "Editor.h"
@@ -151,7 +152,6 @@ namespace UE::Editor::DataStorage::Private
 	
 	constexpr int32 ConvertTableHandleToIndex(const TableHandle InTableHandle)
 	{
-		ensure(InTableHandle == InvalidTableHandle || (InTableHandle <= static_cast<TableHandle>(std::numeric_limits<int32>::max()) && InTableHandle >= 0));
 		return static_cast<int32>(InTableHandle);
 	}
 } // namespace UE::Editor::DataStorage::Private
@@ -171,10 +171,17 @@ UEditorDataStorage::FActiveCallScope::~FActiveCallScope()
 	}
 }
 
+UEditorDataStorage::~UEditorDataStorage()
+{
+	UE_LOGF(LogEditorDataStorage, Log, "Destroyed TEDS Core.");
+}
+
 void UEditorDataStorage::Initialize()
 {
 	using namespace UE::Editor::DataStorage;
 
+	UE_LOGF(LogEditorDataStorage, Log, "Initializing TEDS Core.");
+
 	UMassEntityEditorSubsystem* Mass = UMassEntityEditorSubsystem::Get();
 	check(Mass);
 	OnPreMassTickHandle = Mass->GetOnPreTickDelegate().AddUObject(this, &UEditorDataStorage::OnPreMassTick);
@@ -229,6 +236,8 @@ void UEditorDataStorage::Initialize()
 				Phase, {}, GetQueryTickGroupName(EQueryTickGroups::PostUpdate), EExecutionMode::GameThread /* Needs main thread*/);
 		}
 	}
+
+	UE_LOGF(LogEditorDataStorage, Log, "Initialized TEDS Core.");
 }
 
 void UEditorDataStorage::PostInitialize()
@@ -310,9 +319,13 @@ UEditorDataStorageFactory* UEditorDataStorage::FindFactory(const UClass* Factory
 
 void UEditorDataStorage::Deinitialize()
 {
+	UE_LOGF(LogEditorDataStorage, Log, "De-initializing TEDS Core.");
+
 	checkf(Factories.IsEmpty(), TEXT("ResetFactories should have been called before deinitialized"));
-	
+
 	Reset();
+
+	UE_LOGF(LogEditorDataStorage, Log, "De-initialized TEDS Core.");
 }
 
 void UEditorDataStorage::OnPreMassTick(float DeltaTime)
@@ -373,9 +386,9 @@ UE::Editor::DataStorage::TableHandle UEditorDataStorage::RegisterTable(UE::Edito
 {
 	using namespace UE::Editor::DataStorage;
 
-	const int32 TableHandleAsIndex = Private::ConvertTableHandleToIndex(SourceTable);
-	if (ActiveEditorEntityManager && TableHandleAsIndex < Tables.Num() && !TableNameLookup.Contains(InName))
+	if (ActiveEditorEntityManager && IsTableValid(SourceTable) && !TableNameLookup.Contains(InName))
 	{
+		const int32 TableHandleAsIndex = Private::ConvertTableHandleToIndex(SourceTable);
 		TableHandle Result = Tables.Num();
 		FMassArchetypeCreationParams ArchetypeCreationParams;
 		ArchetypeCreationParams.DebugName = InName;
@@ -398,6 +411,13 @@ UE::Editor::DataStorage::TableHandle UEditorDataStorage::FindTable(const FName&
 	return TableHandle ? *TableHandle : InvalidTableHandle;
 }
 
+bool UEditorDataStorage::IsTableValid(UE::Editor::DataStorage::TableHandle Table) const
+{
+	using namespace UE::Editor::DataStorage;
+	
+	return Table != InvalidTableHandle && Table < Tables.Num();
+}
+
 UE::Editor::DataStorage::RowHandle UEditorDataStorage::ReserveRow()
 {
 	return ActiveEditorEntityManager 
@@ -435,8 +455,8 @@ UE::Editor::DataStorage::RowHandle UEditorDataStorage::AddRow(UE::Editor::DataSt
 
 	FActiveCallScope Scope(*this);
 	
+	checkf(IsTableValid(Table), TEXT("Attempting to add a row to a non-existing table."));
 	const int32 TableHandleAsIndex = Private::ConvertTableHandleToIndex(Table);
-	checkf(TableHandleAsIndex < Tables.Num(), TEXT("Attempting to add a row to a non-existing table."));
 	return ActiveEditorEntityManager 
 		? ActiveEditorEntityManager->CreateEntity(Tables[TableHandleAsIndex]).AsNumber()
 		: UE::Editor::DataStorage::InvalidRowHandle;
@@ -451,9 +471,9 @@ UE::Editor::DataStorage::RowHandle UEditorDataStorage::AddRow(UE::Editor::DataSt
 	{
 		FActiveCallScope Scope(*this);
 		
+		checkf(IsTableValid(Table), TEXT("Attempting to a row to a non-existing table."));
 		const int32 TableHandleAsIndex = Private::ConvertTableHandleToIndex(Table);
-		checkf(TableHandleAsIndex < Tables.Num(), TEXT("Attempting to a row to a non-existing table."));
-
+		
 		TArray<FMassEntityHandle> Entity;
 		Entity.Reserve(1);
 		TSharedRef<FMassEntityManager::FEntityCreationContext> Context =
@@ -471,9 +491,9 @@ bool UEditorDataStorage::AddRow(UE::Editor::DataStorage::RowHandle ReservedRow,
 {
 	using namespace UE::Editor::DataStorage;
 
+	checkf(IsTableValid(Table), TEXT("Attempting to add a row to a non-existing table."));
 	const int32 TableHandleAsIndex = Private::ConvertTableHandleToIndex(Table);
 	checkf(!IsRowAssigned(ReservedRow), TEXT("Attempting to assign a table to row that already has a table assigned."));
-	checkf(TableHandleAsIndex < Tables.Num(), TEXT("Attempting to add a row to a non-existing table."));
 	if (ActiveEditorEntityManager)
 	{
 		FActiveCallScope Scope(*this);
@@ -496,8 +516,8 @@ bool UEditorDataStorage::AddRow(UE::Editor::DataStorage::RowHandle ReservedRow,
 	{
 		FActiveCallScope Scope(*this);
 		
+		checkf(IsTableValid(Table), TEXT("Attempting to add a row to a non-existing table."));
 		const int32 TableHandleAsIndex = Private::ConvertTableHandleToIndex(Table);
-		checkf(TableHandleAsIndex < Tables.Num(), TEXT("Attempting to add a row to a non-existing table."));
 		
 		TSharedRef<FMassEntityManager::FEntityCreationContext> Context =
 			ActiveEditorEntityManager->BatchCreateReservedEntities(Tables[TableHandleAsIndex], { FMassEntityHandle::FromNumber(ReservedRow) });
@@ -517,20 +537,17 @@ bool UEditorDataStorage::BatchAddRow(
 	{
 		FActiveCallScope Scope(*this);
 
-		const int32 TableHandleAsIndex = Private::ConvertTableHandleToIndex(Table);
-
-		// checkf is not clever, so you need to wrap expensive side effect code this way.
-		if (!(TableHandleAsIndex < Tables.Num()))
+		if (!IsTableValid(Table))
 		{
 			checkf(
-				TableHandleAsIndex < Tables.Num(),
-				TEXT("Attempting to add multiple rows to a non-existing table. Index: %d, Tables.Num(): %d, TableName: %s"),
-				TableHandleAsIndex,
-				Tables.Num(),
-				*FindTableName(Table) // this is expensive.
+				false,
+				TEXT("Attempting to add multiple rows to a non-existing table (%" PRIu64 "). Tables.Num(): %d, TableName: %s"),
+				Table, Tables.Num(), *FindTableName(Table)
 			);
+			return false;
 		}
 
+		const int32 TableHandleAsIndex = Private::ConvertTableHandleToIndex(Table);
 		TArray<FMassEntityHandle> Entities;
 		Entities.Reserve(Count);
 		TSharedRef<FMassEntityManager::FEntityCreationContext> Context =
@@ -556,9 +573,17 @@ bool UEditorDataStorage::BatchAddRow(UE::Editor::DataStorage::TableHandle Table,
 	{
 		FActiveCallScope Scope(*this);
 		
+		if (!IsTableValid(Table))
+		{
+			checkf(
+				false,
+				TEXT("Attempting to add multiple reserved rows to a non-existing table (%" PRIu64 "). Tables.Num(): %d, TableName: %s"),
+				Table, Tables.Num(), *FindTableName(Table)
+			);
+			return false;
+		}
+		
 		const int32 TableHandleAsIndex = Private::ConvertTableHandleToIndex(Table);
-		checkf(TableHandleAsIndex < Tables.Num(), TEXT("Attempting to add multiple rows to a non-existing table."));
-	
 		TSharedRef<FMassEntityManager::FEntityCreationContext> Context =
 			ActiveEditorEntityManager->BatchCreateReservedEntities(Tables[TableHandleAsIndex], RowsToMassEntitiesConversion(ReservedHandles));
 
@@ -1729,6 +1754,8 @@ void UEditorDataStorage::FinalizePhase(UE::Editor::DataStorage::EQueryTickPhase
 
 void UEditorDataStorage::Reset()
 {
+	UE_LOGF(LogEditorDataStorage, Log, "Resetting TEDS Core.");
+
 	if (UMassEntityEditorSubsystem* Mass = UMassEntityEditorSubsystem::Get())
 	{
 		Mass->GetOnPostTickDelegate().Remove(OnPostMassTickHandle);
@@ -1749,6 +1776,8 @@ void UEditorDataStorage::Reset()
 
 	ActiveEditorPhaseManager.Reset();
 	ActiveEditorEntityManager.Reset();
+
+	UE_LOGF(LogEditorDataStorage, Log, "Reset TEDS Core.");
 }
 
 int32 UEditorDataStorage::GetTableChunkSize(FName TableName) const
@@ -1790,9 +1819,9 @@ FMassArchetypeHandle UEditorDataStorage::LookupArchetype(UE::Editor::DataStorage
 {
 	using namespace UE::Editor::DataStorage;
 
-	const int32 TableHandleAsIndex = Private::ConvertTableHandleToIndex(InTableHandle);
-	if (Tables.IsValidIndex(TableHandleAsIndex))
+	if (IsTableValid(InTableHandle))
 	{
+		const int32 TableHandleAsIndex = Private::ConvertTableHandleToIndex(InTableHandle);
 		return Tables[TableHandleAsIndex];
 	}
 	return FMassArchetypeHandle();
diff --git a/Engine/Plugins/Experimental/EditorDataStorage/Source/TedsCore/Private/TypedElementDatabase.h b/Engine/Plugins/Experimental/EditorDataStorage/Source/TedsCore/Private/TypedElementDatabase.h
--- a/Engine/Plugins/Experimental/EditorDataStorage/Source/TedsCore/Private/TypedElementDatabase.h
+++ b/Engine/Plugins/Experimental/EditorDataStorage/Source/TedsCore/Private/TypedElementDatabase.h
@@ -61,7 +61,7 @@ public:
 	using FactoryConstIterator = TFactoryIterator<const UEditorDataStorageFactory, const UEditorDataStorage>;
 
 public:
-	~UEditorDataStorage() override = default;
+	~UEditorDataStorage() override;
 	
 	UE_API void Initialize();
 	UE_API void PostInitialize();
@@ -93,6 +93,7 @@ public:
 	UE_API virtual UE::Editor::DataStorage::TableHandle RegisterTable(
 		UE::Editor::DataStorage::TableHandle SourceTable, TConstArrayView<const UScriptStruct*> ColumnList, const FName& Name) override;
 	UE_API virtual UE::Editor::DataStorage::TableHandle FindTable(const FName& Name) override;
+	UE_API virtual bool IsTableValid(UE::Editor::DataStorage::TableHandle Table) const override;
 
 	UE_API virtual UE::Editor::DataStorage::RowHandle ReserveRow() override;
 	UE_API virtual void BatchReserveRows(int32 Count, TFunctionRef<void(UE::Editor::DataStorage::RowHandle)> ReservationCallback) override;
diff --git a/Engine/Plugins/Experimental/EditorDataStorage/Source/TedsCore/Private/TypedElementDatabaseCompatibility.cpp b/Engine/Plugins/Experimental/EditorDataStorage/Source/TedsCore/Private/TypedElementDatabaseCompatibility.cpp
--- a/Engine/Plugins/Experimental/EditorDataStorage/Source/TedsCore/Private/TypedElementDatabaseCompatibility.cpp
+++ b/Engine/Plugins/Experimental/EditorDataStorage/Source/TedsCore/Private/TypedElementDatabaseCompatibility.cpp
@@ -2,6 +2,7 @@
 
 #include "TypedElementDatabaseCompatibility.h"
 
+#include <cinttypes>
 #include <utility>
 #include "Algo/Unique.h"
 #include "Async/UniqueLock.h"
@@ -24,8 +25,6 @@
 
 #include UE_INLINE_GENERATED_CPP_BY_NAME(TypedElementDatabaseCompatibility)
 
-DEFINE_LOG_CATEGORY_STATIC(LogTedsCompat, Log, All);
-
 namespace UE::Editor::DataStorage::Private
 {
 	bool bIntegrateWithGC = true;
@@ -98,10 +97,16 @@ namespace UE::Editor::DataStorage::Private
 
 thread_local UE::Editor::DataStorage::FObjectCompatibilityImmediateScope* UEditorDataStorageCompatibility::ImmediateModeTracker = nullptr;
 
+UEditorDataStorageCompatibility::~UEditorDataStorageCompatibility()
+{
+	UE_LOGF(LogEditorDataStorage, Log, "Destroyed TEDS Compatibility.");
+}
+
 void UEditorDataStorageCompatibility::Initialize(UEditorDataStorage* InStorage)
 {
 	using namespace UE::Editor::DataStorage;
 
+	UE_LOGF(LogEditorDataStorage, Log, "Initializing TEDS Compatibility.");
 	checkf(InStorage, TEXT("TEDS Compatibility is being initialized with an invalid storage target."));
 	
 	FScopedExclusiveLock Lock(EGlobalLockScope::Public);
@@ -120,6 +125,8 @@ void UEditorDataStorageCompatibility::Initialize(UEditorDataStorage* InStorage)
 	ObjectReinstancedDelegateHandle = FCoreUObjectDelegates::OnObjectsReinstanced.AddUObject(this, &UEditorDataStorageCompatibility::OnObjectReinstanced);
 	
 	PostGcUnreachableAnalysisHandle = FCoreUObjectDelegates::PostReachabilityAnalysis.AddUObject(this, &UEditorDataStorageCompatibility::OnPostGcUnreachableAnalysis);
+
+	UE_LOGF(LogEditorDataStorage, Log, "Initialized TEDS Compatibility.");
 }
 
 void UEditorDataStorageCompatibility::PostInitialize(UEditorDataStorage* InStorage)
@@ -130,6 +137,8 @@ void UEditorDataStorageCompatibility::Deinitialize()
 {
 	using namespace UE::Editor::DataStorage;
 
+	UE_LOGF(LogEditorDataStorage, Log, "De-initializing TEDS Compatibility.");
+
 	FScopedExclusiveLock Lock(EGlobalLockScope::Public);
 
 	FCoreUObjectDelegates::PostReachabilityAnalysis.Remove(PostGcUnreachableAnalysisHandle);
@@ -146,6 +155,17 @@ void UEditorDataStorageCompatibility::Deinitialize()
 
 	Environment.Reset();
 	Storage = nullptr;
+	
+	ObjectRegistrationFilters.Reset();
+	ObjectToRowDialiasers.Reset();
+	TypeToTableMap.Reset();
+	ObjectAddedCallbackList.Reset();
+	PreObjectRemovedCallbackList.Reset();
+
+	StandardUObjectTable = InvalidTableHandle;
+	StandardExternalObjectTable = InvalidTableHandle;
+
+	UE_LOGF(LogEditorDataStorage, Log, "De-initialized TEDS Compatibility.");
 }
 
 void UEditorDataStorageCompatibility::RegisterRegistrationFilter(ObjectRegistrationFilter Filter)
@@ -169,14 +189,19 @@ void UEditorDataStorageCompatibility::RegisterTypeTableAssociation(
 {
 	using namespace UE::Editor::DataStorage;
 	
-	if (Private::bUseCommandBuffer)
-	{
-		QueuedCommands.AddCommand(UE::Editor::DataStorage::FRegisterTypeTableAssociation{ .TypeInfo = TypeInfo, .Table = Table });
-	}
-	else
+	if (Storage)
 	{
-		FScopedExclusiveLock Lock(EGlobalLockScope::Public);
-		TypeToTableMap.Add(TypeInfo, Table);
+		checkf(Storage->IsTableValid(Table), TEXT("Attempting to associate invalid table %" PRIu64 " with object type %s"),
+			Table, TypeInfo.IsValid() ? *TypeInfo->GetName() : TEXT("<unknown>"));
+		if (Private::bUseCommandBuffer)
+		{
+			QueuedCommands.AddCommand(UE::Editor::DataStorage::FRegisterTypeTableAssociation{ .TypeInfo = TypeInfo, .Table = Table });
+		}
+		else
+		{
+			FScopedExclusiveLock Lock(EGlobalLockScope::Public);
+			TypeToTableMap.Add(TypeInfo, Table);
+		}
 	}
 }
 
@@ -1046,7 +1071,7 @@ void UEditorDataStorageCompatibility::TickPendingCommands()
 		{
 			TEDS_EVENT_SCOPE(TEXT("Compatibility Tick - Logging"))
 			FString CommandsAsString = FRecordCommands::PrintToString(PendingCommands, Private::PrintCompatCommandBuffer == 2);
-			UE_LOG(LogTedsCompat, Log, TEXT("Pending Commands:\n%s%u Nops"), 
+			UE_LOG(LogEditorDataStorage, Log, TEXT("Pending Commands:\n%s%u Nops"),
 				*CommandsAsString, PendingCommands.GetCommandCount<FNopCommand>());
 		}
 
@@ -1062,7 +1087,7 @@ void UEditorDataStorageCompatibility::TickPendingUObjectRegistration()
 
 	// Thread safe because it's only called from functions that already lock.
 	
-	if (!UObjectsPendingRegistration.IsEmpty())
+	if (Storage && !UObjectsPendingRegistration.IsEmpty())
 	{
 		UObjectsPendingRegistration.ProcessEntries(*Storage, *this,
 			[this](UE::Editor::DataStorage::RowHandle Row, const TWeakObjectPtr<UObject>& Object)
@@ -1078,7 +1103,7 @@ void UEditorDataStorageCompatibility::TickPendingExternalObjectRegistration()
 
 	// Thread safe because it's only called from functions that already lock.
 
-	if (!ExternalObjectsPendingRegistration.IsEmpty())
+	if (Storage && !ExternalObjectsPendingRegistration.IsEmpty())
 	{
 		ExternalObjectsPendingRegistration.ProcessEntries(*Storage, *this,
 			[this](UE::Editor::DataStorage::RowHandle Row, const ExternalObjectRegistration& Object)
diff --git a/Engine/Plugins/Experimental/EditorDataStorage/Source/TedsCore/Private/TypedElementDatabaseCompatibility.h b/Engine/Plugins/Experimental/EditorDataStorage/Source/TedsCore/Private/TypedElementDatabaseCompatibility.h
--- a/Engine/Plugins/Experimental/EditorDataStorage/Source/TedsCore/Private/TypedElementDatabaseCompatibility.h
+++ b/Engine/Plugins/Experimental/EditorDataStorage/Source/TedsCore/Private/TypedElementDatabaseCompatibility.h
@@ -41,7 +41,7 @@ class UEditorDataStorageCompatibility
 	friend struct UE::Editor::DataStorage::FPatchData;
 	friend struct UE::Editor::DataStorage::FPrepareCommands;
 public:
-	~UEditorDataStorageCompatibility() override = default;
+	~UEditorDataStorageCompatibility() override;
 
 	UE_API void Initialize(UEditorDataStorage* InStorage);
 	UE_API void PostInitialize(UEditorDataStorage* InStorage);
diff --git a/Engine/Source/Runtime/TypedElementFramework/Public/Elements/Interfaces/TypedElementDataStorageInterface.h b/Engine/Source/Runtime/TypedElementFramework/Public/Elements/Interfaces/TypedElementDataStorageInterface.h
--- a/Engine/Source/Runtime/TypedElementFramework/Public/Elements/Interfaces/TypedElementDataStorageInterface.h
+++ b/Engine/Source/Runtime/TypedElementFramework/Public/Elements/Interfaces/TypedElementDataStorageInterface.h
@@ -109,8 +109,11 @@ public:
 	template<TColumnType... Columns>
 	TableHandle RegisterTable(TableHandle SourceTable, const FName& Name);
 
-	/** Returns a previously created table with the provided name or TypedElementInvalidTableHandle if not found. */
+	/** Returns a previously created table with the provided name or InvalidTableHandle if not found. */
 	virtual TableHandle FindTable(const FName& Name) = 0;
+
+	/** Checks if the provided table handle is pointing to a valid table. A table handle of value `InvalidTableHandle` is considered invalid. */
+	virtual bool IsTableValid(TableHandle Table) const = 0;
 	
 	/**
 	 * @section Row management
