[TEDS] Fixed two bugs with TEDS Compatibility when processing pending registration.

1. During processing it's possible that observers react to a new entry and then register their own object. This in turn could invalidate the array of pending objects that's being iterated over and cause problems because they're now unsorted and pointing to the wrong memory address. The fix is to double buffer the queue and swap queues so new entries don't interfere with the ones actively being processed.
2. There was an off-by-one bug that would cause potentially one async loaded actor to be lost.

Additional checks were left in case the automated test failure isn't fixed by this change.

#Synced-CL 51558978

diff --git a/Engine/Plugins/Experimental/EditorDataStorage/Source/TedsCore/Private/TypedElementDatabaseCompatibility.cpp b/Engine/Plugins/Experimental/EditorDataStorage/Source/TedsCore/Private/TypedElementDatabaseCompatibility.cpp
--- a/Engine/Plugins/Experimental/EditorDataStorage/Source/TedsCore/Private/TypedElementDatabaseCompatibility.cpp
+++ b/Engine/Plugins/Experimental/EditorDataStorage/Source/TedsCore/Private/TypedElementDatabaseCompatibility.cpp
@@ -531,6 +531,7 @@ UE::Editor::DataStorage::TableHandle UEditorDataStorageCompatibility::FindBestMa
 	{
 		if (const TableHandle* Table = TypeToTableMap.Find(TypeInfo))
 		{
+			checkf(Storage && Storage->IsTableValid(*Table), TEXT("A previously registered table (%" PRIu64 ") has become invalid."), *Table);
 			return *Table;
 		}
 		TypeInfo = TypeInfo->GetSuperStruct();
@@ -863,8 +864,12 @@ void UEditorDataStorageCompatibility::PendingRegistration<AddressType>::ProcessE
 	// Thread-safe as it's only called from functions that already lock using.
 	using namespace UE::Editor::DataStorage;
 
+	// Double buffer because new entries can be added during processing of the queue. This can happen if there are observers that add
+	// new UObjects or external objects in response to one being registered.
+	Swap(Entries, ProcessingEntries);
+
 	// Start by removing any entries that are no longer valid.
-	for (auto It = Entries.CreateIterator(); It; ++It)
+	for (auto It = ProcessingEntries.CreateIterator(); It; ++It)
 	{
 		bool bIsValid = StorageInterface.IsRowAvailable(It->Row);
 		if constexpr (std::is_same_v<AddressType, TWeakObjectPtr<UObject>>)
@@ -897,10 +902,10 @@ void UEditorDataStorageCompatibility::PendingRegistration<AddressType>::ProcessE
 
 	// Check for empty here are the above code could potentially leave an empty array behind. This would result in break the assumption
 	// that there is at least one entry later in this function.
-	if (!Entries.IsEmpty())
+	if (!ProcessingEntries.IsEmpty())
 	{
 		// Next resolve the required table handles.
-		for (FEntry& Entry : Entries)
+		for (FEntry& Entry : ProcessingEntries)
 		{
 			if constexpr (std::is_same_v<AddressType, TWeakObjectPtr<UObject>>)
 			{
@@ -913,6 +918,9 @@ void UEditorDataStorageCompatibility::PendingRegistration<AddressType>::ProcessE
 							TEXT("The data storage could not find any matching tables for object of type '%s'. "
 							"This can mean that the object doesn't derive from UObject or that a table for UObject is no longer registered."), 
 							*Entry.Address->GetClass()->GetFName().ToString());
+						checkf(StorageInterface.IsTableValid(Entry.Table),
+							TEXT("Retrieved table (%" PRIu64 ") for object %s of type %s is invalid."), Entry.Table,
+							*Entry.Address->GetName(), *Entry.Address->GetClass()->GetFName().ToString());
 						break;
 					}
 					case Private::EAsyncUObjectLoadStrategy::DeferRegistration:
@@ -929,6 +937,9 @@ void UEditorDataStorageCompatibility::PendingRegistration<AddressType>::ProcessE
 								TEXT("The data storage could not find any matching tables for object of type '%s'. "
 								"This can mean that the object doesn't derive from UObject or that a table for UObject is no longer registered."), 
 								*Entry.Address->GetClass()->GetFName().ToString());
+							checkf(StorageInterface.IsTableValid(Entry.Table),
+								TEXT("Retrieved table (%" PRIu64 ") for object %s of type %s is invalid."), Entry.Table,
+								*Entry.Address->GetName(), *Entry.Address->GetClass()->GetFName().ToString());
 						}
 						break;
 					}
@@ -950,7 +961,7 @@ void UEditorDataStorageCompatibility::PendingRegistration<AddressType>::ProcessE
 		}
 
 		// Next sort them by table then by row handle to allow batch insertion.
-		Entries.Sort(
+		ProcessingEntries.Sort(
 			[](const FEntry& Lhs, const FEntry& Rhs)
 			{
 				if (Lhs.Table < Rhs.Table)
@@ -968,8 +979,8 @@ void UEditorDataStorageCompatibility::PendingRegistration<AddressType>::ProcessE
 			});
 
 		// Batch up the entries and add them to the storage.
-		FEntry* Current = Entries.GetData();
-		FEntry* End = Current + Entries.Num();
+		FEntry* Current = ProcessingEntries.GetData();
+		FEntry* End = Current + ProcessingEntries.Num();
 		
 		FEntry* TableFront = Current;
 
@@ -977,11 +988,15 @@ void UEditorDataStorageCompatibility::PendingRegistration<AddressType>::ProcessE
 		// This assertion ensures that we catch a change in this behavior should the constant
 		// be modified. 
 		static_assert(InvalidTableHandle == TNumericLimits<TableHandle>::Max()); 
-		if (Entries[0].Table == InvalidTableHandle)
+		if (ProcessingEntries[0].Table == InvalidTableHandle)
 		{
+			Entries.Append(ProcessingEntries);
+			ProcessingEntries.Reset();
 			return;
 		}
-		TableHandle CurrentTable = Entries[0].Table;
+		TableHandle CurrentTable = ProcessingEntries[0].Table;
+		checkf(CurrentTable == InvalidTableHandle || StorageInterface.IsTableValid(CurrentTable),
+			TEXT("The table (%" PRIu64 ") starting the processing queue has an unsupported value."), CurrentTable);
 
 		Compatibility.RowScratchBuffer.Reset();
 		
@@ -989,10 +1004,6 @@ void UEditorDataStorageCompatibility::PendingRegistration<AddressType>::ProcessE
 		{
 			if (Current->Table != CurrentTable)
 			{
-				if (CurrentTable == InvalidTableHandle)
-				{
-					break;
-				}
 				StorageInterface.BatchAddRow(CurrentTable, Compatibility.RowScratchBuffer,
 					[&SetupRowCallback, &TableFront](RowHandle Row)
 					{
@@ -1001,7 +1012,14 @@ void UEditorDataStorageCompatibility::PendingRegistration<AddressType>::ProcessE
 					});
 
 				CurrentTable = Current->Table;
+				checkf(CurrentTable == InvalidTableHandle || StorageInterface.IsTableValid(CurrentTable),
+					TEXT("The next table (%" PRIu64 ") in the processing queue is invalid."), CurrentTable);
 				Compatibility.RowScratchBuffer.Reset();
+
+				if (CurrentTable == InvalidTableHandle)
+				{
+					break;
+				}
 			}
 			Compatibility.RowScratchBuffer.Add(Current->Row);
 		}
@@ -1015,22 +1033,19 @@ void UEditorDataStorageCompatibility::PendingRegistration<AddressType>::ProcessE
 				});
 			// The last entry was able to be added to TEDS
 			// therefore none remain in the list.
-			Entries.Reset();
 		}
 		else
 		{
-			// There are some objects that couldn't be added this frame
-			// Return how many remain in the list, we will try again next frame
+			// There are some objects that couldn't be added this frame. Take that list
+			// and add it to the active entries list so they get picked up next frame.
 			const int32 RemainingEntries = static_cast<int32>(End - Current);
-			// Shuffle all the remaining entries to the front of the list and then
-			// truncate the list
-			const int32 ProcessedEntries = Entries.Num() - RemainingEntries;
-			for (int32 Index = 0; Index < RemainingEntries; ++Index)
+			Entries.Reserve(Entries.Num() + RemainingEntries);
+			for (; Current != End; ++Current)
 			{
-				Entries[Index] = Entries[ProcessedEntries + Index];
+				Entries.Add(*Current);
 			}
-			Entries.SetNum(RemainingEntries, EAllowShrinking::No);
 		}
+		ProcessingEntries.Reset();
 		Compatibility.RowScratchBuffer.Reset();
 	}
 }
@@ -1089,6 +1104,7 @@ void UEditorDataStorageCompatibility::TickPendingUObjectRegistration()
 	
 	if (Storage && !UObjectsPendingRegistration.IsEmpty())
 	{
+		checkf(Storage->IsTableValid(StandardUObjectTable), TEXT("Default UObject table has become invalid."));
 		UObjectsPendingRegistration.ProcessEntries(*Storage, *this,
 			[this](UE::Editor::DataStorage::RowHandle Row, const TWeakObjectPtr<UObject>& Object)
 			{
@@ -1105,6 +1121,7 @@ void UEditorDataStorageCompatibility::TickPendingExternalObjectRegistration()
 
 	if (Storage && !ExternalObjectsPendingRegistration.IsEmpty())
 	{
+		checkf(Storage->IsTableValid(StandardExternalObjectTable), TEXT("Default external object table has become invalid."));
 		ExternalObjectsPendingRegistration.ProcessEntries(*Storage, *this,
 			[this](UE::Editor::DataStorage::RowHandle Row, const ExternalObjectRegistration& Object)
 			{
diff --git a/Engine/Plugins/Experimental/EditorDataStorage/Source/TedsCore/Private/TypedElementDatabaseCompatibility.h b/Engine/Plugins/Experimental/EditorDataStorage/Source/TedsCore/Private/TypedElementDatabaseCompatibility.h
--- a/Engine/Plugins/Experimental/EditorDataStorage/Source/TedsCore/Private/TypedElementDatabaseCompatibility.h
+++ b/Engine/Plugins/Experimental/EditorDataStorage/Source/TedsCore/Private/TypedElementDatabaseCompatibility.h
@@ -178,6 +178,7 @@ private:
 			UE::Editor::DataStorage::TableHandle Table;
 		};
 		TArray<FEntry> Entries;
+		TArray<FEntry> ProcessingEntries;
 
 	public:
 		void Add(UE::Editor::DataStorage::RowHandle ReservedRowHandle, AddressType Address);
