[Bug Fix][TypedElementDatabaseIndexTable] RemapUnguarded not remapping the right row causing ensure

[Bug] RemapUnguarded was calling RemoveUnguarded which internally clears Rows[Index] = InvalidRowHandle before returning the slot index. Then it was passing that bare int32 Index as the RowHandle argument to IndexRowUnguarded. This caused an ensure because that row was often already assigned. This also caused the remap to create a new row for CB assets making a duplicate.

#Synced-CL 51672242

diff --git a/Engine/Plugins/Experimental/EditorDataStorage/Source/TedsCore/Private/TypedElementDatabaseIndexTable.cpp b/Engine/Plugins/Experimental/EditorDataStorage/Source/TedsCore/Private/TypedElementDatabaseIndexTable.cpp
--- a/Engine/Plugins/Experimental/EditorDataStorage/Source/TedsCore/Private/TypedElementDatabaseIndexTable.cpp
+++ b/Engine/Plugins/Experimental/EditorDataStorage/Source/TedsCore/Private/TypedElementDatabaseIndexTable.cpp
@@ -530,9 +530,14 @@ namespace UE::Editor::DataStorage
 
 	void FMappingTable::FDomain::RemapUnguarded(UEditorDataStorage& InDataStorage, const FMapKeyView& OriginalKey, FMapKey NewKey, bool bIsDirty)
 	{
-		if (int32 Index = RemoveUnguarded(OriginalKey); Index != InvalidId)
+		uint64 Hash = OriginalKey.CalculateHash();
+		if (int32 Index = IndexLookupMap.Remove(Hash, OriginalKey, Keys); Index != InvalidId)
 		{
-			IndexRowUnguarded(InDataStorage, MoveTemp(NewKey), Index, bIsDirty);
+			RowHandle StoredRow = Rows[Index];
+			Rows[Index] = InvalidRowHandle;
+			Keys[Index].Clear();
+			FreeList.PushFirst(Index);
+			IndexRowUnguarded(InDataStorage, MoveTemp(NewKey), StoredRow, bIsDirty);
 		}
 		else
 		{
