TEDS UI: Update widget purposes to use the new hierarchy API
TEDS Hierarchies: Add a function to resolve hierarchies immediately (instead of waiting for a processor)

#Synced-CL 51760981

diff --git a/Engine/Plugins/Experimental/EditorDataStorage/Source/TedsCore/Private/Hierarchy/EditorDataStorageHierarchyImplementation.cpp b/Engine/Plugins/Experimental/EditorDataStorage/Source/TedsCore/Private/Hierarchy/EditorDataStorageHierarchyImplementation.cpp
--- a/Engine/Plugins/Experimental/EditorDataStorage/Source/TedsCore/Private/Hierarchy/EditorDataStorageHierarchyImplementation.cpp
+++ b/Engine/Plugins/Experimental/EditorDataStorage/Source/TedsCore/Private/Hierarchy/EditorDataStorageHierarchyImplementation.cpp
@@ -53,6 +53,64 @@ namespace UE::Editor::DataStorage
 		}
 	}
 
+	void FTedsHierarchyRegistrar::ResolveHierarchy(ICoreProvider& CoreProvider, FHierarchyHandle HierarchyHandle)
+	{
+		using namespace UE::Editor::DataStorage::Queries;
+	
+		if (int32 Index = GetHandleIndex(HierarchyHandle); Index != INDEX_NONE)
+		{
+			if (FTedsHierarchyAccessInterface* AccessInterface = RegisteredHierarchies[Index].AccessInterface.Get())
+			{
+				// Lazily register this query when ResolveHierarchy is called for this hierarchy handle the first time
+				if (RegisteredHierarchies[Index].UnresolvedChildRowsQueryHandle == InvalidQueryHandle)
+				{
+					RegisteredHierarchies[Index].UnresolvedChildRowsQueryHandle = CoreProvider.RegisterQuery(
+						Select()
+						.Where()
+							.All(AccessInterface->GetUnresolvedParentColumnType())
+						.Compile());
+				}
+				
+				FRowHandleArray CollectedRows;
+				
+				CoreProvider.RunQuery(RegisteredHierarchies[Index].UnresolvedChildRowsQueryHandle, CreateDirectQueryCallbackBinding(
+					[&CollectedRows](IDirectQueryContext& Context, const RowHandle*)
+					{
+						CollectedRows.Append(Context.GetRowHandles());
+					}));
+				
+				// Iterate over each row and resolve the hierarchy if the parent is found now
+				for (RowHandle UnresolvedChildRow : CollectedRows.GetRows())
+				{
+					if (const void* UnresolvedParentColumnData = CoreProvider.GetColumnData(UnresolvedChildRow, AccessInterface->GetUnresolvedParentColumnType()))
+					{
+						const FEditorDataHierarchyUnresolvedParent_Template* UnresolvedParentColumn = 
+							static_cast<const FEditorDataHierarchyUnresolvedParent_Template*>(UnresolvedParentColumnData);
+						
+						RowHandle ParentRow = CoreProvider.LookupMappedRow(UnresolvedParentColumn->MappingDomain, UnresolvedParentColumn->ParentId);
+						
+						if (CoreProvider.IsRowAvailable(ParentRow))
+						{
+							AccessInterface->SetParentRow(&CoreProvider, UnresolvedChildRow, ParentRow);
+						}
+					}
+				}
+			}
+		}
+		
+	}
+
+	FTedsHierarchyRegistrar::~FTedsHierarchyRegistrar()
+	{
+		if (ICoreProvider* Storage = GetMutableDataStorageFeature<ICoreProvider>(StorageFeatureName))
+		{
+			for (int32 HierarchyIndex = 0; HierarchyIndex < RegisteredHierarchies.Num(); ++HierarchyIndex)
+			{
+				Storage->UnregisterQuery(RegisteredHierarchies[HierarchyIndex].UnresolvedChildRowsQueryHandle);
+			}
+		}
+	}
+
 	FHierarchyHandle FTedsHierarchyRegistrar::RegisterHierarchy(ICoreProvider* InProvider, const FHierarchyRegistrationParams& RegistrationParams)
 	{
 		// Create an access interface
@@ -149,9 +207,8 @@ namespace UE::Editor::DataStorage
 
 		FHierarchyHandle HierarchyHandle = CreateHandle(Index);
 		
-		ICoreProvider* DataStorage = GetMutableDataStorageFeature<ICoreProvider>(StorageFeatureName);
-		RegisterObservers(*DataStorage, HierarchyHandle);
-		RegisterProcessors(*DataStorage, HierarchyHandle);
+		RegisterObservers(*InProvider, HierarchyHandle);
+		RegisterProcessors(*InProvider, HierarchyHandle);
 		
 		return HierarchyHandle;
 	}
diff --git a/Engine/Plugins/Experimental/EditorDataStorage/Source/TedsCore/Private/Hierarchy/EditorDataStorageHierarchyImplementation.h b/Engine/Plugins/Experimental/EditorDataStorage/Source/TedsCore/Private/Hierarchy/EditorDataStorageHierarchyImplementation.h
--- a/Engine/Plugins/Experimental/EditorDataStorage/Source/TedsCore/Private/Hierarchy/EditorDataStorageHierarchyImplementation.h
+++ b/Engine/Plugins/Experimental/EditorDataStorage/Source/TedsCore/Private/Hierarchy/EditorDataStorageHierarchyImplementation.h
@@ -12,6 +12,8 @@ namespace UE::Editor::DataStorage
 	class FTedsHierarchyRegistrar final
 	{
 	public:
+		
+		~FTedsHierarchyRegistrar();
 		/**
 		 * Registers a hierarchy type with TEDS
 		 * This sets up the internal columns and processors to kee the bidirectional relationship in sync
@@ -31,6 +33,11 @@ namespace UE::Editor::DataStorage
 		 */
 		void ListHierarchyNames(TFunctionRef<void(const FName& HierarchyName)> Callback) const;
 		
+		/**
+		 * Try to resolve all child rows with unresolved parent columns
+		 */
+		void ResolveHierarchy(ICoreProvider& CoreProvider, FHierarchyHandle HierarchyHandle);
+		
 	private:
 		void RegisterObservers(ICoreProvider& CoreProvider, const FHierarchyHandle& Handle);
 		void RegisterProcessors(ICoreProvider& CoreProvider, const FHierarchyHandle& Handle);
@@ -44,6 +51,7 @@ namespace UE::Editor::DataStorage
 			const UScriptStruct* UnresolvedParentColumn;
 			const UScriptStruct* ParentChangedColumn;
 			TUniquePtr<FTedsHierarchyAccessInterface> AccessInterface;
+			QueryHandle UnresolvedChildRowsQueryHandle = InvalidQueryHandle;
 		};
 
 		const FRegisteredHierarchy* FindRegisteredHierarchy(const FHierarchyHandle& Handle);
diff --git a/Engine/Plugins/Experimental/EditorDataStorage/Source/TedsCore/Private/TypedElementDatabase.cpp b/Engine/Plugins/Experimental/EditorDataStorage/Source/TedsCore/Private/TypedElementDatabase.cpp
--- a/Engine/Plugins/Experimental/EditorDataStorage/Source/TedsCore/Private/TypedElementDatabase.cpp
+++ b/Engine/Plugins/Experimental/EditorDataStorage/Source/TedsCore/Private/TypedElementDatabase.cpp
@@ -1227,6 +1227,11 @@ const UScriptStruct* UEditorDataStorage::GetParentChangedColumnType(UE::Editor::
 	return nullptr;
 }
 
+void UEditorDataStorage::ResolveHierarchy(UE::Editor::DataStorage::FHierarchyHandle InHierarchyHandle)
+{
+	Environment->GetHierarchyRegistrar().ResolveHierarchy(*this, InHierarchyHandle);
+}
+
 void UEditorDataStorage::ListHierarchyNames(TFunctionRef<void(const FName& HierarchyName)> Callback) const
 {
 	Environment->GetHierarchyRegistrar().ListHierarchyNames(Callback);
diff --git a/Engine/Plugins/Experimental/EditorDataStorage/Source/TedsCore/Private/TypedElementDatabase.h b/Engine/Plugins/Experimental/EditorDataStorage/Source/TedsCore/Private/TypedElementDatabase.h
--- a/Engine/Plugins/Experimental/EditorDataStorage/Source/TedsCore/Private/TypedElementDatabase.h
+++ b/Engine/Plugins/Experimental/EditorDataStorage/Source/TedsCore/Private/TypedElementDatabase.h
@@ -155,6 +155,7 @@ public:
 		UE::Editor::DataStorage::FHierarchyHandle) const override;
 	UE_API virtual const UScriptStruct* GetParentChangedColumnType(
 		UE::Editor::DataStorage::FHierarchyHandle InHierarchyHandle) const override;
+	UE_API virtual void ResolveHierarchy(UE::Editor::DataStorage::FHierarchyHandle InHierarchyHandle) override;
 	UE_API virtual void ListHierarchyNames(TFunctionRef<void(const FName& HierarchyName)> Callback) const override;
 	UE_API virtual void SetParentRow(
 		UE::Editor::DataStorage::FHierarchyHandle, UE::Editor::DataStorage::RowHandle Target,	UE::Editor::DataStorage::RowHandle Parent) override;
diff --git a/Engine/Plugins/Experimental/EditorDataStorage/Source/TedsCore/Private/TypedElementDatabaseUI.cpp b/Engine/Plugins/Experimental/EditorDataStorage/Source/TedsCore/Private/TypedElementDatabaseUI.cpp
--- a/Engine/Plugins/Experimental/EditorDataStorage/Source/TedsCore/Private/TypedElementDatabaseUI.cpp
+++ b/Engine/Plugins/Experimental/EditorDataStorage/Source/TedsCore/Private/TypedElementDatabaseUI.cpp
@@ -119,6 +119,8 @@ namespace Internal
 	static const UE::Editor::DataStorage::IUiProvider::FPurposeID GeneralWidgetPurposeID(
 		UE::Editor::DataStorage::IUiProvider::FPurposeInfo("General", "Cell", NAME_None).GeneratePurposeID());
 	
+	static FName WidgetPurposesHierarchyName("WidgetPurposesHierarchy");
+	
 }
 
 void UEditorDataStorageUi::Initialize(
@@ -131,6 +133,7 @@ void UEditorDataStorageUi::Initialize(
 	StorageCompatibility = StorageCompatibilityInterface;
 	CreateStandardArchetypes();
 	RegisterQueries();
+	RegisterHierarchies();
 
 	Storage->OnUpdateCompleted().AddUObject(this, &UEditorDataStorageUi::PostDataStorageUpdate);
 }
@@ -141,40 +144,7 @@ void UEditorDataStorageUi::PostInitialize(UE::Editor::DataStorage::ICoreProvider
 	// Now that all the default widget purposes have been fully initialized, we can try to resolve the parenting information manually
 	// to ensure it is available for any tools that restored from the layout, which happens before TEDS has a chance to tick and resolve the
 	// information through a processor
-	using namespace UE::Editor::DataStorage::Queries;
-
-	FQueryDescription PurposesQueryDescription =
-		Select()
-		.Where(TColumn<FWidgetPurposeColumn>() && TColumn<FUnresolvedTableRowParentColumn>())
-		.Compile();
-	
-	QueryHandle PurposesQuery = StorageInterface->RegisterQuery(MoveTemp(PurposesQueryDescription));
-
-	FRowHandleArray Rows;
-	StorageInterface->RunQuery(PurposesQuery, CreateDirectQueryCallbackBinding(
-		[&Rows](IDirectQueryContext& Context, RowHandle Row)
-		{
-			Rows.Add(Row);
-		}));
-
-	for (RowHandle Row : Rows.GetRows())
-	{
-		FUnresolvedTableRowParentColumn* UnresolvedParent = Storage->GetColumn<FUnresolvedTableRowParentColumn>(Row);
-
-		if (ensureMsgf(UnresolvedParent, TEXT("Expected valid parent for row in query with TColumn<FUnresolvedTableRowParentColumn>()")))
-		{
-			RowHandle ParentRow = StorageInterface->LookupMappedRow(IUiProvider::PurposeMappingDomain, UnresolvedParent->ParentIdKey);
-			
-			if (Storage->IsRowAvailable(ParentRow))
-			{
-				StorageInterface->RemoveColumns<FUnresolvedTableRowParentColumn>(Row);
-				StorageInterface->AddColumn(Row, FTableRowParentColumn{ .Parent = ParentRow });
-			}
-		}
-	}
-
-	// Post Init, there is a processor run to resolve parent rows automatically so we don't need this query anymore
-	StorageInterface->UnregisterQuery(PurposesQuery);
+	Storage->ResolveHierarchy(WidgetPurposesHierarchyHandle);
 }
 
 void UEditorDataStorageUi::Deinitialize()
@@ -225,11 +195,11 @@ UE::Editor::DataStorage::RowHandle UEditorDataStorageUi::RegisterWidgetPurpose(c
 
 		if (Storage->IsRowAvailable(ParentRowHandle))
 		{
-			Storage->AddColumn(PurposeRowHandle, FTableRowParentColumn{ .Parent = ParentRowHandle });
+			Storage->SetParentRow(WidgetPurposesHierarchyHandle, PurposeRowHandle, ParentRowHandle);
 		}
 		else
 		{
-			Storage->AddColumn(PurposeRowHandle, FUnresolvedTableRowParentColumn{ .ParentIdKey = InPurposeInfo.ParentPurposeID});
+			Storage->SetUnresolvedParent(WidgetPurposesHierarchyHandle, PurposeRowHandle, UE::Editor::DataStorage::FMapKey(InPurposeInfo.ParentPurposeID), PurposeMappingDomain);
 		}
 	}
 	return PurposeRowHandle;
@@ -440,14 +410,7 @@ void UEditorDataStorageUi::CreateWidgetConstructors(UE::Editor::DataStorage::Row
 		// If no factories were found for this purpose, move on to the parent purpose
 		if (Factories.IsEmpty())
 		{
-			if (FTableRowParentColumn* ParentColumn = Storage->GetColumn<FTableRowParentColumn>(PurposeRow))
-			{
-				PurposeRow = ParentColumn->Parent;
-			}
-			else
-			{
-				PurposeRow = UE::Editor::DataStorage::InvalidRowHandle;
-			}
+			PurposeRow = Storage->GetParentRow(WidgetPurposesHierarchyHandle, PurposeRow);
 			continue;
 		}
 
@@ -539,14 +502,7 @@ void UEditorDataStorageUi::CreateWidgetConstructors(UE::Editor::DataStorage::Row
 		}
 
 		// If we have a parent purpose, try matching against factories belonging to it next
-		if (FTableRowParentColumn* ParentColumn = Storage->GetColumn<FTableRowParentColumn>(PurposeRow))
-		{
-			PurposeRow = ParentColumn->Parent;
-		}
-		else
-		{
-			PurposeRow = UE::Editor::DataStorage::InvalidRowHandle;
-		}
+		PurposeRow = Storage->GetParentRow(WidgetPurposesHierarchyHandle, PurposeRow);
 	}
 }
 
@@ -561,14 +517,7 @@ void UEditorDataStorageUi::ConstructWidgets(UE::Editor::DataStorage::RowHandle P
 		GetWidgetFactories(PurposeRow, Factories);
 		
 		// If no factories were found for this purpose, move on to the parent purpose
-		if (FTableRowParentColumn* ParentColumn = Storage->GetColumn<FTableRowParentColumn>(PurposeRow))
-		{
-			PurposeRow = ParentColumn->Parent;
-		}
-		else
-		{
-			PurposeRow = UE::Editor::DataStorage::InvalidRowHandle;
-		}
+		PurposeRow = Storage->GetParentRow(WidgetPurposesHierarchyHandle, PurposeRow);
 	}
 	
 	for (UE::Editor::DataStorage::RowHandle FactoryRow : Factories)
@@ -779,6 +728,11 @@ UE::Editor::DataStorage::IUiProvider::FPurposeID UEditorDataStorageUi::GetGenera
 	return Internal::GeneralWidgetPurposeID;
 }
 
+UE::Editor::DataStorage::FHierarchyHandle UEditorDataStorageUi::GetWidgetPurposeHierarchyHandle() const
+{
+	return WidgetPurposesHierarchyHandle;
+}
+
 UE::Editor::DataStorage::RowHandle UEditorDataStorageUi::FindPurpose(const FPurposeID& PurposeID) const
 {
 	return Storage->LookupMappedRow(PurposeMappingDomain, UE::Editor::DataStorage::FMapKey(PurposeID));
@@ -823,6 +777,18 @@ void UEditorDataStorageUi::RegisterQueries()
 		.Compile());
 }
 
+void UEditorDataStorageUi::RegisterHierarchies()
+{
+	using namespace UE::Editor::DataStorage;
+	
+	UE::Editor::DataStorage::FHierarchyRegistrationParams WidgetPurposeHierarchyParams
+		{
+			.Name = Internal::WidgetPurposesHierarchyName
+		};
+	
+	WidgetPurposesHierarchyHandle = Storage->RegisterHierarchy(WidgetPurposeHierarchyParams);
+}
+
 void UEditorDataStorageUi::PostDataStorageUpdate()
 {
 	// For any widgets that need their decorators updated, re-create decorators based on the columns on the widget row
diff --git a/Engine/Plugins/Experimental/EditorDataStorage/Source/TedsCore/Private/TypedElementDatabaseUI.h b/Engine/Plugins/Experimental/EditorDataStorage/Source/TedsCore/Private/TypedElementDatabaseUI.h
--- a/Engine/Plugins/Experimental/EditorDataStorage/Source/TedsCore/Private/TypedElementDatabaseUI.h
+++ b/Engine/Plugins/Experimental/EditorDataStorage/Source/TedsCore/Private/TypedElementDatabaseUI.h
@@ -114,12 +114,15 @@ public:
 	/** Get the ID of the general TEDS UI purpose used to register general purpose widgets for columns */
 	UE_API virtual FPurposeID GetGeneralWidgetPurposeID() const;
 	
+	UE_API virtual UE::Editor::DataStorage::FHierarchyHandle GetWidgetPurposeHierarchyHandle() const override;
+	
 	/** Find the row handle for a purpose by looking it up using the purpose ID */
 	UE_API virtual UE::Editor::DataStorage::RowHandle FindPurpose(const FPurposeID& PurposeID) const override;
 
 private:
 	UE_API void CreateStandardArchetypes();
 	UE_API void RegisterQueries();
+	UE_API void RegisterHierarchies();
 	void PostDataStorageUpdate();
 
 	using FWidgetConstructionMethod = TFunction<bool(UE::Editor::DataStorage::RowHandle, TArray<TWeakObjectPtr<const UScriptStruct>>)>;
@@ -195,6 +198,8 @@ private:
 
 	// Widget rows that need their decorators updated
 	TSet<UE::Editor::DataStorage::RowHandle> WidgetsToUpdateDecorators;
+	
+	UE::Editor::DataStorage::FHierarchyHandle WidgetPurposesHierarchyHandle;
 };
 
 #undef UE_API
diff --git a/Engine/Plugins/Experimental/EditorDataStorageFeatures/Source/TedsDebugger/Private/Widgets/QueryEditor/TedsQueryEditorResultsView.cpp b/Engine/Plugins/Experimental/EditorDataStorageFeatures/Source/TedsDebugger/Private/Widgets/QueryEditor/TedsQueryEditorResultsView.cpp
--- a/Engine/Plugins/Experimental/EditorDataStorageFeatures/Source/TedsDebugger/Private/Widgets/QueryEditor/TedsQueryEditorResultsView.cpp
+++ b/Engine/Plugins/Experimental/EditorDataStorageFeatures/Source/TedsDebugger/Private/Widgets/QueryEditor/TedsQueryEditorResultsView.cpp
@@ -148,13 +148,13 @@ namespace UE::Editor::DataStorage::Debug::QueryEditor
 				RowHandle DefaultPurposeRowHandle = UiProvider->FindPurpose(UiProvider->GetDefaultWidgetPurposeID());
 
 				// Temporarily add the default purpose as a parent of the general purpose so the debugger can support both
-				TedsInterface.AddColumn(GeneralPurposeRowHandle, FTableRowParentColumn{.Parent = DefaultPurposeRowHandle});
+				TedsInterface.SetParentRow(UiProvider->GetWidgetPurposeHierarchyHandle(), GeneralPurposeRowHandle, DefaultPurposeRowHandle);
 
 				// Update the columns in the table viewer using the selection types from the query description
 				TableViewer->SetColumns(TArray<TWeakObjectPtr<const UScriptStruct>>(TableViewerQueryDescription.SelectionTypes));
 
 				// Remove the parenting chain after we have used it to generate widgets
-				TedsInterface.RemoveColumns<FTableRowParentColumn>(GeneralPurposeRowHandle);
+				TedsInterface.SetParentRow(UiProvider->GetWidgetPurposeHierarchyHandle(), GeneralPurposeRowHandle, InvalidRowHandle);
 
 				// Since setting the columns clears all columns, we are going to re-add the custom column back
 				if(RowHandleColumn)
diff --git a/Engine/Source/Runtime/TypedElementFramework/Public/Elements/Interfaces/TypedElementDataStorageInterface.h b/Engine/Source/Runtime/TypedElementFramework/Public/Elements/Interfaces/TypedElementDataStorageInterface.h
--- a/Engine/Source/Runtime/TypedElementFramework/Public/Elements/Interfaces/TypedElementDataStorageInterface.h
+++ b/Engine/Source/Runtime/TypedElementFramework/Public/Elements/Interfaces/TypedElementDataStorageInterface.h
@@ -397,6 +397,11 @@ public:
 
 	// Gets the column type that is added to a row for a frame when the parent changes (if specified in RegisterHierarchy)
 	virtual const UScriptStruct* GetParentChangedColumnType(FHierarchyHandle InHierarchyHandle) const = 0;
+	
+	// Try to resolve every row that has an Unresolved Parent for the given hierarchy (@see SetUnresolvedParent)
+	// Unresolved hierarchies are evaluated by a processor every frame to see if the parent is known, but this function can be manually called
+	// to immediately try to resolve all rows for a particular hierarchy
+	virtual void ResolveHierarchy(FHierarchyHandle InHierarchyHandle) = 0;
 
 	// List all the currently known hierarchies by name
 	virtual void ListHierarchyNames(TFunctionRef<void(const FName& HierarchyName)> Callback) const = 0;
diff --git a/Engine/Source/Runtime/TypedElementFramework/Public/Elements/Interfaces/TypedElementDataStorageUiInterface.h b/Engine/Source/Runtime/TypedElementFramework/Public/Elements/Interfaces/TypedElementDataStorageUiInterface.h
--- a/Engine/Source/Runtime/TypedElementFramework/Public/Elements/Interfaces/TypedElementDataStorageUiInterface.h
+++ b/Engine/Source/Runtime/TypedElementFramework/Public/Elements/Interfaces/TypedElementDataStorageUiInterface.h
@@ -426,6 +426,9 @@ namespace UE::Editor::DataStorage
 
 		/** Get the name of the general TEDS UI purpose used to register general purpose widgets for columns */
 		virtual FPurposeID GetGeneralWidgetPurposeID() const = 0;
+		
+		/** Get the hierarchy handle that stores widget purposes relationships */
+		virtual FHierarchyHandle GetWidgetPurposeHierarchyHandle() const = 0;
 	
 		/**
 		 * @section Widget constructors.
