TEDS Hierarchies: Add the parent change column to a child row when the parent is destroyed

#Synced-CL 52265298

diff --git a/Engine/Plugins/Experimental/EditorDataStorage/Source/TedsCore/Private/Hierarchy/EditorDataStorageHierarchyImplementation.cpp b/Engine/Plugins/Experimental/EditorDataStorage/Source/TedsCore/Private/Hierarchy/EditorDataStorageHierarchyImplementation.cpp
--- a/Engine/Plugins/Experimental/EditorDataStorage/Source/TedsCore/Private/Hierarchy/EditorDataStorageHierarchyImplementation.cpp
+++ b/Engine/Plugins/Experimental/EditorDataStorage/Source/TedsCore/Private/Hierarchy/EditorDataStorageHierarchyImplementation.cpp
@@ -335,6 +335,10 @@ namespace UE::Editor::DataStorage
 					
 					Context.RemoveColumns(ChildRow, {Hierarchy->ChildTag});
 					
+					if (Hierarchy->ParentChangedColumn)
+					{
+						Context.AddColumns(ChildRow, {Hierarchy->ParentChangedColumn});
+					}
 				}
 
 				ParentHierarchyDataComponent->Children.Reset();
