Fix crash that occurs if passing unassigned rowhandle to Teds via SetParentRow

#Synced-CL 52243953

diff --git a/Engine/Plugins/Experimental/EditorDataStorage/Source/TedsCore/Private/Hierarchy/EditorDataStorageHierarchyImplementation.cpp b/Engine/Plugins/Experimental/EditorDataStorage/Source/TedsCore/Private/Hierarchy/EditorDataStorageHierarchyImplementation.cpp
--- a/Engine/Plugins/Experimental/EditorDataStorage/Source/TedsCore/Private/Hierarchy/EditorDataStorageHierarchyImplementation.cpp
+++ b/Engine/Plugins/Experimental/EditorDataStorage/Source/TedsCore/Private/Hierarchy/EditorDataStorageHierarchyImplementation.cpp
@@ -531,6 +531,12 @@ namespace UE::Editor::DataStorage
 
 	void FTedsHierarchyAccessInterface::SetParentRow(ICoreProvider* CoreProvider, RowHandle Target, RowHandle Parent) const
 	{
+		// The target is not assigned (ie. neither assigned to a table nor reserved)
+		// It will never be a valid row so appease TEDS philosophy of not crashing with invalid rows
+		if (!CoreProvider->IsRowAssigned(Target))
+		{
+			return;
+		}
 		CoreProvider->RemoveColumn(Target, UnresolvedParentColumnType);
 
 		auto AddHierarchyDataColumn = [this](ICoreProvider* CoreProvider, RowHandle Row, RowHandle ParentRow)
