[TEDS] Fix SetParentRow to use IsRowAvailable instead of IsRowAssigned

IsRowAvailable is more permissive than IsRowAssigned - it returns true for
both assigned and reserved rows. IsRowAssigned was added in CL 52284931 but
should have used IsRowAvailable to be consistent with the rest of the
hierarchy implementation and TEDS philosophy.

#Synced-CL 52287409

diff --git a/Engine/Plugins/Experimental/EditorDataStorage/Source/TedsCore/Private/Hierarchy/EditorDataStorageHierarchyImplementation.cpp b/Engine/Plugins/Experimental/EditorDataStorage/Source/TedsCore/Private/Hierarchy/EditorDataStorageHierarchyImplementation.cpp
--- a/Engine/Plugins/Experimental/EditorDataStorage/Source/TedsCore/Private/Hierarchy/EditorDataStorageHierarchyImplementation.cpp
+++ b/Engine/Plugins/Experimental/EditorDataStorage/Source/TedsCore/Private/Hierarchy/EditorDataStorageHierarchyImplementation.cpp
@@ -533,7 +533,7 @@ namespace UE::Editor::DataStorage
 	{
 		// The target is not assigned (ie. neither assigned to a table nor reserved)
 		// It will never be a valid row so appease TEDS philosophy of not crashing with invalid rows
-		if (!CoreProvider->IsRowAssigned(Target))
+		if (!CoreProvider->IsRowAvailable(Target))
 		{
 			return;
 		}
