TEDS: Add a tag to identify objects belonging to the PIE world
Hide PIE specific ULevels from the editor UI

#Synced-CL 52265298

diff --git a/Engine/Plugins/Experimental/EditorDataStorage/Source/TedsCore/Private/Compatibility/TypedElementUObjectWorldQueries.cpp b/Engine/Plugins/Experimental/EditorDataStorage/Source/TedsCore/Private/Compatibility/TypedElementUObjectWorldQueries.cpp
--- a/Engine/Plugins/Experimental/EditorDataStorage/Source/TedsCore/Private/Compatibility/TypedElementUObjectWorldQueries.cpp
+++ b/Engine/Plugins/Experimental/EditorDataStorage/Source/TedsCore/Private/Compatibility/TypedElementUObjectWorldQueries.cpp
@@ -13,6 +13,7 @@ void UObjectWorldDataStorageFactory::RegisterQueries(UE::Editor::DataStorage::IC
 {
 	RegisterAddWorldColumn(DataStorage);
 	RegisterUpdateOrRemoveWorldColumn(DataStorage);
+	RegisterPIEWorldSyncQueries(DataStorage);
 }
 
 void UObjectWorldDataStorageFactory::RegisterAddWorldColumn(UE::Editor::DataStorage::ICoreProvider& DataStorage) const
@@ -49,7 +50,7 @@ void UObjectWorldDataStorageFactory::RegisterUpdateOrRemoveWorldColumn(UE::Edito
 	DataStorage.RegisterQuery(
 		Select(
 			TEXT("Sync UObject's world to column"),
-			FProcessor(EQueryTickPhase::PostPhysics, DataStorage.GetQueryTickGroupName(EQueryTickGroups::SyncExternalToDataStorage))
+			FProcessor(EQueryTickPhase::DuringPhysics, DataStorage.GetQueryTickGroupName(EQueryTickGroups::SyncExternalToDataStorage))
 				.SetExecutionMode(EExecutionMode::GameThread),
 			[](IQueryContext& Context, RowHandle Row, const FTypedElementUObjectColumn& Object, FTypedElementWorldColumn& World)
 			{
@@ -69,3 +70,31 @@ void UObjectWorldDataStorageFactory::RegisterUpdateOrRemoveWorldColumn(UE::Edito
 			.All<FTypedElementSyncFromWorldTag>()
 		.Compile());
 }
+
+void UObjectWorldDataStorageFactory::RegisterPIEWorldSyncQueries(UE::Editor::DataStorage::ICoreProvider& DataStorage) const
+{
+	using namespace UE::Editor::DataStorage::Queries;
+	
+	DataStorage.RegisterQuery(
+		Select(
+			TEXT("Sync PIE Object Tag"),
+			FProcessor(EQueryTickPhase::PostPhysics, DataStorage.GetQueryTickGroupName(EQueryTickGroups::SyncExternalToDataStorage))
+				.SetExecutionMode(EExecutionMode::GameThread),
+			[](IQueryContext& Context, RowHandle Row, const FTypedElementWorldColumn& WorldColumn)
+			{
+				if (UWorld* World = WorldColumn.World.Get())
+				{
+					if (World->IsPlayInEditor())
+					{
+						Context.AddColumns<FPieObjectTag>(Row);
+					}
+					else
+					{
+						Context.RemoveColumns<FPieObjectTag>(Row);
+					}
+				}
+			})
+		.Where()
+			.All<FTypedElementSyncFromWorldTag>()
+		.Compile());
+}
diff --git a/Engine/Plugins/Experimental/EditorDataStorage/Source/TedsCore/Private/Compatibility/TypedElementUObjectWorldQueries.h b/Engine/Plugins/Experimental/EditorDataStorage/Source/TedsCore/Private/Compatibility/TypedElementUObjectWorldQueries.h
--- a/Engine/Plugins/Experimental/EditorDataStorage/Source/TedsCore/Private/Compatibility/TypedElementUObjectWorldQueries.h
+++ b/Engine/Plugins/Experimental/EditorDataStorage/Source/TedsCore/Private/Compatibility/TypedElementUObjectWorldQueries.h
@@ -28,4 +28,6 @@ private:
 	 * with the UObject anymore.
 	 */
 	void RegisterUpdateOrRemoveWorldColumn(UE::Editor::DataStorage::ICoreProvider& DataStorage) const;
+	
+	void RegisterPIEWorldSyncQueries(UE::Editor::DataStorage::ICoreProvider& DataStorage) const;
 };
diff --git a/Engine/Plugins/Experimental/EditorDataStorageFeatures/Source/TedsEditorCompatibility/Private/Level/TedsLevelFactory.cpp b/Engine/Plugins/Experimental/EditorDataStorageFeatures/Source/TedsEditorCompatibility/Private/Level/TedsLevelFactory.cpp
--- a/Engine/Plugins/Experimental/EditorDataStorageFeatures/Source/TedsEditorCompatibility/Private/Level/TedsLevelFactory.cpp
+++ b/Engine/Plugins/Experimental/EditorDataStorageFeatures/Source/TedsEditorCompatibility/Private/Level/TedsLevelFactory.cpp
@@ -106,6 +106,19 @@ void UTedsLevelFactory::RegisterQueries(UE::Editor::DataStorage::ICoreProvider&
 			.All<FLevelTag>()
 		.Compile());
 	
+	// There's no real need for an observer for FPIEObjectTag being removed right now since a PIE level cannot just become an editor level
+	DataStorage.RegisterQuery(
+		Select(TEXT("Hide PIE levels from UI"),
+			FObserver::OnAdd<FPieObjectTag>(),
+			[](IQueryContext& Context, RowHandle Row)
+			{
+				// PIE levels are not meant to be shown in the UI
+				Context.AddColumns<FHideRowFromUITag>(Row);
+			})
+		.Where()
+			.All<FLevelTag>()
+		.Compile());
+	
 	DataStorage.RegisterQuery(
 		Select(
 			TEXT("Sync Level -> World hierarchy"),
diff --git a/Engine/Source/Runtime/TypedElementFramework/Public/Elements/Columns/TypedElementCompatibilityColumns.h b/Engine/Source/Runtime/TypedElementFramework/Public/Elements/Columns/TypedElementCompatibilityColumns.h
--- a/Engine/Source/Runtime/TypedElementFramework/Public/Elements/Columns/TypedElementCompatibilityColumns.h
+++ b/Engine/Source/Runtime/TypedElementFramework/Public/Elements/Columns/TypedElementCompatibilityColumns.h
@@ -161,10 +161,20 @@ struct FEditorDataStorageUObjectIdNameColumn final : public FEditorDataStorageCo
 	FName IdName;
 };
 
+/**
+ * Tag on a row that to denote this is a PIE specific row (e.g PIE actors, PIE folders)
+ */
+USTRUCT(meta = (DisplayName = "PIE Object"))
+struct FEditorDataStoragePieObjectTag final : public FEditorDataStorageTag
+{
+	GENERATED_BODY()
+};
+
 namespace UE::Editor::DataStorage
 {
 	using FUObjectIdNameColumn = FEditorDataStorageUObjectIdNameColumn;
 	using FLevelColumn = FTypedElementLevelColumn;
 	using FLevelTag = FEditorDataStorageLevelTag;
 	using FWorldTag = FEditorDataStorageWorldTag;
+	using FPieObjectTag = FEditorDataStoragePieObjectTag;
 }  // namespace UE::Editor::DataStorage
\ No newline at end of file
