Fixed nightly build error.

#rnx
#Synced-CL 52378109

diff --git a/Engine/Plugins/Experimental/EditorDataStorage/Source/TedsCore/Private/Queries/QueryContextCapabilities/RowBatch.cpp b/Engine/Plugins/Experimental/EditorDataStorage/Source/TedsCore/Private/Queries/QueryContextCapabilities/RowBatch.cpp
--- a/Engine/Plugins/Experimental/EditorDataStorage/Source/TedsCore/Private/Queries/QueryContextCapabilities/RowBatch.cpp
+++ b/Engine/Plugins/Experimental/EditorDataStorage/Source/TedsCore/Private/Queries/QueryContextCapabilities/RowBatch.cpp
@@ -1,7 +1,5 @@
 // Copyright Epic Games, Inc. All Rights Reserved.
 
-#pragma once
-
 #include "Queries/QueryContextCapabilities/RowBatch.h"
 
 #include "Mass/EntityHandle.h"
diff --git a/Engine/Plugins/Experimental/EditorDataStorage/Source/TedsCore/Private/Queries/QueryContextCapabilities/SingleRow.cpp b/Engine/Plugins/Experimental/EditorDataStorage/Source/TedsCore/Private/Queries/QueryContextCapabilities/SingleRow.cpp
--- a/Engine/Plugins/Experimental/EditorDataStorage/Source/TedsCore/Private/Queries/QueryContextCapabilities/SingleRow.cpp
+++ b/Engine/Plugins/Experimental/EditorDataStorage/Source/TedsCore/Private/Queries/QueryContextCapabilities/SingleRow.cpp
@@ -1,7 +1,5 @@
 // Copyright Epic Games, Inc. All Rights Reserved.
 
-#pragma once
-
 #include "Queries/QueryContextCapabilities/SingleRow.h"
 
 #include "Mass/EntityHandle.h"
