[TEDS] Query and table updates

Larger changes introduced with this CR:
- Tables updates:
-- Tables have been moved to the query capabilities. This allows more table functionality to be available from query callbacks.
-- "Foreign keys" have been introduced which are small query callbacks that can be registered with a table and can return a unique identifier for a row such as a pointer or a path.
-- Table relations are fully tracked now, so any derived table or variant is now known and gets its own table handle. This also allow for more info to be retrieved about tables.
-- Tables get a unique tag assigned automatically based on their name. It's now no longer needed to add specialized tag to identify a row such as FEntityTag or FActorTag.
- Const query callbacks: it's not possible to create query callbacks that are fully const. These do not allow any changes to columns or through the context. Filters have been updated to use these and they're used by foreign keys.
- Query callbacks with results now get a row handle. Some results need to track the exact row the results applies to, which this change enables.

#Synced-CL 52445887

diff --git a/Engine/Plugins/ChaosVD/Source/ChaosVD/Private/TEDS/ChaosVDParticleEditorDataFactory.cpp b/Engine/Plugins/ChaosVD/Source/ChaosVD/Private/TEDS/ChaosVDParticleEditorDataFactory.cpp
--- a/Engine/Plugins/ChaosVD/Source/ChaosVD/Private/TEDS/ChaosVDParticleEditorDataFactory.cpp
+++ b/Engine/Plugins/ChaosVD/Source/ChaosVD/Private/TEDS/ChaosVDParticleEditorDataFactory.cpp
@@ -40,13 +40,18 @@ void UChaosVDParticleEditorDataFactory::RegisterTables(UE::Editor::DataStorage::
 
 	using namespace UE::Editor::DataStorage;
 
-	TableHandle ObjectTable = DataStorage.RegisterTable(DataStorage.FindTable(FName("Editor_StandardExternalObjectTable")),
-		TTypedElementColumnTypeList<FChaosVDObjectDataTag, FTypedElementFromCVDWorldTag, FTypedElementLabelColumn, FTypedElementLabelHashColumn, FVisibleInEditorColumn, FChaosVDTableRowParentColumn, FChaosVDActiveObjectTag>{},
-		FName("CVD_SceneObjectDataTable"));
+	TableHandle ObjectTable = DataStorage.RegisterTable
+		<
+			FChaosVDObjectDataTag, FTypedElementFromCVDWorldTag, FTypedElementLabelColumn, FTypedElementLabelHashColumn, 
+			FVisibleInEditorColumn, FChaosVDTableRowParentColumn, FChaosVDActiveObjectTag
+		>(
+			FName("CVD_SceneObjectDataTable"),
+			FTableRegistrationOptions{ .SourceTable = DataStorage.FindTable(FName("Editor_StandardExternalObjectTable")) });
 
-	TableHandle ActorTable = DataStorage.RegisterTable(DataStorage.FindTable(FName("Editor_StandardActorTable")),
-	TTypedElementColumnTypeList<FChaosVDObjectDataTag, FTypedElementFromCVDWorldTag, FVisibleInEditorColumn, FChaosVDTableRowParentColumn, FChaosVDActiveObjectTag>{},
-	FName("CVD_ActorDataTable"));
+	TableHandle ActorTable = DataStorage.RegisterTable
+		<FChaosVDObjectDataTag, FTypedElementFromCVDWorldTag, FVisibleInEditorColumn, FChaosVDTableRowParentColumn, FChaosVDActiveObjectTag>(
+		FName("CVD_ActorDataTable"),
+		FTableRegistrationOptions{ .SourceTable = DataStorage.FindTable(FName("Editor_StandardActorTable")) });
 
 	DataStorageCompatibility.RegisterTypeTableAssociation(FChaosVDBaseSceneObject::StaticStruct(), ObjectTable);
 	DataStorageCompatibility.RegisterTypeTableAssociation(AChaosVDSolverInfoActor::StaticClass(), ActorTable);
diff --git a/Engine/Plugins/Experimental/EditorDataStorage/Source/TedsCore/Private/Memento/TypedElementMementoSystem.cpp b/Engine/Plugins/Experimental/EditorDataStorage/Source/TedsCore/Private/Memento/TypedElementMementoSystem.cpp
--- a/Engine/Plugins/Experimental/EditorDataStorage/Source/TedsCore/Private/Memento/TypedElementMementoSystem.cpp
+++ b/Engine/Plugins/Experimental/EditorDataStorage/Source/TedsCore/Private/Memento/TypedElementMementoSystem.cpp
@@ -11,13 +11,13 @@
 
 namespace UE::Editor::DataStorage
 {
-	FMementoSystem::FMementoSystem(ICoreProvider& InDataStorage)
+	FMementoSystem::FMementoSystem(ICoreProvider& InDataStorage, FTableManager& TableManager)
 		: DataStorage(InDataStorage)
 	{
 		FScopedExclusiveLock Lock(EGlobalLockScope::Public);
 
 		// Register tables that will be used by reinstancing
-		MementoRowBaseTable = DataStorage.RegisterTable<FTypedElementMementoTag>(TEXT("MementoRowBaseTable"));
+		MementoRowBaseTable = TableManager.Register({ FTypedElementMementoTag::StaticStruct() }, TEXT("MementoRowBaseTable"));
 
 		// Discover all MementoTranslators
 		{
diff --git a/Engine/Plugins/Experimental/EditorDataStorage/Source/TedsCore/Private/Memento/TypedElementMementoSystem.h b/Engine/Plugins/Experimental/EditorDataStorage/Source/TedsCore/Private/Memento/TypedElementMementoSystem.h
--- a/Engine/Plugins/Experimental/EditorDataStorage/Source/TedsCore/Private/Memento/TypedElementMementoSystem.h
+++ b/Engine/Plugins/Experimental/EditorDataStorage/Source/TedsCore/Private/Memento/TypedElementMementoSystem.h
@@ -10,11 +10,12 @@ class UTedsMementoTranslatorBase;
 namespace UE::Editor::DataStorage
 {
 	class ICoreProvider;
+	class FTableManager;
 
 	class FMementoSystem
 	{
 	public:
-		explicit FMementoSystem(ICoreProvider& InDataStorage);
+		FMementoSystem(ICoreProvider& InDataStorage, FTableManager& TableManager);
 
 		RowHandle CreateMemento(RowHandle SourceRow);
 		void CreateMemento(RowHandle ReservedMementoRow, RowHandle SourceRow);
diff --git a/Engine/Plugins/Experimental/EditorDataStorage/Source/TedsCore/Private/Processors/TypedElementProcessorAdaptors.cpp b/Engine/Plugins/Experimental/EditorDataStorage/Source/TedsCore/Private/Processors/TypedElementProcessorAdaptors.cpp
--- a/Engine/Plugins/Experimental/EditorDataStorage/Source/TedsCore/Private/Processors/TypedElementProcessorAdaptors.cpp
+++ b/Engine/Plugins/Experimental/EditorDataStorage/Source/TedsCore/Private/Processors/TypedElementProcessorAdaptors.cpp
@@ -6,6 +6,7 @@
 #include "MassEntityTypes.h"
 #include "MassEntityView.h"
 #include "MassExecutionContext.h"
+#include "Misc/EnumClassFlags.h"
 #include "Queries/TypedElementExtendedQueryStore.h"
 #include "Queries/TypedElementQueryContext_MassForwarder.h"
 #include "Processors/TypedElementQueryCapabilityImplemenationForwarder.h"
@@ -893,7 +894,7 @@ namespace UE::Editor::DataStorage
 			RowHandle AddRow(TableHandle Table) 
 			{
 				FMassEntityHandle EntityHandle = Environment.GetMassEntityManager().ReserveEntity();
-				FMassArchetypeHandle ArchetypeHandle = Environment.LookupMassArchetype(Table);
+				FMassArchetypeHandle ArchetypeHandle = Environment.GetTableManager().GetArchetype(Table);
 
 				if (!ArchetypeHandle.IsValid())
 				{
@@ -1081,10 +1082,15 @@ namespace UE::Editor::DataStorage
 			const FTedsHierarchyAccessInterface* HierarchyAccessInterface = nullptr;
 		};
 
-		struct FMassDirectContextForwarder final : Queries::Private::TImplementationForwarder<Queries::FProcessorContextEnvironment, IDirectQueryContext, Queries::FDynamicColumnInfoContextCapability>
+		using FMassDirectContextForwarderBase = Queries::Private::TImplementationForwarder<
+			Queries::FProcessorContextEnvironment,
+			IDirectQueryContext,
+			Queries::FTableInfoContextCapability,
+			Queries::FDynamicColumnInfoContextCapability>;
+		struct FMassDirectContextForwarder final : FMassDirectContextForwarderBase
 		{
 			FMassDirectContextForwarder(FMassExecutionContext& InContext, FEnvironment& InEnvironment, FQueryDescription& Description)
-				: TImplementationForwarder(InEnvironment, InContext)
+				: FMassDirectContextForwarderBase(InEnvironment, InContext)
 				, Implementation(InContext, InEnvironment)
 			{
 				FHierarchyHandle Handle = InEnvironment.GetHierarchyRegistrar().FindHierarchyByName(Description.Hierarchy);
@@ -1120,10 +1126,15 @@ namespace UE::Editor::DataStorage
 			const FTedsHierarchyAccessInterface* HierarchyAccessInterface = nullptr;
 		};
 
-		struct FMassSubqueryContextForwarder : Queries::Private::TImplementationForwarder<Queries::FProcessorContextEnvironment, ISubqueryContext, Queries::FDynamicColumnInfoContextCapability>
+		using FMassSubqueryContextForwarderBase = Queries::Private::TImplementationForwarder<
+			Queries::FProcessorContextEnvironment,
+			ISubqueryContext,
+			Queries::FTableInfoContextCapability,
+			Queries::FDynamicColumnInfoContextCapability>;
+		struct FMassSubqueryContextForwarder : FMassSubqueryContextForwarderBase
 		{
 			FMassSubqueryContextForwarder(FMassExecutionContext& InContext, FEnvironment& InEnvironment, FQueryDescription& Description)
-				: TImplementationForwarder(InEnvironment, InContext)
+				: FMassSubqueryContextForwarderBase(InEnvironment, InContext)
 				, Implementation(InContext, InEnvironment, Description.bShouldBatchModifications)
 			{
 				FHierarchyHandle Handle = InEnvironment.GetHierarchyRegistrar().FindHierarchyByName(Description.Hierarchy);
@@ -1277,14 +1288,19 @@ namespace UE::Editor::DataStorage
 			FExtendedQueryStore& QueryStore;
 		};
 
-		struct FMassContextForwarder final : Queries::Private::TImplementationForwarder<Queries::FProcessorContextEnvironment, IQueryContext, Queries::FDynamicColumnInfoContextCapability>
+		using FMassContextForwarderBase = Queries::Private::TImplementationForwarder<
+			Queries::FProcessorContextEnvironment,
+			IQueryContext,
+			Queries::FTableInfoContextCapability,
+			Queries::FDynamicColumnInfoContextCapability>;
+		struct FMassContextForwarder final : FMassContextForwarderBase
 		{
 			FMassContextForwarder(
 				FQueryDescription& InQueryDescription,
-				FMassExecutionContext& InContext,
+				FMassExecutionContext& InContext, 
 				FExtendedQueryStore& InQueryStore,
 				FEnvironment& InEnvironment)
-					: TImplementationForwarder(InEnvironment, InContext)
+					: FMassContextForwarderBase(InEnvironment, InContext)
 					, Implementation(InQueryDescription, InContext, InQueryStore, InEnvironment)
 			{
 			}
@@ -1621,6 +1637,13 @@ UE::Editor::DataStorage::FQueryResult FTypedElementQueryProcessorData::Execute(
 	FQueryResult Result;
 	Result.Completed = FQueryResult::ECompletion::Fully;
 
+	// Check compatibility here once instead of on the call itself to avoid doing the same thing for every chunk.
+	if (!QueryContext_ProcessorApi::SupportsCapabilities(Callback.CapabilityMask))
+	{
+		QueryContext_ProcessorApi::StaticReportCompatibility(Callback.CapabilityMask);
+		return FQueryResult{ .Completed = FQueryResult::ECompletion::Unsupported };
+	}	
+
 	if (EnumHasAnyFlags(ExecutionFlags, ERunQueryFlags::AllowBoundQueries) || !Description.Callback.Function)
 	{
 		FMassExecutionContext Context(EntityManager);
@@ -1629,11 +1652,65 @@ UE::Editor::DataStorage::FQueryResult FTypedElementQueryProcessorData::Execute(
 		auto ExecuteFunction = [&Result, &Callback, &Environment](FMassExecutionContext& Context)
 			{
 				using namespace Processors::Private;
-				// No need to cache any subsystem dependencies as these are not accessible from a direct query.
-				FProcessorContextEnvironment Env(Environment, Context);
-				QueryContext_ProcessorApi QueryContext(Env);
+				FProcessorContextEnvironment MassContextEnvironment(Environment, Context);
+				QueryContext_ProcessorApi QueryContext(MassContextEnvironment);
+				// This assumes that there already has been a check that the provided function is compatible with the query.
+				Callback.Call(QueryContext, MassContextEnvironment);
+				Result.Count += Context.GetNumEntities();
+			};
+		if (EnumHasAnyFlags(ExecutionFlags, ERunQueryFlags::ParallelizeChunks))
+		{
+			FMassEntityQuery::EParallelExecutionFlags Mode =
+				EnumHasAnyFlags(ExecutionFlags, ERunQueryFlags::AutoBalanceParallelChunkProcessing)
+				? FMassEntityQuery::EParallelExecutionFlags::AutoBalance
+				: FMassEntityQuery::EParallelExecutionFlags::Default;
+			NativeQuery.ParallelForEachEntityChunk(Context, ExecuteFunction, Mode);
+		}
+		else
+		{
+			NativeQuery.ForEachEntityChunk(Context, ExecuteFunction);
+		}
+	}
+	else
+	{
+		Result.Completed = FQueryResult::ECompletion::Unsupported;
+	}
+	return Result;
+}
+
+UE::Editor::DataStorage::FQueryResult FTypedElementQueryProcessorData::Execute(
+	const UE::Editor::DataStorage::Queries::TConstQueryFunction<void>& Callback,
+	const UE::Editor::DataStorage::FQueryDescription& Description,
+	FMassEntityQuery& NativeQuery,
+	FMassEntityManager& EntityManager,
+	UE::Editor::DataStorage::FEnvironment& Environment,
+	UE::Editor::DataStorage::ERunQueryFlags ExecutionFlags)
+{
+	using namespace UE::Editor::DataStorage;
+	using namespace UE::Editor::DataStorage::Queries;
+
+	// Check compatibility here once instead of on the call itself to avoid doing the same thing for every chunk.
+	if (!QueryContext_ProcessorApi::SupportsCapabilities(Callback.CapabilityMask))
+	{
+		QueryContext_ProcessorApi::StaticReportCompatibility(Callback.CapabilityMask);
+		return FQueryResult{ .Completed = FQueryResult::ECompletion::Unsupported };
+	}
+
+	FQueryResult Result;
+	Result.Completed = FQueryResult::ECompletion::Fully;
+
+	if (EnumHasAnyFlags(ExecutionFlags, ERunQueryFlags::AllowBoundQueries) || !Description.Callback.Function)
+	{
+		FMassExecutionContext Context(EntityManager);
+		Context.SetDeferredCommandBuffer(Environment.GetSharedMassCommandBuffer());
+		Context.SetFlushDeferredCommands(false);
+		auto ExecuteFunction = [&Result, &Callback, &Environment](FMassExecutionContext& Context)
+			{
+				using namespace Processors::Private;
+				FConstProcessorContextEnvironment MassContextEnvironment(Environment, Context);
+				QueryContext_ConstProcessorApi QueryContext(MassContextEnvironment);
 				// This assumes that there already has been a check that the provided function is compatible with the query.
-				Callback.Call(QueryContext, Env);
+				Callback.Call(QueryContext, MassContextEnvironment);
 				Result.Count += Context.GetNumEntities();
 			};
 		if (EnumHasAnyFlags(ExecutionFlags, ERunQueryFlags::ParallelizeChunks))
diff --git a/Engine/Plugins/Experimental/EditorDataStorage/Source/TedsCore/Private/Processors/TypedElementProcessorAdaptors.h b/Engine/Plugins/Experimental/EditorDataStorage/Source/TedsCore/Private/Processors/TypedElementProcessorAdaptors.h
--- a/Engine/Plugins/Experimental/EditorDataStorage/Source/TedsCore/Private/Processors/TypedElementProcessorAdaptors.h
+++ b/Engine/Plugins/Experimental/EditorDataStorage/Source/TedsCore/Private/Processors/TypedElementProcessorAdaptors.h
@@ -67,6 +67,13 @@ struct FTypedElementQueryProcessorData
 		FMassEntityManager& EntityManager,
 		UE::Editor::DataStorage::FEnvironment& Environment,
 		UE::Editor::DataStorage::ERunQueryFlags ExecutionFlags);
+	static UE::Editor::DataStorage::FQueryResult Execute(
+		const UE::Editor::DataStorage::Queries::TConstQueryFunction<void>& Callback,
+		const UE::Editor::DataStorage::FQueryDescription& Description,
+		FMassEntityQuery& NativeQuery,
+		FMassEntityManager& EntityManager,
+		UE::Editor::DataStorage::FEnvironment& Environment,
+		UE::Editor::DataStorage::ERunQueryFlags ExecutionFlags);
 	static UE::Editor::DataStorage::FQueryResult Execute(
 		UE::Editor::DataStorage::SubqueryCallbackRef& Callback,
 		UE::Editor::DataStorage::FQueryDescription& Description,
diff --git a/Engine/Plugins/Experimental/EditorDataStorage/Source/TedsCore/Private/Queries/QueryContextCapabilities/RowBatch.cpp b/Engine/Plugins/Experimental/EditorDataStorage/Source/TedsCore/Private/Queries/QueryContextCapabilities/RowBatch.cpp
--- a/Engine/Plugins/Experimental/EditorDataStorage/Source/TedsCore/Private/Queries/QueryContextCapabilities/RowBatch.cpp
+++ b/Engine/Plugins/Experimental/EditorDataStorage/Source/TedsCore/Private/Queries/QueryContextCapabilities/RowBatch.cpp
@@ -2,7 +2,7 @@
 
 #include "Queries/QueryContextCapabilities/RowBatch.h"
 
-#include "Mass/EntityHandle.h"
+#include "Queries/QueryContextCapabilities/TableCapabilities.h"
 #include "Queries/TypedElementQueryContext_DirectSingleApi.h"
 #include "Queries/TypedElementQueryContext_MassForwarder.h"
 #include "TypedElementDatabase.h"
@@ -28,22 +28,7 @@ namespace UE::Editor::DataStorage::Queries
 	bool FRowBatchInfoContextCapability::CurrentBatchTableHasColumns(
 		const IContextEnvironment& Environment, TConstArrayView<const UScriptStruct*> Columns)
 	{
-		// Pick the first row since all rows in a batch are in the same table.
-		if (FRowHandleArrayView Rows = Environment.GetBatchRows(); !Rows.IsEmpty())
-		{
-			FMassEntityHandle Entity = FMassEntityHandle::FromNumber(Rows[0]);
-			FMassArchetypeHandle Archetype = Environment.GetEnvironment()->GetMassEntityManager().GetArchetypeForEntity(Entity);
-			const FMassElementBitSet& ArchetypeBits = Archetype.GetCompositionBitSetChecked();
-			for (const UScriptStruct* Column : Columns)
-			{
-				if (!ArchetypeBits.Contains(Column))
-				{
-					return false;
-				}
-			}
-			return true;
-		}
-		return false;
+		return FCurrentTableInfoContextCapability::CurrentTableHasColumns(Environment, Columns);
 	}
 
 
@@ -56,15 +41,7 @@ namespace UE::Editor::DataStorage::Queries
 	bool FRowBatchInfoContextCapability::CurrentBatchTableHasColumns(
 		const FProcessorContextEnvironment& Environment, TConstArrayView<const UScriptStruct*> Columns)
 	{
-		const FMassExecutionContext& Context = Environment.MassContext;
-		for (const UScriptStruct* Column : Columns)
-		{
-			if (!Context.DoesArchetypeHaveElement(Column))
-			{
-				return false;
-			}
-		}
-		return true;
+		return FCurrentTableInfoContextCapability::CurrentTableHasColumns(Environment, Columns);
 	}
 
 	const void* FRowBatchInfoContextCapability::GetColumnBatchAddress(
@@ -83,6 +60,26 @@ namespace UE::Editor::DataStorage::Queries
 		return Environment.MassContext.GetMutableFragmentView(ColumnType).GetData();
 	}
 
+
+
+	//
+	// FConstProcessorContextEnvironment
+	// 
+
+	bool FRowBatchInfoContextCapability::CurrentBatchTableHasColumns(
+		const FConstProcessorContextEnvironment& Environment, TConstArrayView<const UScriptStruct*> Columns)
+	{
+		return FCurrentTableInfoContextCapability::CurrentTableHasColumns(Environment, Columns);
+	}
+
+	const void* FRowBatchInfoContextCapability::GetColumnBatchAddress(
+		const FConstProcessorContextEnvironment& Environment, const UScriptStruct* ColumnType)
+	{
+		checkf(UE::Mass::IsA<FMassFragment>(ColumnType), TEXT("Trying to retrieve a data pointer for column '%s' which doesn't have data."),
+			ColumnType ? *ColumnType->GetName() : TEXT("<null>"));
+		return Environment.MassContext.GetFragmentView(ColumnType).GetData();
+	}
+
 	
 	
 	//
@@ -107,4 +104,18 @@ namespace UE::Editor::DataStorage::Queries
 			? Environment.DirectContext->GetColumnData(*Environment.Rows.GetData(), ColumnType)
 			: nullptr;
 	}
+
+
+	//
+	// FConstDirectContextEnvironment
+	//
+
+	const void* FRowBatchInfoContextCapability::GetColumnBatchAddress(const FConstDirectContextEnvironment& Environment, const UScriptStruct* ColumnType)
+	{
+		checkf(UE::Mass::IsA<FMassFragment>(ColumnType), TEXT("Trying to retrieve a data pointer for column '%s' which doesn't have data."),
+			ColumnType ? *ColumnType->GetName() : TEXT("<null>"));
+		return !Environment.Rows.IsEmpty()
+			? Environment.DirectContext->GetColumnData(*Environment.Rows.GetData(), ColumnType)
+			: nullptr;
+	}
 } // namespace UE::Editor::DataStorage::Queries
diff --git a/Engine/Plugins/Experimental/EditorDataStorage/Source/TedsCore/Private/Queries/QueryContextCapabilities/RowBatch.h b/Engine/Plugins/Experimental/EditorDataStorage/Source/TedsCore/Private/Queries/QueryContextCapabilities/RowBatch.h
--- a/Engine/Plugins/Experimental/EditorDataStorage/Source/TedsCore/Private/Queries/QueryContextCapabilities/RowBatch.h
+++ b/Engine/Plugins/Experimental/EditorDataStorage/Source/TedsCore/Private/Queries/QueryContextCapabilities/RowBatch.h
@@ -14,7 +14,6 @@ namespace UE::Editor::DataStorage::Queries
 	struct FDirectContextEnvironment;
 	struct FProcessorContextEnvironment;
 
-	// For documentation on these functions see "TypedElementQueryCapabilities.inl"
 	class FRowBatchInfoContextCapability : public ImplementsContextCapability<RowBatchInfo>
 	{
 	public:
@@ -26,7 +25,12 @@ namespace UE::Editor::DataStorage::Queries
 		static const void* GetColumnBatchAddress(const FProcessorContextEnvironment& Environment, const UScriptStruct* ColumnType);
 		static void* GetMutableColumnBatchAddress(FProcessorContextEnvironment& Environment, const UScriptStruct* ColumnType);
 
+		static bool CurrentBatchTableHasColumns(const FConstProcessorContextEnvironment& Environment, TConstArrayView<const UScriptStruct*> Columns);
+		static const void* GetColumnBatchAddress(const FConstProcessorContextEnvironment& Environment, const UScriptStruct* ColumnType);
+
 		static const void* GetColumnBatchAddress(const FDirectContextEnvironment& Environment, const UScriptStruct* ColumnType);
 		static void* GetMutableColumnBatchAddress(FDirectContextEnvironment& Environment, const UScriptStruct* ColumnType);
+
+		static const void* GetColumnBatchAddress(const FConstDirectContextEnvironment& Environment, const UScriptStruct* ColumnType);
 	};
 } // namespace UE::Editor::DataStorage::Queries
\ No newline at end of file
diff --git a/Engine/Plugins/Experimental/EditorDataStorage/Source/TedsCore/Private/Queries/QueryContextCapabilities/SingleRow.cpp b/Engine/Plugins/Experimental/EditorDataStorage/Source/TedsCore/Private/Queries/QueryContextCapabilities/SingleRow.cpp
--- a/Engine/Plugins/Experimental/EditorDataStorage/Source/TedsCore/Private/Queries/QueryContextCapabilities/SingleRow.cpp
+++ b/Engine/Plugins/Experimental/EditorDataStorage/Source/TedsCore/Private/Queries/QueryContextCapabilities/SingleRow.cpp
@@ -2,9 +2,8 @@
 
 #include "Queries/QueryContextCapabilities/SingleRow.h"
 
-#include "Mass/EntityHandle.h"
 #include "Queries/QueryContextCapabilities/EditorDataStorageContextEnvironment.h"
-#include "Queries/TypedElementQueryContext_MassForwarder.h"
+#include "Queries/QueryContextCapabilities/TableCapabilities.h"
 #include "TypedElementDatabaseEnvironment.h"
 
 namespace UE::Editor::DataStorage::Queries
@@ -12,35 +11,18 @@ namespace UE::Editor::DataStorage::Queries
 	bool FSingleRowInfoContextCapability::CurrentRowHasColumns(
 		const IContextEnvironment& Environment, TConstArrayView<const UScriptStruct*> Columns)
 	{
-		// Pick the first row since all rows in a batch are in the same table.
-		if (FRowHandleArrayView Rows = Environment.GetBatchRows(); !Rows.IsEmpty())
-		{
-			FMassEntityHandle Entity = FMassEntityHandle::FromNumber(Rows[0]);
-			FMassArchetypeHandle Archetype = Environment.GetEnvironment()->GetMassEntityManager().GetArchetypeForEntity(Entity);
-			const FMassElementBitSet& ArchetypeBits = Archetype.GetCompositionBitSetChecked();
-			for (const UScriptStruct* Column : Columns)
-			{
-				if (!ArchetypeBits.Contains(Column))
-				{
-					return false;
-				}
-			}
-			return true;
-		}
-		return false;
+		return FCurrentTableInfoContextCapability::CurrentTableHasColumns(Environment, Columns);
 	}
 
 	bool FSingleRowInfoContextCapability::CurrentRowHasColumns(
 		const FProcessorContextEnvironment& Environment, TConstArrayView<const UScriptStruct*> Columns)
 	{
-		const FMassExecutionContext& Context = Environment.MassContext;
-		for (const UScriptStruct* Column : Columns)
-		{
-			if (!Context.DoesArchetypeHaveElement(Column))
-			{
-				return false;
-			}
-		}
-		return true;
+		return FCurrentTableInfoContextCapability::CurrentTableHasColumns(Environment, Columns);
+	}
+
+	bool FSingleRowInfoContextCapability::CurrentRowHasColumns(
+		const FConstProcessorContextEnvironment& Environment, TConstArrayView<const UScriptStruct*> Columns)
+	{
+		return FCurrentTableInfoContextCapability::CurrentTableHasColumns(Environment, Columns);
 	}
 } // namespace UE::Editor::DataStorage::Queries
diff --git a/Engine/Plugins/Experimental/EditorDataStorage/Source/TedsCore/Private/Queries/QueryContextCapabilities/SingleRow.h b/Engine/Plugins/Experimental/EditorDataStorage/Source/TedsCore/Private/Queries/QueryContextCapabilities/SingleRow.h
--- a/Engine/Plugins/Experimental/EditorDataStorage/Source/TedsCore/Private/Queries/QueryContextCapabilities/SingleRow.h
+++ b/Engine/Plugins/Experimental/EditorDataStorage/Source/TedsCore/Private/Queries/QueryContextCapabilities/SingleRow.h
@@ -19,13 +19,13 @@ namespace UE::Editor::DataStorage::Queries
 		{ Type.CurrentRow } -> UE::CConvertibleTo<RowHandle>;
 	};
 
-	// For documentation on these functions see "TypedElementQueryCapabilities.inl"
 	class FSingleRowInfoContextCapability : public ImplementsContextCapability<SingleRowInfo>
 	{
 	public:
 		static bool CurrentRowHasColumns(const IContextEnvironment& Environment, TConstArrayView<const UScriptStruct*> Columns);
 		static bool CurrentRowHasColumns(const FProcessorContextEnvironment& Environment, TConstArrayView<const UScriptStruct*> Columns);
-		
+		static bool CurrentRowHasColumns(const FConstProcessorContextEnvironment& Environment, TConstArrayView<const UScriptStruct*> Columns);
+
 		template<CHasCurrentRow EnvironmentType>
 		static RowHandle GetCurrentRow(const EnvironmentType& Environment)
 		{ 
diff --git a/Engine/Plugins/Experimental/EditorDataStorage/Source/TedsCore/Private/Queries/QueryContextCapabilities/TableCapabilities.cpp b/Engine/Plugins/Experimental/EditorDataStorage/Source/TedsCore/Private/Queries/QueryContextCapabilities/TableCapabilities.cpp
new file mode 100644
--- /dev/null
+++ b/Engine/Plugins/Experimental/EditorDataStorage/Source/TedsCore/Private/Queries/QueryContextCapabilities/TableCapabilities.cpp
@@ -0,0 +1,188 @@
+// Copyright Epic Games, Inc. All Rights Reserved.
+
+#include "Queries/QueryContextCapabilities/TableCapabilities.h"
+
+#include "MassArchetypeTypes.h"
+#include "MassExecutionContext.h"
+#include "Queries/TypedElementQueryContext_DirectSingleApi.h"
+#include "Queries/TypedElementQueryContext_MassForwarder.h"
+#include "TypedElementDatabaseEnvironment.h"
+
+namespace UE::Editor::DataStorage::Queries
+{
+	// 
+	// FTableInfoContextCapability
+	//
+
+	bool FTableInfoContextCapability::IsTableValid(const IContextEnvironment& Environment, TableHandle Table)
+	{
+		return Environment.GetEnvironment()->GetTableManager().IsTableValid(Table);
+	}
+
+	TableHandle FTableInfoContextCapability::FindTable(const IContextEnvironment& Environment, const FName& TableName)
+	{
+		return Environment.GetEnvironment()->GetTableManager().Find(TableName);
+	}
+
+	TableHandle FTableInfoContextCapability::FindTable(const IContextEnvironment& Environment, RowHandle Row)
+	{
+		return Environment.GetEnvironment()->GetTableManager().Find(Row);
+	}
+
+	void FTableInfoContextCapability::ListTables(
+		const IContextEnvironment& Environment, TFunctionRef<void(TableHandle)> Callback)
+	{
+		return Environment.GetEnvironment()->GetTableManager().List(Callback);
+	}
+
+	void FTableInfoContextCapability::ListTables(
+		const IContextEnvironment& Environment, ETableType TableType, TFunctionRef<void(TableHandle)> Callback)
+	{
+		return Environment.GetEnvironment()->GetTableManager().List(TableType, Callback);
+	}
+
+	void FTableInfoContextCapability::ListTableColumns(
+		const IContextEnvironment& Environment, TableHandle Table, TFunctionRef<bool(const UScriptStruct*)> Callback)
+	{
+		return Environment.GetEnvironment()->GetTableManager().ListColumns(Table, Callback);
+	}
+
+	bool FTableInfoContextCapability::TableHasColumns(
+		const IContextEnvironment& Environment, TableHandle Table, TConstArrayView<const UScriptStruct*> Columns)
+	{
+		return Environment.GetEnvironment()->GetTableManager().HasColumns(Table, Columns);
+	}
+
+	FTableInfoView FTableInfoContextCapability::GetTableInfo(const IContextEnvironment& Environment, TableHandle Table)
+	{
+		return Environment.GetEnvironment()->GetTableManager().GetInfo(Table);
+	}
+
+	void FTableInfoContextCapability::ListTableRows(
+		const IContextEnvironment& Environment, TableHandle Table, TFunctionRef<void(FRowHandleArrayView)> Callback)
+	{
+		return Environment.GetEnvironment()->GetTableManager().ListRows(Table, Callback);
+	}
+
+	void FTableInfoContextCapability::ListTableForeignKeyDomains(
+		const IContextEnvironment& Environment, TableHandle Table, bool bIncludeParents, TFunctionRef<void(const FName&)> Callback)
+	{
+		return Environment.GetEnvironment()->GetTableManager().ListForeignKeyDomains(Table, bIncludeParents, Callback);
+	}
+
+
+
+	// 
+	// FCurrentTableInfoContextCapability
+	//
+
+	TableHandle FCurrentTableInfoContextCapability::GetCurrentTable(const IContextEnvironment& Environment)
+	{
+		FRowHandleArrayView Rows = Environment.GetBatchRows();
+		return !Rows.IsEmpty()
+			? FTableInfoContextCapability::FindTable(Environment, Rows.First())
+			: InvalidTableHandle;
+	}
+
+	TableHandle FCurrentTableInfoContextCapability::GetCurrentTable(const FProcessorContextEnvironment& Environment)
+	{
+		return Environment.GetEnvironment()->GetTableManager().Find(Environment.MassContext.GetEntityCollection().GetArchetype());
+	}
+
+	TableHandle FCurrentTableInfoContextCapability::GetCurrentTable(const FConstProcessorContextEnvironment& Environment)
+	{
+		return Environment.GetEnvironment()->GetTableManager().Find(Environment.MassContext.GetEntityCollection().GetArchetype());
+	}
+
+	bool FCurrentTableInfoContextCapability::CurrentTableHasColumns(
+		const IContextEnvironment& Environment, TConstArrayView<const UScriptStruct*> Columns)
+	{
+		TableHandle Table = GetCurrentTable(Environment);
+		return Table != InvalidTableHandle ? FTableInfoContextCapability::TableHasColumns(Environment, Table, Columns) : false;
+	}
+
+	bool FCurrentTableInfoContextCapability::CurrentTableHasColumns(
+		const FProcessorContextEnvironment& Environment, TConstArrayView<const UScriptStruct*> Columns)
+	{
+		const FMassExecutionContext& Context = Environment.MassContext;
+		for (const UScriptStruct* Column : Columns)
+		{
+			if (!Context.DoesArchetypeHaveElement(Column))
+			{
+				return false;
+			}
+		}
+		return true;
+	}
+
+	bool FCurrentTableInfoContextCapability::CurrentTableHasColumns(
+		const FConstProcessorContextEnvironment& Environment, TConstArrayView<const UScriptStruct*> Columns)
+	{
+		const FMassExecutionContext& Context = Environment.MassContext;
+		for (const UScriptStruct* Column : Columns)
+		{
+			if (!Context.DoesArchetypeHaveElement(Column))
+			{
+				return false;
+			}
+		}
+		return true;
+	}
+
+
+
+	// 
+	// FTableManagementContextCapability
+	//
+
+	TableHandle FTableManagementContextCapability::RegisterTable(IContextEnvironment& Environment,
+		TConstArrayView<const UScriptStruct*> Columns, const FName& TableName)
+	{
+		return Environment.GetEnvironment()->GetTableManager().Register(Columns, TableName);
+	}
+
+	TableHandle FTableManagementContextCapability::RegisterTable(IContextEnvironment& Environment,
+		TConstArrayView<const UScriptStruct*> Columns, const FName& TableName, FTableRegistrationOptions Options)
+	{
+		return Environment.GetEnvironment()->GetTableManager().Register(Columns, TableName, Options);
+	}
+
+	void FTableManagementContextCapability::RegisterTableForeignKey(IContextEnvironment& Environment,
+		TableHandle Table, const FName& Domain, const Queries::TConstQueryFunction<FForeignKey>& KeyConstructor)
+	{
+		return Environment.GetEnvironment()->GetTableManager().RegisterForeignKey(Table, Domain, KeyConstructor);
+	}
+
+	void FTableManagementContextCapability::RegisterTableForeignKey(IContextEnvironment& Environment,
+		TableHandle Table, const FName& Domain, Queries::TConstQueryFunction<FForeignKey>&& KeyConstructor)
+	{
+		return Environment.GetEnvironment()->GetTableManager().RegisterForeignKey(Table, Domain, MoveTemp(KeyConstructor));
+	}
+
+	FForeignKey FTableManagementContextCapability::BuildForeignKey(
+		const FConstDirectContextEnvironment& Environment, const FName& Domain, RowHandle Row)
+	{
+		return BuildForeignKey(Environment.GetEnvironment(), Environment.DirectContext, Domain, Row);
+	}
+
+	FForeignKey FTableManagementContextCapability::BuildForeignKey(
+		const FDirectContextEnvironment& Environment, const FName& Domain, RowHandle Row)
+	{
+		return BuildForeignKey(Environment.GetEnvironment(), Environment.DirectContext, Domain, Row);
+	}
+
+	FForeignKey FTableManagementContextCapability::BuildForeignKey(
+		const FEnvironment* Environment, const UEditorDataStorage* DirectContext, const FName& Domain, RowHandle Row)
+	{
+		checkf(Environment, TEXT("Unable to build foreign key because the TEDS environment doesn't exist."));
+		checkf(DirectContext, TEXT("Unable to build foreign key because TEDS Core hasn't been instantiated."));
+
+		FConstDirectContextEnvironment LocalContextEnvironment(*Environment, *DirectContext);
+		QueryContext_ConstDirectApi Context(LocalContextEnvironment);
+		LocalContextEnvironment.SetBatch(FRowHandleArrayView(&Row, 1, FRowHandleArrayView::EFlags::IsSorted | FRowHandleArrayView::EFlags::IsUnique));
+
+		FForeignKey Result;
+		Environment->GetTableManager().BuildForeignKeys(TArrayView<FForeignKey>(&Result, 1), Domain, Context, LocalContextEnvironment);
+		return Result;
+	}
+} // namespace UE::Editor::DataStorage::Queries
diff --git a/Engine/Plugins/Experimental/EditorDataStorage/Source/TedsCore/Private/Queries/QueryContextCapabilities/TableCapabilities.h b/Engine/Plugins/Experimental/EditorDataStorage/Source/TedsCore/Private/Queries/QueryContextCapabilities/TableCapabilities.h
new file mode 100644
--- /dev/null
+++ b/Engine/Plugins/Experimental/EditorDataStorage/Source/TedsCore/Private/Queries/QueryContextCapabilities/TableCapabilities.h
@@ -0,0 +1,76 @@
+// Copyright Epic Games, Inc. All Rights Reserved.
+
+#pragma once
+
+#include "DataStorage/CommonTypes.h"
+#include "DataStorage/Queries/Types.h"
+#include "Elements/Framework/TypedElementQueryCapabilities.h"
+#include "Queries/QueryContextCapabilities/EditorDataStorageContextEnvironment.h"
+
+class FName;
+class UScriptStruct;
+class UEditorDataStorage;
+
+namespace UE::Editor::DataStorage::Queries
+{
+	struct IContextEnvironment;
+	struct FConstDirectContextEnvironment;
+	struct FConstProcessorContextEnvironment;
+	struct FDirectContextEnvironment;
+	struct FProcessorContextEnvironment;
+	
+	class FTableInfoContextCapability : public ImplementsContextCapability<TableInfo>
+	{
+	public:
+		static bool IsTableValid(const IContextEnvironment& Environment, TableHandle Table);
+		static TableHandle FindTable(const IContextEnvironment& Environment, const FName& TableName);
+		static TableHandle FindTable(const IContextEnvironment& Environment, RowHandle Row);
+		static void ListTables(const IContextEnvironment& Environment, TFunctionRef<void(TableHandle)> Callback);
+		static void ListTables(const IContextEnvironment& Environment, ETableType TableType, TFunctionRef<void(TableHandle)> Callback);
+		static void ListTableColumns(
+			const IContextEnvironment& Environment, TableHandle Table, TFunctionRef<bool(const UScriptStruct*)> Callback);
+
+		static bool TableHasColumns(
+			const IContextEnvironment& Environment, TableHandle Table, TConstArrayView<const UScriptStruct*> Columns);
+		static FTableInfoView GetTableInfo(const IContextEnvironment& Environment, TableHandle Table);
+		static void ListTableRows(
+			const IContextEnvironment& Environment, TableHandle Table, TFunctionRef<void(FRowHandleArrayView)> Callback);
+		
+		static void ListTableForeignKeyDomains(
+			const IContextEnvironment& Environment, TableHandle Table, bool bIncludeParents, TFunctionRef<void(const FName&)> Callback);
+	};
+
+	class FCurrentTableInfoContextCapability : public ImplementsContextCapability<CurrentTableInfo>
+	{
+	public:
+		static TableHandle GetCurrentTable(const IContextEnvironment& Environment);
+		static TableHandle GetCurrentTable(const FProcessorContextEnvironment& Environment);
+		static TableHandle GetCurrentTable(const FConstProcessorContextEnvironment& Environment);
+		
+		static bool CurrentTableHasColumns(
+			const IContextEnvironment& Environment, TConstArrayView<const UScriptStruct*> Columns);
+		static bool CurrentTableHasColumns(
+			const FProcessorContextEnvironment& Environment, TConstArrayView<const UScriptStruct*> Columns);
+		static bool CurrentTableHasColumns(
+			const FConstProcessorContextEnvironment& Environment, TConstArrayView<const UScriptStruct*> Columns);
+	};
+	
+	class FTableManagementContextCapability : public ImplementsContextCapability<TableManagement>
+	{
+	public:
+		static TableHandle RegisterTable(IContextEnvironment& Environment,
+			TConstArrayView<const UScriptStruct*> Columns, const FName& TableName);
+		static TableHandle RegisterTable(IContextEnvironment& Environment,
+			TConstArrayView<const UScriptStruct*> Columns, const FName& TableName, FTableRegistrationOptions Options);
+		static void RegisterTableForeignKey(IContextEnvironment& Environment,
+			TableHandle Table, const FName& Domain, const Queries::TConstQueryFunction<FForeignKey>& KeyConstructor);
+		static void RegisterTableForeignKey(IContextEnvironment& Environment,
+			TableHandle Table, const FName& Domain, Queries::TConstQueryFunction<FForeignKey>&& KeyConstructor);
+		static FForeignKey BuildForeignKey(const FDirectContextEnvironment& Environment, const FName& Domain, RowHandle Row);
+		static FForeignKey BuildForeignKey(const FConstDirectContextEnvironment& Environment, const FName& Domain, RowHandle Row);
+
+	private:
+		static FForeignKey BuildForeignKey(
+			const FEnvironment* Environment, const UEditorDataStorage* DirectContext, const FName& Domain, RowHandle Row);
+	};
+} // namespace UE::Editor::DataStorage::Queries
diff --git a/Engine/Plugins/Experimental/EditorDataStorage/Source/TedsCore/Private/Queries/TypedElementExtendedQueryStore.cpp b/Engine/Plugins/Experimental/EditorDataStorage/Source/TedsCore/Private/Queries/TypedElementExtendedQueryStore.cpp
--- a/Engine/Plugins/Experimental/EditorDataStorage/Source/TedsCore/Private/Queries/TypedElementExtendedQueryStore.cpp
+++ b/Engine/Plugins/Experimental/EditorDataStorage/Source/TedsCore/Private/Queries/TypedElementExtendedQueryStore.cpp
@@ -448,6 +448,58 @@ namespace UE::Editor::DataStorage
 		return Result;
 	}
 
+	FQueryResult FExtendedQueryStore::RunQueryCallbackCommon(
+		FMassEntityManager& EntityManager,
+		FEnvironment& Environment,
+		FMassExecutionContext* ParentContext,
+		Handle Query,
+		ERunQueryFlags ExecutionFlags,
+		const Queries::TConstQueryFunction<void>& Callback) const
+	{
+		FQueryResult Result;
+		if (const FExtendedQuery* QueryData = Get(Query))
+		{
+			switch (QueryData->Description.Action)
+			{
+			case FQueryDescription::EActionType::None:
+				Result.Completed = FQueryResult::ECompletion::Fully;
+				break;
+			case FQueryDescription::EActionType::Select:
+				if (!QueryData->Processor.IsValid())
+				{
+					if (CheckCompatibility(QueryData->Description, Callback))
+					{
+						Result = FTypedElementQueryProcessorData::Execute(
+							Callback, QueryData->Description, QueryData->NativeQuery, EntityManager, Environment, ExecutionFlags);
+					}
+					else
+					{
+						Result.Completed = FQueryResult::ECompletion::Incompatible;
+					}
+				}
+				else
+				{
+					Result.Completed = FQueryResult::ECompletion::Unsupported;
+				}
+				break;
+			case FQueryDescription::EActionType::Count:
+				// Only the count is requested so no need to trigger the callback.
+				Result.Count = QueryData->NativeQuery.GetNumMatchingEntities();
+				Result.Completed = FQueryResult::ECompletion::Fully;
+				break;
+			default:
+				Result.Completed = FQueryResult::ECompletion::Unsupported;
+				break;
+			}
+		}
+		else
+		{
+			Result.Completed = FQueryResult::ECompletion::Unavailable;
+		}
+
+		return Result;
+	}
+
 	FQueryResult FExtendedQueryStore::RunQuery(FMassEntityManager& EntityManager,
 		FEnvironment& Environment, Handle Query, EDirectQueryExecutionFlags DirectExecutionFlags,
 		DirectQueryCallbackRef Callback)
@@ -461,6 +513,13 @@ namespace UE::Editor::DataStorage
 	{
 		return RunQueryCallbackCommon(EntityManager, Environment, nullptr, Query, Flags, Callback);
 	}
+
+	FQueryResult FExtendedQueryStore::RunQuery(FMassEntityManager& EntityManager,
+		FEnvironment& Environment, Handle Query, ERunQueryFlags Flags,
+		const Queries::TConstQueryFunction<void>& Callback) const
+	{
+		return RunQueryCallbackCommon(EntityManager, Environment, nullptr, Query, Flags, Callback);
+	}
 	
 	FQueryResult FExtendedQueryStore::RunQuery(FMassEntityManager& EntityManager,
 		FEnvironment& Environment, FMassExecutionContext& ParentContext, Handle Query,
@@ -1343,4 +1402,17 @@ namespace UE::Editor::DataStorage
 		}
 		return true;
 	}
+
+	bool FExtendedQueryStore::CheckCompatibility(const FQueryDescription& QueryDescription, const Queries::TConstQueryFunction<void> QueryCallback) const
+	{
+		for (const UScriptStruct* Column : QueryCallback.ConstColumnTypes)
+		{
+			if (int32 Index = QueryDescription.SelectionTypes.Find(Column);
+				Index == INDEX_NONE || QueryDescription.SelectionAccessTypes[Index] == EQueryAccessType::OptionalReadOnly)
+			{
+				return false;
+			}
+		}
+		return true;
+	}
 } // namespace UE::Editor::DataStorage
diff --git a/Engine/Plugins/Experimental/EditorDataStorage/Source/TedsCore/Private/Queries/TypedElementExtendedQueryStore.h b/Engine/Plugins/Experimental/EditorDataStorage/Source/TedsCore/Private/Queries/TypedElementExtendedQueryStore.h
--- a/Engine/Plugins/Experimental/EditorDataStorage/Source/TedsCore/Private/Queries/TypedElementExtendedQueryStore.h
+++ b/Engine/Plugins/Experimental/EditorDataStorage/Source/TedsCore/Private/Queries/TypedElementExtendedQueryStore.h
@@ -21,7 +21,7 @@ namespace UE::Editor::DataStorage
 
 	struct FExtendedQuery
 	{
-		FMassEntityQuery NativeQuery; // Used if there's no processor bound.
+		mutable FMassEntityQuery NativeQuery; // Used if there's no processor bound. Mutable because NativeQuery does internal state caching.
 		FQueryDescription Description;
 		TStrongObjectPtr<UMassProcessor> Processor;
 		FMassEntityQuery* QueryReference;
@@ -138,6 +138,12 @@ namespace UE::Editor::DataStorage
 			Handle Query,
 			ERunQueryFlags Flags,
 			const Queries::TQueryFunction<void>& Callback);
+		FQueryResult RunQuery(
+			FMassEntityManager& EntityManager,
+			FEnvironment& Environment,
+			Handle Query,
+			ERunQueryFlags Flags,
+			const Queries::TConstQueryFunction<void>& Callback) const;
 		FQueryResult RunQuery(
 			FMassEntityManager& EntityManager,
 			FEnvironment& Environment,
@@ -195,6 +201,13 @@ namespace UE::Editor::DataStorage
 			Handle Query,
 			ExecutionFlagsType ExecutionFlags,
 			CallbackReference Callback);
+		FQueryResult RunQueryCallbackCommon(
+			FMassEntityManager& EntityManager,
+			FEnvironment& Environment,
+			FMassExecutionContext* ParentContext,
+			Handle Query,
+			ERunQueryFlags ExecutionFlags,
+			const Queries::TConstQueryFunction<void>& Callback) const;
 
 		FMassEntityQuery& SetupNativeQuery(FQueryDescription& Query, FExtendedQuery& StoredQuery, FEnvironment& Environment);
 		bool SetupDynamicColumns(FQueryDescription& Query, FEnvironment& Environment);
@@ -222,6 +235,7 @@ namespace UE::Editor::DataStorage
 		void UnregisterQueryData(Handle Query, FExtendedQuery& QueryData, FMassEntityManager& EntityManager, FMassProcessingPhaseManager& PhaseManager);
 
 		bool CheckCompatibility(const FQueryDescription& QueryDescription, const Queries::TQueryFunction<void> QueryCallback) const;
+		bool CheckCompatibility(const FQueryDescription& QueryDescription, const Queries::TConstQueryFunction<void> QueryCallback) const;
 
 		static const FQueryDescription EmptyDescription;
 
diff --git a/Engine/Plugins/Experimental/EditorDataStorage/Source/TedsCore/Private/Queries/TypedElementQueryContext_DirectSingleApi.cpp b/Engine/Plugins/Experimental/EditorDataStorage/Source/TedsCore/Private/Queries/TypedElementQueryContext_DirectSingleApi.cpp
--- a/Engine/Plugins/Experimental/EditorDataStorage/Source/TedsCore/Private/Queries/TypedElementQueryContext_DirectSingleApi.cpp
+++ b/Engine/Plugins/Experimental/EditorDataStorage/Source/TedsCore/Private/Queries/TypedElementQueryContext_DirectSingleApi.cpp
@@ -47,6 +47,10 @@ namespace UE::Editor::DataStorage::Queries
 				*ColumnsDataIt++ = DirectContext->GetColumnData(Root, ColumnType);
 			}
 		}
+		else
+		{
+			FMemory::Memzero(ColumnsData.GetData(), ColumnsData.NumBytes());
+		}
 	}
 
 	void FDirectContextEnvironment::GetMutableColumns(
@@ -62,6 +66,10 @@ namespace UE::Editor::DataStorage::Queries
 				*ColumnsDataIt++ = DirectContext->GetColumnData(Root, ColumnType);
 			}
 		}
+		else
+		{
+			FMemory::Memzero(ColumnsData.GetData(), ColumnsData.NumBytes());
+		}
 	}
 
 	void FDirectContextEnvironment::SetCurrentRow(RowHandle Row)
@@ -69,4 +77,66 @@ namespace UE::Editor::DataStorage::Queries
 		CurrentRow = Row;
 	}
 
+
+
+	//
+	// FConstDirectContextEnvironment
+	//
+
+	FConstDirectContextEnvironment::FConstDirectContextEnvironment(const FEnvironment& Environment, const UEditorDataStorage& DirectContext)
+		: Environment(&Environment)
+		, DirectContext(&DirectContext)
+	{
+	}
+
+	void FConstDirectContextEnvironment::SetBatch(FRowHandleArrayView InRows)
+	{
+		Rows = InRows;
+	}
+
+	FEnvironment* FConstDirectContextEnvironment::GetEnvironment()
+	{
+		checkf(false, TEXT("Calling a mutable version of the environment is not supported by a const environment."));
+		return nullptr;
+	}
+
+	const FEnvironment* FConstDirectContextEnvironment::GetEnvironment() const
+	{
+		return Environment;
+	}
+
+	FRowHandleArrayView FConstDirectContextEnvironment::GetBatchRows() const
+	{
+		return Rows;
+	}
+
+	void FConstDirectContextEnvironment::GetConstColumns(
+		TArrayView<const void*> ColumnsData, TConstArrayView<const UScriptStruct*> ColumnTypes) const
+	{
+		if (!Rows.IsEmpty())
+		{
+			RowHandle Root = *Rows.GetData();
+			const void** ColumnsDataIt = ColumnsData.GetData();
+			for (const UScriptStruct* ColumnType : ColumnTypes)
+			{
+				*ColumnsDataIt++ = DirectContext->GetColumnData(Root, ColumnType);
+			}
+		}
+		else
+		{
+			FMemory::Memzero(ColumnsData.GetData(), ColumnsData.NumBytes());
+		}
+	}
+
+	void FConstDirectContextEnvironment::GetMutableColumns(
+		TArrayView<void*> ColumnsData, TConstArrayView<const UScriptStruct*> ColumnTypes)
+	{
+		checkf(false, TEXT("Retrieving mutable columns is not supported from a const environment."));
+		FMemory::Memzero(ColumnsData.GetData(), ColumnsData.NumBytes());
+	}
+
+	void FConstDirectContextEnvironment::SetCurrentRow(RowHandle Row)
+	{
+		CurrentRow = Row;
+	}
 } // namespace UE::Editor::DataStorage::Queries
\ No newline at end of file
diff --git a/Engine/Plugins/Experimental/EditorDataStorage/Source/TedsCore/Private/Queries/TypedElementQueryContext_DirectSingleApi.h b/Engine/Plugins/Experimental/EditorDataStorage/Source/TedsCore/Private/Queries/TypedElementQueryContext_DirectSingleApi.h
--- a/Engine/Plugins/Experimental/EditorDataStorage/Source/TedsCore/Private/Queries/TypedElementQueryContext_DirectSingleApi.h
+++ b/Engine/Plugins/Experimental/EditorDataStorage/Source/TedsCore/Private/Queries/TypedElementQueryContext_DirectSingleApi.h
@@ -8,6 +8,7 @@
 #include "TypedElementDatabaseEnvironment.h"
 #include "Queries/QueryContextCapabilities/EditorDataStorageContextEnvironment.h"
 #include "Queries/QueryContextCapabilities/EditorDataStorageDynamicColumnInfo.h"
+#include "Queries/QueryContextCapabilities/TableCapabilities.h"
 #include "Queries/QueryContextCapabilities/RowBatch.h"
 #include "Queries/QueryContextCapabilities/SingleRow.h"
 
@@ -36,7 +37,38 @@ namespace UE::Editor::DataStorage::Queries
 		RowHandle CurrentRow = InvalidRowHandle;
 	};
 
-	using QueryContext_DirectApi = TQueryContextImpl<FDirectContextEnvironment,
+	struct FConstDirectContextEnvironment final : IContextEnvironment
+	{
+		/** The provided rows need to be contiguous in the same table. */
+		FConstDirectContextEnvironment() = default;
+		FConstDirectContextEnvironment(const FEnvironment& Environment, const UEditorDataStorage& DirectContext);
+
+		void SetBatch(FRowHandleArrayView InRows);
+
+		virtual FEnvironment* GetEnvironment() override;
+		virtual const FEnvironment* GetEnvironment() const override;
+
+		virtual FRowHandleArrayView GetBatchRows() const override;
+		virtual void GetConstColumns(TArrayView<const void*> ColumnsData, TConstArrayView<const UScriptStruct*> ColumnTypes) const override;
+		virtual void GetMutableColumns(TArrayView<void*> ColumnsData, TConstArrayView<const UScriptStruct*> ColumnTypes) override;
+		virtual void SetCurrentRow(RowHandle Row) override;
+
+		FRowHandleArrayView Rows;
+		const FEnvironment* Environment = nullptr;
+		const UEditorDataStorage* DirectContext = nullptr;
+		RowHandle CurrentRow = InvalidRowHandle;
+	};
+
+	using QueryContext_DirectApi = TQueryContextImpl<false, FDirectContextEnvironment,
+		FCurrentTableInfoContextCapability,
+		FTableInfoContextCapability,
+		FSingleRowInfoContextCapability,
+		FRowBatchInfoContextCapability,
+		FDynamicColumnInfoContextCapability>;
+
+	using QueryContext_ConstDirectApi = TQueryContextImpl<true, FConstDirectContextEnvironment,
+		FCurrentTableInfoContextCapability,
+		FTableInfoContextCapability,
 		FSingleRowInfoContextCapability,
 		FRowBatchInfoContextCapability,
 		FDynamicColumnInfoContextCapability>;
diff --git a/Engine/Plugins/Experimental/EditorDataStorage/Source/TedsCore/Private/Queries/TypedElementQueryContext_MassForwarder.cpp b/Engine/Plugins/Experimental/EditorDataStorage/Source/TedsCore/Private/Queries/TypedElementQueryContext_MassForwarder.cpp
--- a/Engine/Plugins/Experimental/EditorDataStorage/Source/TedsCore/Private/Queries/TypedElementQueryContext_MassForwarder.cpp
+++ b/Engine/Plugins/Experimental/EditorDataStorage/Source/TedsCore/Private/Queries/TypedElementQueryContext_MassForwarder.cpp
@@ -66,4 +66,60 @@ namespace UE::Editor::DataStorage::Queries
 	{
 		CurrentRow = Row;
 	}
+
+
+	//
+	// FConstProcessorContextEnvironment
+	//
+
+	FConstProcessorContextEnvironment::FConstProcessorContextEnvironment(
+		const FEnvironment& Environment, const FMassExecutionContext& MassContext)
+		: Environment(Environment)
+		, MassContext(MassContext)
+	{
+	}
+
+	FEnvironment* FConstProcessorContextEnvironment::GetEnvironment()
+	{
+		checkf(false, TEXT("Calling a mutable version of the environment is not supported by a const environment."));
+		return nullptr;
+	}
+
+	const FEnvironment* FConstProcessorContextEnvironment::GetEnvironment() const
+	{
+		return &Environment;
+	}
+
+	FRowHandleArrayView FConstProcessorContextEnvironment::GetBatchRows() const
+	{
+		TConstArrayView<RowHandle> Rows = MassEntitiesToRowsConversion(MassContext.GetEntities());
+		return FRowHandleArrayView(Rows, FRowHandleArrayView::EFlags::IsUnique);
+	}
+
+	void FConstProcessorContextEnvironment::GetConstColumns(
+		TArrayView<const void*> ColumnsData, TConstArrayView<const UScriptStruct*> ColumnTypes) const
+	{
+		checkfSlow(ColumnsData.Num() >= ColumnTypes.Num(), TEXT("The list size of requested const column types (%i) doesn't match the "
+			"available space for the results (%i)"), ColumnTypes.Num(), ColumnsData.Num());
+
+		const void** ColumnsDataIt = ColumnsData.GetData();
+		for (const UScriptStruct* ColumnType : ColumnTypes)
+		{
+			checkfSlow(UE::Mass::IsA<FMassFragment>(ColumnType),
+				TEXT("Unable to retrieve const column data for a column type without data: %s"), *ColumnType->GetFullName());
+			*ColumnsDataIt++ = MassContext.GetFragmentView(ColumnType).GetData();
+		}
+	}
+
+	void FConstProcessorContextEnvironment::GetMutableColumns(
+		TArrayView<void*> ColumnsData, TConstArrayView<const UScriptStruct*> ColumnTypes)
+	{
+		checkf(false, TEXT("Retrieving mutable columns is not supported from a const environment."));
+		FMemory::Memzero(ColumnsData.GetData(), ColumnsData.NumBytes());
+	}
+
+	void FConstProcessorContextEnvironment::SetCurrentRow(RowHandle Row)
+	{
+		CurrentRow = Row;
+	}
 } // namespace UE::Editor::DataStorage::Queries
\ No newline at end of file
diff --git a/Engine/Plugins/Experimental/EditorDataStorage/Source/TedsCore/Private/Queries/TypedElementQueryContext_MassForwarder.h b/Engine/Plugins/Experimental/EditorDataStorage/Source/TedsCore/Private/Queries/TypedElementQueryContext_MassForwarder.h
--- a/Engine/Plugins/Experimental/EditorDataStorage/Source/TedsCore/Private/Queries/TypedElementQueryContext_MassForwarder.h
+++ b/Engine/Plugins/Experimental/EditorDataStorage/Source/TedsCore/Private/Queries/TypedElementQueryContext_MassForwarder.h
@@ -6,6 +6,7 @@
 #include "Elements/Framework/TypedElementQueryFunctions.h"
 #include "TypedElementDatabaseEnvironment.h"
 #include "Queries/QueryContextCapabilities/EditorDataStorageDynamicColumnInfo.h"
+#include "Queries/QueryContextCapabilities/TableCapabilities.h"
 #include "Queries/QueryContextCapabilities/RowBatch.h"
 #include "Queries/QueryContextCapabilities/SingleRow.h"
 
@@ -30,7 +31,33 @@ namespace UE::Editor::DataStorage::Queries
 		RowHandle CurrentRow = InvalidRowHandle;
 	};
 
-	using QueryContext_ProcessorApi = TQueryContextImpl<FProcessorContextEnvironment,
+	struct FConstProcessorContextEnvironment final : IContextEnvironment
+	{
+		FConstProcessorContextEnvironment(const FEnvironment& Environment, const FMassExecutionContext& MassContext);
+
+		virtual FEnvironment* GetEnvironment() override;
+		virtual const FEnvironment* GetEnvironment() const override;
+
+		virtual FRowHandleArrayView GetBatchRows() const override;
+		virtual void GetConstColumns(TArrayView<const void*> ColumnsData, TConstArrayView<const UScriptStruct*> ColumnTypes) const override;
+		virtual void GetMutableColumns(TArrayView<void*> ColumnsData, TConstArrayView<const UScriptStruct*> ColumnTypes) override;
+		virtual void SetCurrentRow(RowHandle Row) override;
+
+		const FEnvironment& Environment;
+		const FMassExecutionContext& MassContext;
+		RowHandle CurrentRow = InvalidRowHandle;
+	};
+
+	using QueryContext_ProcessorApi = TQueryContextImpl<false, FProcessorContextEnvironment,
+		FCurrentTableInfoContextCapability,
+		FTableInfoContextCapability,
+		FSingleRowInfoContextCapability,
+		FRowBatchInfoContextCapability,
+		FDynamicColumnInfoContextCapability>;
+	
+	using QueryContext_ConstProcessorApi = TQueryContextImpl<true, FConstProcessorContextEnvironment,
+		FCurrentTableInfoContextCapability,
+		FTableInfoContextCapability,
 		FSingleRowInfoContextCapability,
 		FRowBatchInfoContextCapability,
 		FDynamicColumnInfoContextCapability>;
diff --git a/Engine/Plugins/Experimental/EditorDataStorage/Source/TedsCore/Private/TableManager.cpp b/Engine/Plugins/Experimental/EditorDataStorage/Source/TedsCore/Private/TableManager.cpp
new file mode 100644
--- /dev/null
+++ b/Engine/Plugins/Experimental/EditorDataStorage/Source/TedsCore/Private/TableManager.cpp
@@ -0,0 +1,358 @@
+// Copyright Epic Games, Inc. All Rights Reserved.
+
+#include "TableManager.h"
+
+#include <cinttypes>
+#include "DataStorage/Debug/Log.h"
+#include "DynamicColumnGenerator.h"
+#include "EditorDataStorageSettings.h"
+#include "Elements/Framework/TypedElementQueryContract.h"
+#include "MassEntityManager.h"
+#include "MassEntityTypes.h"
+#include "TypedElementUtils.h"
+
+namespace UE::Editor::DataStorage
+{
+	FTableManager::FTableManager(FMassEntityManager& InMassEntityManager, FDynamicColumnGenerator& InDynamicColumns)
+		: MassEntityManager(InMassEntityManager)
+		, DynamicColumns(InDynamicColumns)
+	{
+		OnNewArchetypeEventHandle = InMassEntityManager.GetOnNewArchetypeEvent().AddRaw(this, &FTableManager::OnNewArchetype);
+	}
+
+	FTableManager::~FTableManager()
+	{
+		MassEntityManager.GetOnNewArchetypeEvent().Remove(OnNewArchetypeEventHandle);
+		Clear();
+	}
+
+	void FTableManager::Clear()
+	{
+		Tables.Reset();
+		NameLookup.Reset();
+		MassLookup.Reset();
+	}
+
+	bool FTableManager::IsTableValid(TableHandle Table) const
+	{
+		return Table != InvalidTableHandle && Table < Tables.Num();
+	}
+
+	TableHandle FTableManager::Register(TConstArrayView<const UScriptStruct*> ColumnList, const FName& TableName)
+	{
+		return Register(ColumnList, TableName, {});
+	}
+
+	TableHandle FTableManager::Register(
+		TConstArrayView<const UScriptStruct*> ColumnList, const FName& TableName, FTableRegistrationOptions Options)
+	{
+		checkf(!NameLookup.Contains(TableName), TEXT("Table '%s' has already been registered with the editor data storage."),
+			*TableName.ToString());
+		TableHandle Result = Tables.Num();
+
+		ETableType Type;
+		FMassArchetypeCompositionDescriptor Composition;
+		if (IsTableValid(Options.SourceTable))
+		{
+			int32 TableHandleAsIndex = ConvertTableHandleToIndex(Options.SourceTable);
+			FTableInfo& TableInfo = Tables[TableHandleAsIndex];
+			TableInfo.Derivatives.Add(Result);
+			Composition = MassEntityManager.GetArchetypeComposition(TableInfo.Archetype);
+			RemoveTableTag(Composition);
+			Type = ETableType::Derivative;
+
+			UE_LOG(LogEditorDataStorage, VeryVerbose, TEXT("New table derivative '%s' from template '%s"), 
+				*TableName.ToString(), 
+				*TableInfo.Name.ToString());
+		}
+		else
+		{
+			Type = ETableType::Declared;
+
+			UE_LOG(LogEditorDataStorage, VeryVerbose, TEXT("New table declared '%s'"), *TableName.ToString());
+		}
+
+		FMassArchetypeCreationParams ArchetypeCreationParams;
+		ArchetypeCreationParams.DebugName = TableName;
+		ArchetypeCreationParams.ChunkMemorySize = Options.MemorySize == 0 ? GetTableChunkSize(TableName) : Options.MemorySize;
+
+		for (const UScriptStruct* Column : ColumnList)
+		{
+			Composition.Add(Column);
+		}
+		// Add table tag.
+		const UScriptStruct* TableTag = DynamicColumns.GenerateColumn(*FTableDynamicTag::StaticStruct(), TableName).Type;
+		Composition.Add(TableTag);
+
+		bRegisteringNewTable = true;
+		FMassArchetypeHandle Archetype = MassEntityManager.CreateArchetype(Composition, ArchetypeCreationParams);
+		bRegisteringNewTable = false;
+
+		Tables.Emplace(FTableInfo
+			{
+				.Archetype = Archetype,
+				.Name = TableName,
+				.Parent = Options.SourceTable,
+				.Type = Type
+			});
+		NameLookup.Add(TableName, Result);
+		MassLookup.Add(Archetype, Result);
+		return Result;
+	}
+
+	TableHandle FTableManager::Find(const FName& Name) const
+	{
+		const TableHandle* TableHandle = NameLookup.Find(Name);
+		return TableHandle ? *TableHandle : InvalidTableHandle;
+	}
+
+	TableHandle FTableManager::Find(const FMassArchetypeHandle& Archetype) const
+	{
+		const TableHandle* TableHandle = MassLookup.Find(Archetype);
+		return TableHandle ? *TableHandle : InvalidTableHandle;
+	}
+
+	TableHandle FTableManager::Find(RowHandle Row) const
+	{
+		return Find(MassEntityManager.GetArchetypeForEntity(FMassEntityHandle::FromNumber(Row)));
+	}
+
+	void FTableManager::List(TFunctionRef<void(TableHandle)> Callback) const
+	{
+		int32 TableCount = Tables.Num();
+		for (int32 Table = 0; Table < TableCount; Table++)
+		{
+			Callback(Table);
+		}
+	}
+
+	void FTableManager::List(ETableType TableType, TFunctionRef<void(TableHandle)> Callback) const
+	{
+		int32 Table = 0;
+		for (const FTableInfo& TableInfo : Tables)
+		{
+			if (TableInfo.Type == TableType)
+			{
+				Callback(Table);
+			}
+			Table++;
+		}
+	}
+
+	void FTableManager::ListColumns(TableHandle Table, TFunctionRef<bool(const UScriptStruct*)> Callback) const
+	{
+		if (IsTableValid(Table))
+		{
+			int32 TableHandleAsIndex = ConvertTableHandleToIndex(Table);
+			MassEntityManager.ForEachArchetypeElementType(Tables[TableHandleAsIndex].Archetype, Callback);
+		}
+	}
+
+	bool FTableManager::HasColumns(TableHandle Table, TConstArrayView<const UScriptStruct*> Columns) const
+	{
+		if (IsTableValid(Table))
+		{
+			int32 TableHandleAsIndex = ConvertTableHandleToIndex(Table);
+			FMassArchetypeCompositionDescriptor Composition = MassEntityManager.GetArchetypeComposition(Tables[TableHandleAsIndex].Archetype);
+			for (const UScriptStruct* Column : Columns)
+			{
+				if (!Composition.Contains(Column))
+				{
+					return false;
+				}
+			}
+			return true;
+		}
+		else
+		{
+			return false;
+		}
+	}
+
+	FTableInfoView FTableManager::GetInfo(TableHandle Table) const
+	{
+		if (IsTableValid(Table))
+		{
+			int32 TableHandleAsIndex = ConvertTableHandleToIndex(Table);
+			const FTableInfo& Info = Tables[TableHandleAsIndex];
+			return FTableInfoView
+			{
+				.Variants = Info.Variants,
+				.Derivatives = Info.Derivatives,
+				.Name = Info.Name,
+				.Type = Info.Type,
+				.Parent = Info.Parent
+			};
+		}
+		else
+		{
+			return FTableInfoView();
+		}
+	}
+
+	void FTableManager::ListRows(TableHandle Table, TFunctionRef<void (FRowHandleArrayView)> Callback) const
+	{
+		if (FMassArchetypeHandle Archetype = GetArchetype(Table); Archetype.IsValid())
+		{
+			MassEntityManager.ForEachArchetypeChunkEntityList(Archetype,
+				[&Callback](TConstArrayView<FMassEntityHandle> EntityHandles)
+				{
+					Callback(FRowHandleArrayView(MassEntitiesToRowsConversion(EntityHandles), FRowHandleArrayView::EFlags::IsUnique));
+				});
+		}
+	}
+
+	void FTableManager::RegisterForeignKey(TableHandle Table, const FName& Domain, Queries::TConstQueryFunction<FForeignKey> KeyConstructor)
+	{
+		if (IsTableValid(Table))
+		{
+			int32 TableHandleAsIndex = ConvertTableHandleToIndex(Table);
+			Tables[TableHandleAsIndex].ForeignKeyConstructors.FindOrAdd(Domain) = MoveTemp(KeyConstructor);
+		}
+	}
+
+	Queries::EFlowControl FTableManager::BuildForeignKeys(TArrayView<FForeignKey> Results, const FName& Domain,
+		Queries::TForwardingConstQueryContext<Queries::CurrentTableInfo> Contract, Queries::IQueryFunctionResponse& Response) const
+	{
+		using namespace UE::Editor::DataStorage::Queries;
+
+		struct FQueryResult : TResult<FForeignKey>
+		{
+			TArrayView<FForeignKey> Results;
+			int32 Index = 0;
+			virtual void Add(RowHandle, FForeignKey Key) override
+			{
+				Results[Index++] = MoveTemp(Key);
+			}
+		};
+
+		FQueryResult QueryResult;
+		QueryResult.Results = Results;
+
+		TableHandle Table = Contract.GetCurrentTable();
+		if (IsTableValid(Table))
+		{
+			int32 TableHandleAsIndex = ConvertTableHandleToIndex(Table);
+			const FTableInfo* TableInfo = &Tables[TableHandleAsIndex];
+			const TConstQueryFunction<FForeignKey>* KeyConstructor = TableInfo->ForeignKeyConstructors.Find(Domain);
+			while(true)
+			{
+				if (KeyConstructor)
+				{
+					return KeyConstructor->Call<
+						EFunctionCallConfig::VerifyCapabilityCompatibility |
+						EFunctionCallConfig::ReportCapabilityCompatibility |
+						EFunctionCallConfig::VerifyColumns
+					>(QueryResult, Contract, Response);
+				}
+				else if (IsTableValid(TableInfo->Parent))
+				{
+					TableHandleAsIndex = ConvertTableHandleToIndex(TableInfo->Parent);
+					checkf(TableHandleAsIndex < Tables.Num(),
+						TEXT("Unable to retrieve internal archetype in editor storage because table handle %" PRIu64 " is invalid."), Table);
+					TableInfo = &Tables[TableHandleAsIndex];
+					KeyConstructor = TableInfo->ForeignKeyConstructors.Find(Domain);
+				}
+				else
+				{
+					break;
+				}
+			}
+		}
+		return Queries::EFlowControl::Break;
+	}
+
+	void FTableManager::ListForeignKeyDomains(TableHandle Table, bool bIncludeParents, TFunctionRef<void(const FName&)> Callback) const
+	{
+		using ForeignKeyMapConstIterator = FTableInfo::ForeignKeyMap::TConstIterator;
+
+		while (IsTableValid(Table))
+		{
+			int32 TableHandleAsIndex = ConvertTableHandleToIndex(Table);
+			const FTableInfo& TableInfo = Tables[TableHandleAsIndex];
+			for (ForeignKeyMapConstIterator It = TableInfo.ForeignKeyConstructors.CreateConstIterator(); It; ++It)
+			{
+				Callback(It.Key());
+			}
+
+			Table = bIncludeParents ? TableInfo.Parent : InvalidTableHandle;
+		}
+	}
+
+	FMassArchetypeHandle FTableManager::GetArchetype(TableHandle Table) const
+	{
+		return IsTableValid(Table)
+			? Tables[ConvertTableHandleToIndex(Table)].Archetype
+			: FMassArchetypeHandle();
+	}
+
+	void FTableManager::OnNewArchetype(const FMassArchetypeHandle& Archetype)
+	{
+		if (!bRegisteringNewTable)
+		{
+			const UScriptStruct* TableTag = nullptr;
+
+			FMassArchetypeCompositionDescriptor Composition = MassEntityManager.GetArchetypeComposition(Archetype);
+			const FMassElementBitSet& Columns = Composition.GetElementsBitSet();
+			Columns.ExportTypes([&TableTag](const UScriptStruct* TagType)
+				{
+					if (TagType->IsChildOf(FTableDynamicTag::StaticStruct()))
+					{
+						TableTag = TagType;
+						return false;
+					}
+					return true;
+				});
+
+			checkf(TableTag, TEXT("Expected a table tag to be on a variant archetype. This function shouldn't be called on new archetypes."));
+			const FString& ParentName = TableTag->GetMetaData(TEXT("EditorDataStorage_DerivedFromDynamicTemplate"));
+			FName  ParentFName = FName(ParentName);
+			TableHandle Parent = Find(ParentFName);
+			checkf(IsTableValid(Parent), TEXT("Unable to register a table variant because parent table '%s' wasn't found."), *ParentName);
+			int32 NewTable = Tables.Emplace(FTableInfo
+				{
+					.Archetype = Archetype,
+					.Name = ParentFName,
+					.Parent = Parent,
+					.Type = ETableType::Variant
+				});
+			TableHandle Variant = MassLookup.Add(Archetype, NewTable);
+			Tables[ConvertTableHandleToIndex(Parent)].Variants.Add(NewTable);
+
+			UE_LOG(LogEditorDataStorage, VeryVerbose, TEXT("New table variant for '%s'"), *ParentName);
+		}
+	}
+
+	int32 FTableManager::ConvertTableHandleToIndex(const TableHandle InTableHandle) const
+	{
+		checkf(
+			InTableHandle != InvalidTableHandle &&
+			(InTableHandle <= static_cast<TableHandle>(std::numeric_limits<int32>::max()) && 
+			InTableHandle >= 0),
+			TEXT("Invalid editor data storage table handle: %" PRIu64 "."), InTableHandle);
+		return static_cast<int32>(InTableHandle);
+	}
+
+	int32 FTableManager::GetTableChunkSize(const FName& TableName) const
+	{
+		const UEditorDataStorageSettings* Settings = GetDefault<UEditorDataStorageSettings>();
+		const EChunkMemorySize* TableSpecificSize = Settings->TableSpecificChunkMemorySize.Find(TableName);
+		return TableSpecificSize
+			? static_cast<int32>(*TableSpecificSize)
+			: static_cast<int32>(Settings->ChunkMemorySize);
+	}
+
+	void FTableManager::RemoveTableTag(FMassArchetypeCompositionDescriptor& Composition)
+	{
+		FMassElementBitSet& Columns = Composition.GetElementsBitSet();
+		Columns.ExportTypes([&Columns](const UScriptStruct* TagType)
+			{
+				if (TagType->IsChildOf(FTableDynamicTag::StaticStruct()))
+				{
+					Columns.Remove(TagType);
+					return false;
+				}
+				return true;
+			});
+	}
+} // namespace UE::Editor::DataStorage
diff --git a/Engine/Plugins/Experimental/EditorDataStorage/Source/TedsCore/Private/TableManager.h b/Engine/Plugins/Experimental/EditorDataStorage/Source/TedsCore/Private/TableManager.h
new file mode 100644
--- /dev/null
+++ b/Engine/Plugins/Experimental/EditorDataStorage/Source/TedsCore/Private/TableManager.h
@@ -0,0 +1,102 @@
+// Copyright Epic Games, Inc. All Rights Reserved.
+
+#pragma once
+
+#include "Containers/Array.h"
+#include "Containers/Map.h"
+#include "DataStorage/CommonTypes.h"
+#include "DataStorage/Handles.h"
+#include "Delegates/IDelegateInstance.h"
+#include "Elements/Framework/TypedElementQueryCapabilities.h"
+#include "Elements/Framework/TypedElementQueryContext.h"
+#include "Elements/Framework/TypedElementQueryFunctions.h"
+#include "Elements/Framework/TypedElementRowHandleArrayView.h"
+#include "MassArchetypeTypes.h"
+#include "UObject/NameTypes.h"
+
+#include "TableManager.generated.h"
+
+struct FMassArchetypeCompositionDescriptor;
+struct FMassEntityManager;
+class UScriptStruct;
+
+/**
+ * Base type for the tag that automatically gets added to any table to identify it.
+ */
+USTRUCT(meta = (DisplayName = "Table", EditorDataStorage_DynamicColumnTemplate))
+struct FEditorDataStorageTableDynamicTag final : public FEditorDataStorageTag
+{
+	GENERATED_BODY()
+};
+
+namespace UE::Editor::DataStorage
+{
+	using FTableDynamicTag = FEditorDataStorageTableDynamicTag;
+
+	class FDynamicColumnGenerator;
+
+	class FTableManager final
+	{
+	public:
+		FTableManager(FMassEntityManager& InMassEntityManager, FDynamicColumnGenerator& InDynamicColumns);
+		~FTableManager();
+
+		void Clear();
+
+		TableHandle Register(TConstArrayView<const UScriptStruct*> ColumnList, const FName& TableName);
+		TableHandle Register(
+			TConstArrayView<const UScriptStruct*> ColumnList, const FName& TableName, FTableRegistrationOptions Options);
+		
+		bool IsTableValid(TableHandle Table) const;
+
+		TableHandle Find(const FName& TableName) const;
+		TableHandle Find(const FMassArchetypeHandle& Archetype) const;
+		TableHandle Find(RowHandle Row) const;
+		
+		void List(TFunctionRef<void(TableHandle)> Callback) const;
+		void List(ETableType TableType, TFunctionRef<void(TableHandle)> Callback) const;
+		
+		void ListColumns(TableHandle Table, TFunctionRef<bool(const UScriptStruct*)> Callback) const;
+		bool HasColumns(TableHandle Table, TConstArrayView<const UScriptStruct*> Columns) const;
+		
+		FTableInfoView GetInfo(TableHandle Table) const;
+		void ListRows(TableHandle Table, TFunctionRef<void(FRowHandleArrayView)> Callback) const;
+
+		void RegisterForeignKey(TableHandle Table, const FName& Domain, Queries::TConstQueryFunction<FForeignKey> KeyConstructor);
+		Queries::EFlowControl BuildForeignKeys(TArrayView<FForeignKey> Results, const FName& Domain,
+			Queries::TForwardingConstQueryContext<Queries::CurrentTableInfo> Contract, Queries::IQueryFunctionResponse& Response) const;
+		void ListForeignKeyDomains(TableHandle Table, bool bIncludeParents, TFunctionRef<void(const FName&)> Callback) const;
+
+		FMassArchetypeHandle GetArchetype(TableHandle Table) const;
+	private:
+		struct FTableInfo
+		{
+			using ForeignKeyMap = TMap<FName, Queries::TConstQueryFunction<FForeignKey>>;
+			ForeignKeyMap ForeignKeyConstructors;
+			TArray<TableHandle> Variants;
+			TArray<TableHandle> Derivatives;
+			FMassArchetypeHandle Archetype;
+			FName Name;
+			// If this is a dynamically generated variant then the Parent represents the table used as the basis for the variant. 
+			// If this is a derived copy then the Parent represents the table that was passed in as the template for this table.
+			TableHandle Parent;
+			ETableType Type;
+		};
+
+		void OnNewArchetype(const FMassArchetypeHandle& Archetype);
+
+		int32 ConvertTableHandleToIndex(const TableHandle InTableHandle) const;
+		int32 GetTableChunkSize(const FName& TableName) const;
+		void RemoveTableTag(FMassArchetypeCompositionDescriptor& Composition);
+
+		TArray<FTableInfo> Tables; // Contains tables and variants.
+		TMap<FName, TableHandle> NameLookup; // Only includes tables.
+		TMap<FMassArchetypeHandle, TableHandle> MassLookup; // Includes tables and variants.
+
+		FDelegateHandle OnNewArchetypeEventHandle;
+		FMassEntityManager& MassEntityManager;
+		FDynamicColumnGenerator& DynamicColumns;
+
+		bool bRegisteringNewTable = false;
+	};
+} // namespace UE::Editor::DataStorage
diff --git a/Engine/Plugins/Experimental/EditorDataStorage/Source/TedsCore/Private/TypedElementDatabase.cpp b/Engine/Plugins/Experimental/EditorDataStorage/Source/TedsCore/Private/TypedElementDatabase.cpp
--- a/Engine/Plugins/Experimental/EditorDataStorage/Source/TedsCore/Private/TypedElementDatabase.cpp
+++ b/Engine/Plugins/Experimental/EditorDataStorage/Source/TedsCore/Private/TypedElementDatabase.cpp
@@ -6,7 +6,6 @@
 #include "DataStorage/Features.h"
 #include "DataStorage/Debug/Log.h"
 #include "Editor.h"
-#include "EditorDataStorageSettings.h"
 #include "Elements/Interfaces/TypedElementDataStorageFactory.h"
 #include "Elements/Interfaces/TypedElementDataStorageInterface.h"
 #include "Engine/World.h"
@@ -149,14 +148,9 @@ namespace UE::Editor::DataStorage::Private
 		}
 		return Result;
 	}
-
-	constexpr int32 ConvertTableHandleToIndex(const TableHandle InTableHandle)
-	{
-		return static_cast<int32>(InTableHandle);
-	}
 } // namespace UE::Editor::DataStorage::Private
 
-UEditorDataStorage::FActiveCallScope::FActiveCallScope(UEditorDataStorage& Owner)
+UEditorDataStorage::FActiveCallScope::FActiveCallScope(const UEditorDataStorage& Owner)
 	: Owner(Owner)
 {
 	Owner.ActiveCallsCount++;
@@ -363,62 +357,7 @@ TSharedPtr<const FMassEntityManager> UEditorDataStorage::GetActiveEditorEntityMa
 	return ActiveEditorEntityManager;
 }
 
-UE::Editor::DataStorage::TableHandle UEditorDataStorage::RegisterTable(TConstArrayView<const UScriptStruct*> ColumnList, const FName& InName)
-{
-	using namespace UE::Editor::DataStorage;
 
-	if (ActiveEditorEntityManager && !TableNameLookup.Contains(InName))
-	{
-		TableHandle Result = Tables.Num();
-		FMassArchetypeCreationParams ArchetypeCreationParams;
-		ArchetypeCreationParams.DebugName = InName;
-		ArchetypeCreationParams.ChunkMemorySize = GetTableChunkSize(InName);
-		Tables.Add(ActiveEditorEntityManager->CreateArchetype(ColumnList, ArchetypeCreationParams));
-		if (InName.IsValid())
-		{
-			TableNameLookup.Add(InName, Result);
-		}
-		return Result;
-	}
-	return InvalidTableHandle;
-}
-
-UE::Editor::DataStorage::TableHandle UEditorDataStorage::RegisterTable(UE::Editor::DataStorage::TableHandle SourceTable,
-	TConstArrayView<const UScriptStruct*> ColumnList, const FName& InName)
-{
-	using namespace UE::Editor::DataStorage;
-
-	if (ActiveEditorEntityManager && IsTableValid(SourceTable) && !TableNameLookup.Contains(InName))
-	{
-		const int32 TableHandleAsIndex = Private::ConvertTableHandleToIndex(SourceTable);
-		TableHandle Result = Tables.Num();
-		FMassArchetypeCreationParams ArchetypeCreationParams;
-		ArchetypeCreationParams.DebugName = InName;
-		ArchetypeCreationParams.ChunkMemorySize = GetTableChunkSize(InName);
-		Tables.Add(ActiveEditorEntityManager->CreateArchetype(Tables[TableHandleAsIndex], ColumnList, ArchetypeCreationParams));
-		if (InName.IsValid())
-		{
-			TableNameLookup.Add(InName, Result);
-		}
-		return Result;
-	}
-	return InvalidTableHandle;
-}
-
-UE::Editor::DataStorage::TableHandle UEditorDataStorage::FindTable(const FName& InName)
-{
-	using namespace UE::Editor::DataStorage;
-
-	TableHandle* TableHandle = TableNameLookup.Find(InName);
-	return TableHandle ? *TableHandle : InvalidTableHandle;
-}
-
-bool UEditorDataStorage::IsTableValid(UE::Editor::DataStorage::TableHandle Table) const
-{
-	using namespace UE::Editor::DataStorage;
-
-	return Table != InvalidTableHandle && Table < Tables.Num();
-}
 
 UE::Editor::DataStorage::RowHandle UEditorDataStorage::ReserveRow()
 {
@@ -453,15 +392,13 @@ void UEditorDataStorage::BatchReserveRows(TArrayView<UE::Editor::DataStorage::Ro
 
 UE::Editor::DataStorage::RowHandle UEditorDataStorage::AddRow(UE::Editor::DataStorage::TableHandle Table)
 {
-	using namespace UE::Editor::DataStorage;
-
+	checkf(ActiveEditorEntityManager.IsValid() && Environment.IsValid(), 
+		TEXT("Attempting to add a row to the editor data storage before it has been fully initialized."));
+	checkf(IsTableValid(Table), TEXT("Attempting to add a row to a non-existing table."));
 	FActiveCallScope Scope(*this);
+	return ActiveEditorEntityManager->CreateEntity(Environment->GetTableManager().GetArchetype(Table)).AsNumber();
+	
 
-	checkf(IsTableValid(Table), TEXT("Attempting to add a row to a non-existing table."));
-	const int32 TableHandleAsIndex = Private::ConvertTableHandleToIndex(Table);
-	return ActiveEditorEntityManager
-		? ActiveEditorEntityManager->CreateEntity(Tables[TableHandleAsIndex]).AsNumber()
-		: UE::Editor::DataStorage::InvalidRowHandle;
 }
 
 UE::Editor::DataStorage::RowHandle UEditorDataStorage::AddRow(UE::Editor::DataStorage::TableHandle Table,
@@ -469,134 +406,99 @@ UE::Editor::DataStorage::RowHandle UEditorDataStorage::AddRow(UE::Editor::DataSt
 {
 	using namespace UE::Editor::DataStorage;
 
-	if (ActiveEditorEntityManager)
-	{
-		FActiveCallScope Scope(*this);
-
-		checkf(IsTableValid(Table), TEXT("Attempting to a row to a non-existing table."));
-		const int32 TableHandleAsIndex = Private::ConvertTableHandleToIndex(Table);
+	checkf(ActiveEditorEntityManager.IsValid() && Environment.IsValid(),
+		TEXT("Attempting to add a row to the editor data storage before it has been fully initialized."));	
+	checkf(IsTableValid(Table), TEXT("Attempting to a row to a non-existing table."));
+	
+	FActiveCallScope Scope(*this);
 
-		TArray<FMassEntityHandle> Entity;
-		Entity.Reserve(1);
-		TSharedRef<FMassEntityManager::FEntityCreationContext> Context =
-			ActiveEditorEntityManager->BatchCreateEntities(Tables[TableHandleAsIndex], 1, Entity);
+	TArray<FMassEntityHandle> Entity;
+	Entity.Reserve(1);
+	TSharedRef<FMassEntityManager::FEntityCreationContext> Context =
+		ActiveEditorEntityManager->BatchCreateEntities(Environment->GetTableManager().GetArchetype(Table), 1, Entity);
 
-		checkf(!Entity.IsEmpty(), TEXT("Add row tried to create a new row but none were provided by the backend."));
-		RowHandle Result = Entity[0].AsNumber();
-		OnCreated(Entity[0].AsNumber());
-		return Result;
-	}
-	return InvalidRowHandle;
+	checkf(!Entity.IsEmpty(), TEXT("Add row tried to create a new row but none were provided by the backend."));
+	RowHandle Result = Entity[0].AsNumber();
+	OnCreated(Entity[0].AsNumber());
+	return Result;
 }
 
 bool UEditorDataStorage::AddRow(UE::Editor::DataStorage::RowHandle ReservedRow, UE::Editor::DataStorage::TableHandle Table)
 {
-	using namespace UE::Editor::DataStorage;
-
+	checkf(ActiveEditorEntityManager.IsValid() && Environment.IsValid(),
+		TEXT("Attempting to add a row to the editor data storage before it has been fully initialized."));
 	checkf(IsTableValid(Table), TEXT("Attempting to add a row to a non-existing table."));
-	const int32 TableHandleAsIndex = Private::ConvertTableHandleToIndex(Table);
 	checkf(!IsRowAssigned(ReservedRow), TEXT("Attempting to assign a table to row that already has a table assigned."));
-	if (ActiveEditorEntityManager)
-	{
-		FActiveCallScope Scope(*this);
 
-		ActiveEditorEntityManager->BuildEntity(FMassEntityHandle::FromNumber(ReservedRow), Tables[TableHandleAsIndex]);
-		return true;
-	}
-	else
-	{
-		return false;
-	}
+	FActiveCallScope Scope(*this);
+	ActiveEditorEntityManager->BuildEntity(
+		FMassEntityHandle::FromNumber(ReservedRow), 
+		Environment->GetTableManager().GetArchetype(Table));
+		
+	return true;
 }
 
 bool UEditorDataStorage::AddRow(UE::Editor::DataStorage::RowHandle ReservedRow, UE::Editor::DataStorage::TableHandle Table,
 	UE::Editor::DataStorage::RowCreationCallbackRef OnCreated)
 {
-	using namespace UE::Editor::DataStorage;
-
-	if (ActiveEditorEntityManager)
-	{
-		FActiveCallScope Scope(*this);
-
-		checkf(IsTableValid(Table), TEXT("Attempting to add a row to a non-existing table."));
-		const int32 TableHandleAsIndex = Private::ConvertTableHandleToIndex(Table);
+	checkf(ActiveEditorEntityManager.IsValid() && Environment.IsValid(),
+		TEXT("Attempting to add a row to the editor data storage before it has been fully initialized."));	
+	checkf(IsTableValid(Table), TEXT("Attempting to add a row to a non-existing table."));
 
-		TSharedRef<FMassEntityManager::FEntityCreationContext> Context =
-			ActiveEditorEntityManager->BatchCreateReservedEntities(Tables[TableHandleAsIndex], { FMassEntityHandle::FromNumber(ReservedRow) });
+	FActiveCallScope Scope(*this);
+	TSharedRef<FMassEntityManager::FEntityCreationContext> Context =
+		ActiveEditorEntityManager->BatchCreateReservedEntities(
+			Environment->GetTableManager().GetArchetype(Table), 
+			{ FMassEntityHandle::FromNumber(ReservedRow) });
 
-		OnCreated(ReservedRow);
-		return true;
-	}
-	return false;
+	OnCreated(ReservedRow);
+	return true;
 }
 
 bool UEditorDataStorage::BatchAddRow(
 	UE::Editor::DataStorage::TableHandle Table, int32 Count, UE::Editor::DataStorage::RowCreationCallbackRef OnCreated)
 {
-	using namespace UE::Editor::DataStorage;
-
-	if (ActiveEditorEntityManager)
-	{
-		FActiveCallScope Scope(*this);
-
-		if (!IsTableValid(Table))
-		{
-			checkf(
-				false,
-				TEXT("Attempting to add multiple rows to a non-existing table (%" PRIu64 "). Tables.Num(): %d, TableName: %s"),
-				Table, Tables.Num(), *FindTableName(Table)
-			);
-			return false;
-		}
-
-		const int32 TableHandleAsIndex = Private::ConvertTableHandleToIndex(Table);
-		TArray<FMassEntityHandle> Entities;
-		Entities.Reserve(Count);
-		TSharedRef<FMassEntityManager::FEntityCreationContext> Context =
-			ActiveEditorEntityManager->BatchCreateEntities(Tables[TableHandleAsIndex], Count, Entities);
+	checkf(ActiveEditorEntityManager.IsValid() && Environment.IsValid(),
+		TEXT("Attempting to batch add rows to the editor data storage before it has been fully initialized."));
+	checkf(IsTableValid(Table), TEXT("Attempting to add rows to a non-existing table."));
 
-		for (FMassEntityHandle Entity : Entities)
-		{
-			OnCreated(Entity.AsNumber());
-		}
+	FActiveCallScope Scope(*this);
 
-		return true;
+	TArray<FMassEntityHandle> Entities;
+	Entities.Reserve(Count);
+	TSharedRef<FMassEntityManager::FEntityCreationContext> Context = 
+		ActiveEditorEntityManager->BatchCreateEntities(Environment->GetTableManager().GetArchetype(Table), Count, Entities);
+		
+	for (FMassEntityHandle Entity : Entities)
+	{
+		OnCreated(Entity.AsNumber());
 	}
-	return false;
+
+	return true;
 }
 
 bool UEditorDataStorage::BatchAddRow(UE::Editor::DataStorage::TableHandle Table,
 	TConstArrayView<UE::Editor::DataStorage::RowHandle> ReservedHandles, UE::Editor::DataStorage::RowCreationCallbackRef OnCreated)
 {
-	using namespace UE::Editor;
 	using namespace UE::Editor::DataStorage;
 
-	if (ActiveEditorEntityManager)
-	{
-		FActiveCallScope Scope(*this);
+	checkf(ActiveEditorEntityManager.IsValid() && Environment.IsValid(),
+		TEXT("Attempting to batch add rows to the editor data storage before it has been fully initialized."));
+	checkf(IsTableValid(Table), TEXT("Attempting to add rows to a non-existing table."));
 
-		if (!IsTableValid(Table))
-		{
-			checkf(
-				false,
-				TEXT("Attempting to add multiple reserved rows to a non-existing table (%" PRIu64 "). Tables.Num(): %d, TableName: %s"),
-				Table, Tables.Num(), *FindTableName(Table)
-			);
-			return false;
-		}
+	FActiveCallScope Scope(*this);
 
-		const int32 TableHandleAsIndex = Private::ConvertTableHandleToIndex(Table);
-		TSharedRef<FMassEntityManager::FEntityCreationContext> Context =
-			ActiveEditorEntityManager->BatchCreateReservedEntities(Tables[TableHandleAsIndex], RowsToMassEntitiesConversion(ReservedHandles));
+	TSharedRef<FMassEntityManager::FEntityCreationContext> Context =
+		ActiveEditorEntityManager->BatchCreateReservedEntities(
+			Environment->GetTableManager().GetArchetype(Table), 
+			RowsToMassEntitiesConversion(ReservedHandles));
 
-		for (RowHandle Entity : ReservedHandles)
-		{
-			OnCreated(Entity);
-		}
-
-		return true;
+	for (RowHandle Entity : ReservedHandles)
+	{
+		OnCreated(Entity);
 	}
-	return false;
+
+	return true;
 }
 
 
@@ -693,44 +595,49 @@ bool UEditorDataStorage::IsRowAssignedUnsafe(UE::Editor::DataStorage::RowHandle
 	return UE::Editor::DataStorage::Legacy::FCommandBuffer::Execute_IsRowAssigned(*ActiveEditorEntityManager, Row);
 }
 
-void UEditorDataStorage::FilterRowsBy(UE::Editor::DataStorage::FRowHandleArray& Result, UE::Editor::DataStorage::FRowHandleArrayView Input,
-	EFilterOptions Options, UE::Editor::DataStorage::Queries::TQueryFunction<bool>& Filter)
+void UEditorDataStorage::FilterRowsBy(UE::Editor::DataStorage::FRowHandleArray& Result,  UE::Editor::DataStorage::FRowHandleArrayView Input, 
+	EFilterOptions Options, UE::Editor::DataStorage::Queries::TConstQueryFunction<bool>& Filter) const
 {
 	using namespace UE::Editor::DataStorage;
 	using namespace UE::Editor::DataStorage::Queries;
 
-	if (ActiveEditorEntityManager && Environment)
+	if (ActiveEditorEntityManager && Environment && !Input.IsEmpty())
 	{
 		struct FQueryResult : TResult<bool>
 		{
-			bool bResult;
-			explicit FQueryResult(const bool bResultValue)
-				: bResult(bResultValue)
-			{
-			}
+			FRowHandleArray& Result;
+			bool bAcceptedResult;
+
+			explicit FQueryResult(FRowHandleArray& Result, bool bAcceptedResult)
+				: Result(Result)
+				, bAcceptedResult(bAcceptedResult)
+			{}
 
-			virtual void Add(bool bResultValue) override
+			virtual void Add(RowHandle Row, bool bResultValue) override
 			{
-				bResult = bResultValue;
+				if (bResultValue == bAcceptedResult)
+				{
+					Result.Add(Row);
+				}
 			}
 		};
 
-		FDirectContextEnvironment LocalContextEnvironment(*Environment, *this);
-		QueryContext_DirectApi Context(LocalContextEnvironment);
-		
-		const bool bVerifyColumns = !EnumHasAnyFlags(Options, EFilterOptions::IncludeRowsMissingColumns);
 		const bool bAcceptedResult = !EnumHasAnyFlags(Options, EFilterOptions::InvertFilter);
+		constexpr EFunctionCallConfig CallingFlags = 
+			EFunctionCallConfig::VerifyColumns | 
+			EFunctionCallConfig::VerifyCapabilityCompatibility |
+			EFunctionCallConfig::ReportCapabilityCompatibility;
 
+		FConstDirectContextEnvironment LocalContextEnvironment(*Environment, *this);
+		QueryContext_ConstDirectApi Context(LocalContextEnvironment);
 		for (RowHandle Row : Input)
 		{
-			LocalContextEnvironment.SetBatch(FRowHandleArrayView(&Row, 1,
+			LocalContextEnvironment.SetBatch(FRowHandleArrayView(&Row, 1, 
 				FRowHandleArrayView::EFlags::IsSorted | FRowHandleArrayView::EFlags::IsUnique));
-
-			FQueryResult QueryResult(!bVerifyColumns);
-			Filter.Call<EFunctionCallConfig::VerifyColumns>(QueryResult, Context, LocalContextEnvironment);
-			if (QueryResult.bResult == bAcceptedResult)
+			FQueryResult QueryResult(Result, bAcceptedResult);
+			if (Filter.Call<CallingFlags>(QueryResult, Context, LocalContextEnvironment) == EFlowControl::Break)
 			{
-				Result.Add(Row);
+				break;
 			}
 		}
 	}
@@ -1456,6 +1363,27 @@ UE::Editor::DataStorage::FQueryResult UEditorDataStorage::RunQuery(
 	}
 }
 
+UE::Editor::DataStorage::FQueryResult UEditorDataStorage::RunQuery(
+	UE::Editor::DataStorage::QueryHandle Query,
+	UE::Editor::DataStorage::ERunQueryFlags InFlags,
+	const UE::Editor::DataStorage::Queries::TConstQueryFunction<void>& Callback) const
+{
+	using namespace UE::Editor::DataStorage;
+
+	TRACE_CPUPROFILER_EVENT_SCOPE(TEDS.RunQuery);
+
+	if (ActiveEditorEntityManager)
+	{
+		FActiveCallScope Scope(*this);
+		return Environment->GetQueryStore().RunQuery(*ActiveEditorEntityManager, *Environment,
+			FExtendedQueryStore::Handle(Query), InFlags, Callback);
+	}
+	else
+	{
+		return FQueryResult();
+	}
+}
+
 void UEditorDataStorage::ActivateQueries(FName ActivationName)
 {
 	if (ActiveEditorEntityManager)
@@ -1693,18 +1621,6 @@ void UEditorDataStorage::TickCooperativeTasks()
 	}
 }
 
-FString UEditorDataStorage::FindTableName(UE::Editor::DataStorage::TableHandle InHandle) const
-{
-	for (const TPair<FName, UE::Editor::DataStorage::TableHandle>& Pair : TableNameLookup)
-	{
-		if (Pair.Value == InHandle)
-		{
-			return Pair.Key.ToString();
-		}
-	}
-	return TEXT("<unknown>");
-};
-
 bool UEditorDataStorage::IsAvailable() const
 {
 	return bool(ActiveEditorEntityManager);
@@ -1789,9 +1705,8 @@ void UEditorDataStorage::Reset()
 	if (ActiveEditorEntityManager && ActiveEditorPhaseManager)
 	{
 		Environment->GetQueryStore().Clear(*ActiveEditorEntityManager.Get(), *ActiveEditorPhaseManager.Get());
+		Environment->GetTableManager().Clear();
 	}
-	Tables.Reset();
-	TableNameLookup.Reset();
 	
 	checkf(Environment.IsUnique(), TEXT("UEditorDataStorage should hold the last reference to the Environment and be the one to delete it."));
 	Environment.Reset();
@@ -1802,19 +1717,6 @@ void UEditorDataStorage::Reset()
 	UE_LOGF(LogEditorDataStorage, Log, "Reset TEDS Core.");
 }
 
-int32 UEditorDataStorage::GetTableChunkSize(FName TableName) const
-{
-	const UEditorDataStorageSettings* Settings = GetDefault<UEditorDataStorageSettings>();
-	if (const EChunkMemorySize* TableSpecificSize = Settings->TableSpecificChunkMemorySize.Find(TableName))
-	{
-		return static_cast<int32>(*TableSpecificSize);
-	}
-	else
-	{
-		return static_cast<int32>(Settings->ChunkMemorySize);
-	}
-}
-
 TSharedPtr<UE::Editor::DataStorage::FEnvironment> UEditorDataStorage::GetEnvironmentPtr()
 {
 	return Environment;
@@ -1837,19 +1739,6 @@ const UE::Editor::DataStorage::FEnvironment& UEditorDataStorage::GetEnvironment(
 	return *Environment;
 }
 
-
-FMassArchetypeHandle UEditorDataStorage::LookupArchetype(UE::Editor::DataStorage::TableHandle InTableHandle) const
-{
-	using namespace UE::Editor::DataStorage;
-
-	if (IsTableValid(InTableHandle))
-	{
-		const int32 TableHandleAsIndex = Private::ConvertTableHandleToIndex(InTableHandle);
-		return Tables[TableHandleAsIndex];
-	}
-	return FMassArchetypeHandle();
-}
-
 void UEditorDataStorage::DebugPrintQueryCallbacks(FOutputDevice& Output)
 {
 	Environment->GetQueryStore().DebugPrintQueryCallbacks(Output);
diff --git a/Engine/Plugins/Experimental/EditorDataStorage/Source/TedsCore/Private/TypedElementDatabase.h b/Engine/Plugins/Experimental/EditorDataStorage/Source/TedsCore/Private/TypedElementDatabase.h
--- a/Engine/Plugins/Experimental/EditorDataStorage/Source/TedsCore/Private/TypedElementDatabase.h
+++ b/Engine/Plugins/Experimental/EditorDataStorage/Source/TedsCore/Private/TypedElementDatabase.h
@@ -11,6 +11,7 @@
 #include "TypedElementDatabaseCommandBuffer.h"
 #include "UObject/ObjectMacros.h"
 #include "Queries/QueryContextCapabilities/EditorDataStorageDynamicColumnInfo.h"
+#include "Queries/QueryContextCapabilities/TableCapabilities.h"
 #include "Queries/TypedElementQueryContext_DirectSingleApi.h"
 
 #include "TypedElementDatabase.generated.h"
@@ -31,6 +32,8 @@ namespace UE::Editor::DataStorage
 using ImplementationForwarder = UE::Editor::DataStorage::Queries::Private::TImplementationForwarder<
 	UE::Editor::DataStorage::Queries::FDirectContextEnvironment,
 	UE::Editor::DataStorage::ICoreProvider,
+	UE::Editor::DataStorage::Queries::FTableInfoContextCapability,
+	UE::Editor::DataStorage::Queries::FTableManagementContextCapability,
 	UE::Editor::DataStorage::Queries::FDynamicColumnInfoContextCapability>;
 
 UCLASS(MinimalAPI)
@@ -93,11 +96,6 @@ public:
 	UE_API TSharedPtr<FMassEntityManager> GetActiveMutableEditorEntityManager();
 	UE_API TSharedPtr<const FMassEntityManager> GetActiveEditorEntityManager() const;
 
-	UE_API virtual UE::Editor::DataStorage::TableHandle RegisterTable(TConstArrayView<const UScriptStruct*> ColumnList, const FName& Name) override;
-	UE_API virtual UE::Editor::DataStorage::TableHandle RegisterTable(
-		UE::Editor::DataStorage::TableHandle SourceTable, TConstArrayView<const UScriptStruct*> ColumnList, const FName& Name) override;
-	UE_API virtual UE::Editor::DataStorage::TableHandle FindTable(const FName& Name) override;
-	UE_API virtual bool IsTableValid(UE::Editor::DataStorage::TableHandle Table) const override;
 
 	UE_API virtual UE::Editor::DataStorage::RowHandle ReserveRow() override;
 	UE_API virtual void BatchReserveRows(int32 Count, TFunctionRef<void(UE::Editor::DataStorage::RowHandle)> ReservationCallback) override;
@@ -122,7 +120,7 @@ public:
 	/** Same as IsRowAssigned, but doesn't check if the data storage has been initialized. */
 	UE_API bool IsRowAssignedUnsafe(UE::Editor::DataStorage::RowHandle Row) const;
 	virtual void FilterRowsBy(UE::Editor::DataStorage::FRowHandleArray& Result, UE::Editor::DataStorage::FRowHandleArrayView Input, 
-		EFilterOptions Options, UE::Editor::DataStorage::Queries::TQueryFunction<bool>& Filter) override;
+		EFilterOptions Options, UE::Editor::DataStorage::Queries::TConstQueryFunction<bool>& Filter) const override;
 
 	UE_API virtual void AddColumn(UE::Editor::DataStorage::RowHandle Row, const UScriptStruct* ColumnType) override;
 	UE_API virtual void AddColumn(UE::Editor::DataStorage::RowHandle Row, const UE::Editor::DataStorage::FValueTag& Tag, const FName& InValue) override;
@@ -215,6 +213,8 @@ public:
 		UE::Editor::DataStorage::DirectQueryCallbackRef Callback) override;
 	UE_API virtual UE::Editor::DataStorage::FQueryResult RunQuery(UE::Editor::DataStorage::QueryHandle Query, UE::Editor::DataStorage::ERunQueryFlags Flags,
 		const UE::Editor::DataStorage::Queries::TQueryFunction<void>& Callback) override;
+	UE_API virtual UE::Editor::DataStorage::FQueryResult RunQuery(UE::Editor::DataStorage::QueryHandle Query, UE::Editor::DataStorage::ERunQueryFlags Flags,
+		const UE::Editor::DataStorage::Queries::TConstQueryFunction<void>& Callback) const override;
 	UE_API virtual void ActivateQueries(FName ActivationName) override;
 	
 	UE_API virtual UE::Editor::DataStorage::RowHandle LookupMappedRow(
@@ -244,8 +244,6 @@ public:
 	UE::Editor::DataStorage::FEnvironment& GetEnvironment();
 	const UE::Editor::DataStorage::FEnvironment& GetEnvironment() const;
 
-	UE_API FMassArchetypeHandle LookupArchetype(UE::Editor::DataStorage::TableHandle InTableHandle) const;
-
 	UE_API void DebugPrintQueryCallbacks(FOutputDevice& Output) override;
 
 	static UE_API void AddReferencedObjects(UObject* InThis, FReferenceCollector& Collector);
@@ -255,11 +253,8 @@ private:
 	void FinalizePhase(UE::Editor::DataStorage::EQueryTickPhase Phase, float DeltaTime);
 	void Reset();
 
-	int32 GetTableChunkSize(FName TableName) const;
 	void TickCooperativeTasks();
-	// Reverse-lookup name from handle (O(n), only use for diagnostics, slow!)
-	FString FindTableName(UE::Editor::DataStorage::TableHandle InHandle) const;
-
+	
 	struct FFactoryTypePair
 	{
 		// Used to find the factory by type without needing to dereference each one
@@ -271,7 +266,7 @@ private:
 	// Wrapper around ActiveCallsCount to avoid issues with timing such as early outs.
 	struct FActiveCallScope
 	{
-		explicit FActiveCallScope(UEditorDataStorage& Owner);
+		explicit FActiveCallScope(const UEditorDataStorage& Owner);
 		~FActiveCallScope();
 
 		FActiveCallScope(const FActiveCallScope&) = delete;
@@ -279,7 +274,7 @@ private:
 		FActiveCallScope& operator=(const FActiveCallScope&) = delete;
 		FActiveCallScope& operator=(FActiveCallScope&&) = delete;
 
-		UEditorDataStorage& Owner;
+		const UEditorDataStorage& Owner;
 	};
 	
 	static UE_API const FName TickGroupName_Default;
@@ -290,9 +285,6 @@ private:
 	static UE_API const FName TickGroupName_SyncExternalToDataStorage;
 	static UE_API const FName TickGroupName_SyncDataStorageToExternal;
 	
-	TArray<FMassArchetypeHandle> Tables;
-	TMap<FName, UE::Editor::DataStorage::TableHandle> TableNameLookup;
-
 	// Ordered array of factories by the return value of GetOrder()
 	TArray<FFactoryTypePair> Factories;
 
@@ -325,7 +317,7 @@ private:
 
 	// Track the number of active calls, incrementing whenever a function calls another direct API call. This prevents
 	// the command buffer from being flushed early and guarantees it only being flushed when the call chain has completed.
-	uint32 ActiveCallsCount = 0;
+	mutable uint32 ActiveCallsCount = 0;
 };
 
 template <typename FactoryType, typename DatabaseType>
diff --git a/Engine/Plugins/Experimental/EditorDataStorage/Source/TedsCore/Private/TypedElementDatabaseCompatibility.cpp b/Engine/Plugins/Experimental/EditorDataStorage/Source/TedsCore/Private/TypedElementDatabaseCompatibility.cpp
--- a/Engine/Plugins/Experimental/EditorDataStorage/Source/TedsCore/Private/TypedElementDatabaseCompatibility.cpp
+++ b/Engine/Plugins/Experimental/EditorDataStorage/Source/TedsCore/Private/TypedElementDatabaseCompatibility.cpp
@@ -16,6 +16,7 @@
 #include "Elements/Columns/TypedElementTransformColumns.h"
 #include "Elements/Columns/TypedElementTypeInfoColumns.h"
 #include "Elements/Framework/TypedElementQueryBuilder.h"
+#include "Elements/Framework/TypedElementQueryContext.h"
 #include "HAL/IConsoleManager.h"
 #include "Logging/LogMacros.h"
 #include "Memento/TypedElementMementoRowTypes.h"
@@ -465,19 +466,45 @@ void UEditorDataStorageCompatibility::Prepare()
 
 void UEditorDataStorageCompatibility::CreateStandardArchetypes()
 {
+	using namespace UE::Editor::DataStorage::Queries;
+
 	// Thread-safe as this is only called from a function that has an exclusive lock.
 	using namespace UE::Editor::DataStorage;
-	StandardUObjectTable = Storage->RegisterTable(TTypedElementColumnTypeList<
+	StandardUObjectTable = Storage->RegisterTable<
 			FTypedElementUObjectColumn, FTypedElementUObjectIdColumn, FUObjectIdNameColumn,
-			FTypedElementClassTypeInfoColumn, FTypedElementSyncFromWorldTag>(),
-		FName("Editor_StandardUObjectTable"));
-
-	StandardExternalObjectTable = Storage->RegisterTable(TTypedElementColumnTypeList<
-			FTypedElementExternalObjectColumn, FTypedElementScriptStructTypeInfoColumn,
-			FTypedElementSyncFromWorldTag>(), 
-		FName("Editor_StandardExternalObjectTable"));
+			FTypedElementClassTypeInfoColumn, FTypedElementSyncFromWorldTag>
+		(FName("Editor_StandardUObjectTable"));
 
+	Storage->RegisterTableForeignKey(StandardUObjectTable, FName("Serialization"), 
+			[](TConstQueryContext<RowBatchInfo> Context, TResult<FForeignKey>& Result, TConstBatch<FTypedElementUObjectColumn> Objects)
+			{
+				Context.ForEachRow([&Result](RowHandle Row, const FTypedElementUObjectColumn& Object)
+					{
+						const UObject* ObjectPtr = Object.Object.Get(true);
+						Result.Add(Row, ObjectPtr
+							? FForeignKey(ObjectPtr->GetPathName())
+							: FForeignKey(FKeyIssue(TEXT("Row doesn't contain a valid UObject.")))
+						);
+					}, Objects);
+			});
+	Storage->RegisterTableForeignKey(StandardUObjectTable, FName("Memento"),
+			[](TConstQueryContext<RowBatchInfo> Context, TResult<FForeignKey>& Result, TConstBatch<FTypedElementUObjectColumn> Objects)
+			{
+				Context.ForEachRow([&Result](RowHandle Row, const FTypedElementUObjectColumn& Object)
+					{
+						const UObject* ObjectPtr = Object.Object.Get(true);
+						Result.Add(Row, ObjectPtr
+							? FForeignKey(ObjectPtr)
+							: FForeignKey(FKeyIssue(TEXT("Row doesn't contain a valid UObject.")))
+						);
+					}, Objects);
+			});
 	RegisterTypeTableAssociation(UObject::StaticClass(), StandardUObjectTable);
+
+	StandardExternalObjectTable = Storage->RegisterTable<
+			FTypedElementExternalObjectColumn, FTypedElementScriptStructTypeInfoColumn,
+			FTypedElementSyncFromWorldTag>
+		(FName("Editor_StandardExternalObjectTable"));
 }
 
 void UEditorDataStorageCompatibility::RegisterTypeInformationQueries()
diff --git a/Engine/Plugins/Experimental/EditorDataStorage/Source/TedsCore/Private/TypedElementDatabaseEnvironment.cpp b/Engine/Plugins/Experimental/EditorDataStorage/Source/TedsCore/Private/TypedElementDatabaseEnvironment.cpp
--- a/Engine/Plugins/Experimental/EditorDataStorage/Source/TedsCore/Private/TypedElementDatabaseEnvironment.cpp
+++ b/Engine/Plugins/Experimental/EditorDataStorage/Source/TedsCore/Private/TypedElementDatabaseEnvironment.cpp
@@ -9,11 +9,11 @@ namespace UE::Editor::DataStorage
 {
 	FEnvironment::FEnvironment(UEditorDataStorage& InDataStorage,
 		FMassEntityManager& InMassEntityManager, FMassProcessingPhaseManager& InMassPhaseManager)
-		: DataStorage(InDataStorage)
-		, DirectDeferredCommands(*this)
+		: DirectDeferredCommands(*this)
+		, TableManager(InMassEntityManager, DynamicColumnGenerator)
 		, MappingTable(InDataStorage)
 		, Queries(DynamicColumnGenerator)
-		, MementoSystem(InDataStorage)
+		, MementoSystem(InDataStorage, TableManager)
 		, ScopeSystem(InDataStorage)
 		, ValueTagManager(InMassEntityManager, DynamicColumnGenerator)
 		, ActiveCommandQueue(CommandQueues)
@@ -109,11 +109,6 @@ namespace UE::Editor::DataStorage
 		return HierarchyInterface;
 	}
 
-	FMassArchetypeHandle FEnvironment::LookupMassArchetype(TableHandle Table) const
-	{
-		return DataStorage.LookupArchetype(Table);
-	}
-
 	FMassProcessingPhaseManager& FEnvironment::GetMassPhaseManager()
 	{
 		return MassPhaseManager;
@@ -134,6 +129,16 @@ namespace UE::Editor::DataStorage
 		return SharedMassCommandBuffer;
 	}
 
+	FTableManager& FEnvironment::GetTableManager()
+	{
+		return TableManager;
+	}
+
+	const FTableManager& FEnvironment::GetTableManager() const
+	{
+		return TableManager;
+	}
+
 	const UScriptStruct* FEnvironment::FindDynamicColumn(const UScriptStruct& Template, const FName& Identifier) const
 	{
 		return DynamicColumnGenerator.FindByTemplateId(Template, Identifier);
diff --git a/Engine/Plugins/Experimental/EditorDataStorage/Source/TedsCore/Private/TypedElementDatabaseEnvironment.h b/Engine/Plugins/Experimental/EditorDataStorage/Source/TedsCore/Private/TypedElementDatabaseEnvironment.h
--- a/Engine/Plugins/Experimental/EditorDataStorage/Source/TedsCore/Private/TypedElementDatabaseEnvironment.h
+++ b/Engine/Plugins/Experimental/EditorDataStorage/Source/TedsCore/Private/TypedElementDatabaseEnvironment.h
@@ -6,6 +6,7 @@
 #include "DynamicColumnGenerator.h"
 #include "MassEntityManager.h"
 #include "MassProcessingPhaseManager.h"
+#include "TableManager.h"
 #include "Templates/SharedPointer.h"
 #include "TypedElementDatabaseCommandBuffer.h"
 #include "TypedElementDatabaseScratchBuffer.h"
@@ -52,14 +53,15 @@ namespace UE::Editor::DataStorage
 		FTedsHierarchyRegistrar& GetHierarchyRegistrar();
 		const FTedsHierarchyRegistrar& GetHierarchyRegistrar() const;
 
-		FMassArchetypeHandle LookupMassArchetype(TableHandle Table) const;
-
 		FMassProcessingPhaseManager& GetMassPhaseManager();
 		const FMassProcessingPhaseManager& GetMassPhaseManager() const;
 
 		TSharedPtr<FMassCommandBuffer> GetSharedMassCommandBuffer();
 		TSharedPtr<const FMassCommandBuffer> GetSharedMassCommandBuffer() const;
 
+		FTableManager& GetTableManager();
+		const FTableManager& GetTableManager() const;
+
 		// Finds the type information for the dynamic column
 		// Dynamic columns are specified by a template layout and a FName Identifier
 		const UScriptStruct* FindDynamicColumn(const UScriptStruct& Template, const FName& Identifier) const;
@@ -91,18 +93,17 @@ namespace UE::Editor::DataStorage
 		void FlushCommands();
 
 	private:
-		
-		UEditorDataStorage& DataStorage;
 		Legacy::FCommandBuffer DirectDeferredCommands;
+		FDynamicColumnGenerator DynamicColumnGenerator;
+		FTableManager TableManager;
 		FMappingTable MappingTable;
 		FScratchBuffer ScratchBuffer;
-		FDynamicColumnGenerator DynamicColumnGenerator;
 		FExtendedQueryStore Queries;
 		FMementoSystem MementoSystem;
 		FScopeSystem ScopeSystem;
 		FValueTagManager ValueTagManager;
 		FTedsHierarchyRegistrar HierarchyInterface;
-
+		
 		FMutex CommandQueueMutex;
 		TArray<FEnvironmentCommand> CommandQueues[2];
 		TArray<FEnvironmentCommand>* ActiveCommandQueue;
@@ -114,7 +115,10 @@ namespace UE::Editor::DataStorage
 		uint64 UpdateCycleId = 0;
 	};
 
-	class IQueryEnvironmentRetrieval
+	struct FEmptyEnvironmentRetrievalBase {};
+
+	template<typename Base = FEmptyEnvironmentRetrievalBase>
+	class IQueryEnvironmentRetrieval : public Base
 	{
 	public:
 		virtual ~IQueryEnvironmentRetrieval() = default;
@@ -123,4 +127,6 @@ namespace UE::Editor::DataStorage
 		virtual FEnvironment& GetEnvironment() = 0;
 		virtual const FEnvironment& GetEnvironment() const = 0;
 	};
+
+	using QueryEnvironmentRetrieval = IQueryEnvironmentRetrieval<>;
 } // namespace UE::Editor::DataStorage
diff --git a/Engine/Plugins/Experimental/EditorDataStorageFeatures/Source/TedsActorCompatibility/Private/Compatibility/TedsActorComponentFactory.cpp b/Engine/Plugins/Experimental/EditorDataStorageFeatures/Source/TedsActorCompatibility/Private/Compatibility/TedsActorComponentFactory.cpp
--- a/Engine/Plugins/Experimental/EditorDataStorageFeatures/Source/TedsActorCompatibility/Private/Compatibility/TedsActorComponentFactory.cpp
+++ b/Engine/Plugins/Experimental/EditorDataStorageFeatures/Source/TedsActorCompatibility/Private/Compatibility/TedsActorComponentFactory.cpp
@@ -29,9 +29,13 @@ void UTedsActorComponentFactory::RegisterTables(UE::Editor::DataStorage::ICorePr
 	
 	if (bEnableActorComponentCompatibility)
 	{
-		TableHandle ActorComponentTable = DataStorage.RegisterTable<FActorComponentTypeTag, FTypedElementLabelColumn, FLevelColumn, FTypedElementLabelHashColumn>(
-											DataStorage.FindTable(FName(TEXT("Editor_StandardUObjectTable"))),
-											FName("Editor_ActorComponentTable"));
+		TableHandle ActorComponentTable = 
+			DataStorage.RegisterTable<FActorComponentTypeTag, FLevelColumn, FTypedElementLabelColumn, FTypedElementLabelHashColumn>(
+				FName("Editor_ActorComponentTable"),
+				FTableRegistrationOptions
+				{ 
+					.SourceTable = DataStorage.FindTable(FName(TEXT("Editor_StandardUObjectTable"))) 
+				});
 		
 		if (ICompatibilityProvider* CompatibilityProvider = GetMutableDataStorageFeature<ICompatibilityProvider>(CompatibilityFeatureName))
 		{
diff --git a/Engine/Plugins/Experimental/EditorDataStorageFeatures/Source/TedsAlerts/Private/TedsAlertsFactory.cpp b/Engine/Plugins/Experimental/EditorDataStorageFeatures/Source/TedsAlerts/Private/TedsAlertsFactory.cpp
--- a/Engine/Plugins/Experimental/EditorDataStorageFeatures/Source/TedsAlerts/Private/TedsAlertsFactory.cpp
+++ b/Engine/Plugins/Experimental/EditorDataStorageFeatures/Source/TedsAlerts/Private/TedsAlertsFactory.cpp
@@ -15,8 +15,11 @@ const FName UTedsAlertsFactory::UnsortedAlertChainTableName = "Alerts chain (uns
 
 void UTedsAlertsFactory::RegisterTables(UE::Editor::DataStorage::ICoreProvider& DataStorage)
 {
+	using namespace UE::Editor::DataStorage;
+
 	ChainTable = DataStorage.RegisterTable<FTedsAlertColumn, FTedsAlertChainTag>(AlertChainTableName);
-	UnsortedChainTable = DataStorage.RegisterTable<FTedsUnsortedAlertChainTag>(ChainTable, UnsortedAlertChainTableName);
+	UnsortedChainTable = DataStorage.RegisterTable<FTedsUnsortedAlertChainTag>(UnsortedAlertChainTableName, 
+		FTableRegistrationOptions{.SourceTable = ChainTable});
 }
 
 void UTedsAlertsFactory::RegisterQueries(UE::Editor::DataStorage::ICoreProvider& DataStorage)
diff --git a/Engine/Plugins/Experimental/EditorDataStorageFeatures/Source/TedsOutliner/Private/TedsOutlinerFilter.cpp b/Engine/Plugins/Experimental/EditorDataStorageFeatures/Source/TedsOutliner/Private/TedsOutlinerFilter.cpp
--- a/Engine/Plugins/Experimental/EditorDataStorageFeatures/Source/TedsOutliner/Private/TedsOutlinerFilter.cpp
+++ b/Engine/Plugins/Experimental/EditorDataStorageFeatures/Source/TedsOutliner/Private/TedsOutlinerFilter.cpp
@@ -27,13 +27,13 @@ FTedsOutlinerFilter::FTedsOutlinerFilter(const FName& InFilterName, const FText&
 {}
 
 FTedsOutlinerFilter::FTedsOutlinerFilter(const FName& InFilterName, const FText& InFilterDisplayName, const FText& InFilterToolTip, 
-	const FName& InFilterIconName, const TSharedPtr<FFilterCategory>& InCategory, const TQueryFunction<bool>& InFilterQuery,
+	const FName& InFilterIconName, const TSharedPtr<FFilterCategory>& InCategory, const TConstQueryFunction<bool>& InFilterQuery,
 	const bool bInteractiveFilter)
 	: FTedsFilterBase(InFilterName, InFilterDisplayName, InFilterToolTip, InFilterIconName, InCategory, InFilterQuery)
 	, bInteractiveFilter(bInteractiveFilter)
 {}
 
-FTedsOutlinerFilter::FTedsOutlinerFilter(const FName& InFilterName, const FText& InFilterDisplayName, const TQueryFunction<bool>& InFilterQuery,
+FTedsOutlinerFilter::FTedsOutlinerFilter(const FName& InFilterName, const FText& InFilterDisplayName, const TConstQueryFunction<bool>& InFilterQuery,
 	const bool bInteractiveFilter)
 	: FTedsOutlinerFilter(InFilterName, InFilterDisplayName, FText::FromString(InFilterName.ToString()), FName(), 
 		nullptr, InFilterQuery, bInteractiveFilter)
@@ -61,15 +61,15 @@ void FTedsOutlinerFilter::ActiveStateChanged(bool bActive)
 			if (bIsClassFilter)
 			{
 				// Class filters are separated since we want to OR them with each other
-				TedsOutlinerImplPin->AddClassQueryFunction(FilterName, FilterQuery.Get<TQueryFunction<bool>>());
+				TedsOutlinerImplPin->AddClassQueryFunction(FilterName, FilterQuery.Get<TConstQueryFunction<bool>>());
 			}
 			else if (FilterQuery.IsType<QueryHandle>())
 			{
 				TedsOutlinerImplPin->AddExternalQuery(FilterName, FilterQuery.Get<QueryHandle>());
 			}
-			else /* if (FilterQuery.IsType<TQueryFunction<bool>>()) */
+			else /* if (FilterQuery.IsType<TConstQueryFunction<bool>>()) */
 			{
-				TedsOutlinerImplPin->AddExternalQueryFunction(FilterName, FilterQuery.Get<TQueryFunction<bool>>());
+				TedsOutlinerImplPin->AddExternalQueryFunction(FilterName, FilterQuery.Get<TConstQueryFunction<bool>>());
 			}
 		}
 		else
diff --git a/Engine/Plugins/Experimental/EditorDataStorageFeatures/Source/TedsOutliner/Private/TedsOutlinerImpl.cpp b/Engine/Plugins/Experimental/EditorDataStorageFeatures/Source/TedsOutliner/Private/TedsOutlinerImpl.cpp
--- a/Engine/Plugins/Experimental/EditorDataStorageFeatures/Source/TedsOutliner/Private/TedsOutlinerImpl.cpp
+++ b/Engine/Plugins/Experimental/EditorDataStorageFeatures/Source/TedsOutliner/Private/TedsOutlinerImpl.cpp
@@ -559,7 +559,7 @@ void FTedsOutlinerImpl::AppendExternalQueries(FQueryDescription& OutQuery)
 	}
 }
 
-void FTedsOutlinerImpl::AddExternalQueryFunction(const FName& QueryName, const DataStorage::Queries::TQueryFunction<bool>& InQueryFunction)
+void FTedsOutlinerImpl::AddExternalQueryFunction(const FName& QueryName, const DataStorage::Queries::TConstQueryFunction<bool>& InQueryFunction)
 {
 	// Store the external query functions as functions instead of RowNodes since the row node to filter on has not been initialized yet
 	ExternalQueryFunctions.Emplace(QueryName, InQueryFunction);
@@ -570,7 +570,7 @@ void FTedsOutlinerImpl::RemoveExternalQueryFunction(const FName& QueryName)
 	ExternalQueryFunctions.Remove(QueryName);
 }
 
-void FTedsOutlinerImpl::AddClassQueryFunction(const FName& ClassName, const TQueryFunction<bool>& InClassQueryFunction)
+void FTedsOutlinerImpl::AddClassQueryFunction(const FName& ClassName, const TConstQueryFunction<bool>& InClassQueryFunction)
 {
 	ClassFilters.Emplace(ClassName, InClassQueryFunction);
 }
@@ -1041,7 +1041,7 @@ void FTedsOutlinerImpl::RecompileQueries()
 	FilterNodes.Empty();
 	TSharedPtr<IRowNode> PrimaryFilterNode = CombinedRowCollectorNode;
 	// Combines all active external filters using AND
-	for (const TPair<FName, TQueryFunction<bool>>& ExternalQueryFunction : ExternalQueryFunctions)
+	for (const TPair<FName, TConstQueryFunction<bool>>& ExternalQueryFunction : ExternalQueryFunctions)
 	{
 		TSharedPtr<FRowFilterNode> FilterNode = MakeShared<FRowFilterNode>(Storage, PrimaryFilterNode, ExternalQueryFunction.Value);
 		FilterNodes.Add(FilterNode);
@@ -1052,7 +1052,7 @@ void FTedsOutlinerImpl::RecompileQueries()
 	// Combines all active class filters using OR
 	if (!ClassFilters.IsEmpty())
 	{
-		for (const TPair<FName, TQueryFunction<bool>>& ClassFilter : ClassFilters)
+		for (const TPair<FName, TConstQueryFunction<bool>>& ClassFilter : ClassFilters)
 		{
 			ClassFilterNodes.Add(MakeShared<FRowFilterNode>(Storage, InitialRowCollectorNode, ClassFilter.Value));
 		}
diff --git a/Engine/Plugins/Experimental/EditorDataStorageFeatures/Source/TedsOutliner/Private/TedsOutlinerModule.cpp b/Engine/Plugins/Experimental/EditorDataStorageFeatures/Source/TedsOutliner/Private/TedsOutlinerModule.cpp
--- a/Engine/Plugins/Experimental/EditorDataStorageFeatures/Source/TedsOutliner/Private/TedsOutlinerModule.cpp
+++ b/Engine/Plugins/Experimental/EditorDataStorageFeatures/Source/TedsOutliner/Private/TedsOutlinerModule.cpp
@@ -263,10 +263,11 @@ FTedsSceneOutlinerWorldFilter::FTedsSceneOutlinerWorldFilter(const FName& InFilt
 		: FTedsOutlinerFilter(
           	InFilterName,
           	FText(), 
-          	Queries::BuildQueryFunction<bool>([this](TQueryContext<RowBatchInfo> Context, TResult<bool>& Result, TConstBatch<FTypedElementWorldColumn> WorldColumns)
-			  {
-				  QueryFunction(Context, Result, WorldColumns);
-			  }),
+          	Queries::BuildConstQueryFunction<bool>(
+				[this](TConstQueryContext<RowBatchInfo> Context, TResult<bool>& Result, TConstBatch<FTypedElementWorldColumn> WorldColumns)
+				{
+					QueryFunction(Context, Result, WorldColumns);
+				}),
           	false)
 		, SpecifiedWorldToDisplay(InSpecifiedWorldToDisplay)
 {
@@ -307,7 +308,7 @@ void FTedsSceneOutlinerWorldFilter::BuildWorldFilter()
 	}
 }
 
-void FTedsSceneOutlinerWorldFilter::QueryFunction(TQueryContext<RowBatchInfo> Context, TResult<bool>& Result, TConstBatch<FTypedElementWorldColumn> WorldColumns)
+void FTedsSceneOutlinerWorldFilter::QueryFunction(TConstQueryContext<RowBatchInfo> Context, TResult<bool>& Result, TConstBatch<FTypedElementWorldColumn> WorldColumns)
 {
 	// We always want this to be valid, so if not, find a new world to filter for
 	if (!RepresentingWorld.IsValid())
@@ -315,15 +316,15 @@ void FTedsSceneOutlinerWorldFilter::QueryFunction(TQueryContext<RowBatchInfo> Co
 		BuildWorldFilter();
 	}
 	
-	Context.ForEachRow([this, &Result](const FTypedElementWorldColumn& WorldColumn)
+	Context.ForEachRow([this, &Result](RowHandle Row, const FTypedElementWorldColumn& WorldColumn)
 		  {
 			  if (WorldColumn.World.IsValid() && RepresentingWorld.IsValid())
 			  {
-				  Result.Add(WorldColumn.World.Get() == RepresentingWorld.Get());
+				  Result.Add(Row, WorldColumn.World.Get() == RepresentingWorld.Get());
 			  }
 			  else
 			  {
-				  Result.Add(false);
+				  Result.Add(Row, false);
 			  }
 		  }, WorldColumns);
 }
diff --git a/Engine/Plugins/Experimental/EditorDataStorageFeatures/Source/TedsOutliner/Private/Widgets/SceneOutlinerWidget.cpp b/Engine/Plugins/Experimental/EditorDataStorageFeatures/Source/TedsOutliner/Private/Widgets/SceneOutlinerWidget.cpp
--- a/Engine/Plugins/Experimental/EditorDataStorageFeatures/Source/TedsOutliner/Private/Widgets/SceneOutlinerWidget.cpp
+++ b/Engine/Plugins/Experimental/EditorDataStorageFeatures/Source/TedsOutliner/Private/Widgets/SceneOutlinerWidget.cpp
@@ -186,7 +186,7 @@ void FTedsSceneOutlinerWidgetConstructor::GetShowOptionsFilters(UE::Editor::Data
 			LOCTEXT("ToggleHideTemporaryActorsToolTip", "When enabled, hides temporary/run-time actors."),
 		NAME_None,
 		nullptr,
-		[](TQueryContext<SingleRowInfo>, const FTypedElementUObjectColumn& ObjectColumn)
+		[](TConstQueryContext<SingleRowInfo>, const FTypedElementUObjectColumn& ObjectColumn)
 	 	{
 			if (const UObject* Object = ObjectColumn.Object.Get())
 			{
@@ -203,7 +203,7 @@ void FTedsSceneOutlinerWidgetConstructor::GetShowOptionsFilters(UE::Editor::Data
 		LOCTEXT("ToggleHideLevelInstancesToolTip", "When enabled, hides all level instance content."),
 		NAME_None,
 		nullptr,
-		[DataStorage](TQueryContext<SingleRowInfo> Context, const FTypedElementUObjectColumn& ObjectColumn)
+		[DataStorage](TConstQueryContext<SingleRowInfo> Context, const FTypedElementUObjectColumn& ObjectColumn)
 		 {
 			if (UObject* Object = ObjectColumn.Object.Get())
 			{
@@ -258,7 +258,7 @@ void FTedsSceneOutlinerWidgetConstructor::GetShowOptionsFilters(UE::Editor::Data
 		LOCTEXT("ToggleShowOnlyCurrentLevelToolTip", "When enabled, only shows objects that are in the Current Level."),
 		NAME_None,
 		nullptr,
-		[DataStorage](TQueryContext<SingleRowInfo> Context, const FTypedElementUObjectColumn& ObjectColumn)
+		[DataStorage](TConstQueryContext<SingleRowInfo> Context, const FTypedElementUObjectColumn& ObjectColumn)
 		 {
 			if (UObject* Object = ObjectColumn.Object.Get())
 			{
diff --git a/Engine/Plugins/Experimental/EditorDataStorageFeatures/Source/TedsOutliner/Public/TedsOutlinerFilter.h b/Engine/Plugins/Experimental/EditorDataStorageFeatures/Source/TedsOutliner/Public/TedsOutlinerFilter.h
--- a/Engine/Plugins/Experimental/EditorDataStorageFeatures/Source/TedsOutliner/Public/TedsOutlinerFilter.h
+++ b/Engine/Plugins/Experimental/EditorDataStorageFeatures/Source/TedsOutliner/Public/TedsOutlinerFilter.h
@@ -55,7 +55,7 @@ namespace UE::Editor::Outliner
 		 * @param bInteractiveFilter Whether this filter is visible in the Add Filter Menu and toggleable or is hidden and enabled by default
 		 */
 		TEDSOUTLINER_API FTedsOutlinerFilter(const FName& InFilterName, const FText& InFilterDisplayName, const FText& InFilterToolTip, const FName& InFilterIconName,
-			const TSharedPtr<FFilterCategory>& InCategory, const TQueryFunction<bool>& InFilterQuery, const bool bInteractiveFilter = true);
+			const TSharedPtr<FFilterCategory>& InCategory, const TConstQueryFunction<bool>& InFilterQuery, const bool bInteractiveFilter = true);
 
 		/**
 		 * Constructor for a TEDS Outliner Filter
@@ -64,7 +64,7 @@ namespace UE::Editor::Outliner
 		 * @param InFilterQuery QueryFunction to filter for, Batching should be used for improved performance if possible
 		 * @param bInteractiveFilter Whether this filter is visible in the Add Filter Menu and toggleable or is hidden and enabled by default
 		 */
-		TEDSOUTLINER_API FTedsOutlinerFilter(const FName& InFilterName, const FText& InFilterDisplayName, const TQueryFunction<bool>& InFilterQuery,
+		TEDSOUTLINER_API FTedsOutlinerFilter(const FName& InFilterName, const FText& InFilterDisplayName, const TConstQueryFunction<bool>& InFilterQuery,
 			const bool bInteractiveFilter = true);
 		
 		/**
@@ -81,7 +81,7 @@ namespace UE::Editor::Outliner
 		 */
 		template< typename FilterFunction UE_REQUIRES(
 			!std::is_same_v<std::remove_cvref_t<FilterFunction>, QueryHandle> &&
-			!std::is_same_v<std::remove_cvref_t<FilterFunction>, Queries::TQueryFunction<bool>>) >
+			!std::is_same_v<std::remove_cvref_t<FilterFunction>, Queries::TConstQueryFunction<bool>>) >
 		FTedsOutlinerFilter(const FName& InFilterName, const FText& InFilterDisplayName, const FText& InFilterToolTip, const FName& InFilterIconName,
 			const TSharedPtr<FFilterCategory>& InCategory, FilterFunction&& InFilterQuery, const bool bInteractiveFilter = true)
 			: FTedsOutlinerFilter(
@@ -90,7 +90,7 @@ namespace UE::Editor::Outliner
 				InFilterToolTip, 
 				InFilterIconName, 
 				InCategory,
-				Queries::BuildQueryFunction<bool>(Forward<FilterFunction>(InFilterQuery)),
+				Queries::BuildConstQueryFunction<bool>(Forward<FilterFunction>(InFilterQuery)),
 				bInteractiveFilter)
 		{}
 
@@ -105,7 +105,7 @@ namespace UE::Editor::Outliner
 		 */
 		template< typename FilterFunction UE_REQUIRES(
 			!std::is_same_v<std::remove_cvref_t<FilterFunction>, QueryHandle> &&
-			!std::is_same_v<std::remove_cvref_t<FilterFunction>, Queries::TQueryFunction<bool>>) >
+			!std::is_same_v<std::remove_cvref_t<FilterFunction>, Queries::TConstQueryFunction<bool>>) >
 		FTedsOutlinerFilter(const FName& InFilterName, const FText& InFilterDisplayName, FilterFunction&& InFilterQuery,
 			const bool bInteractiveFilter = true)
 			: FTedsOutlinerFilter(
@@ -114,7 +114,7 @@ namespace UE::Editor::Outliner
 				FText::FromString(InFilterName.ToString()), 
 				FName(), 
 				nullptr,
-				Queries::BuildQueryFunction<bool>(Forward<FilterFunction>(InFilterQuery)),
+				Queries::BuildConstQueryFunction<bool>(Forward<FilterFunction>(InFilterQuery)),
 				bInteractiveFilter)
 		{}
 
diff --git a/Engine/Plugins/Experimental/EditorDataStorageFeatures/Source/TedsOutliner/Public/TedsOutlinerImpl.h b/Engine/Plugins/Experimental/EditorDataStorageFeatures/Source/TedsOutliner/Public/TedsOutlinerImpl.h
--- a/Engine/Plugins/Experimental/EditorDataStorageFeatures/Source/TedsOutliner/Public/TedsOutlinerImpl.h
+++ b/Engine/Plugins/Experimental/EditorDataStorageFeatures/Source/TedsOutliner/Public/TedsOutlinerImpl.h
@@ -233,11 +233,11 @@ public:
 	UE_API void AppendExternalQueries(DataStorage::FQueryDescription& OutQuery);
 
 	// Add an external query function to the Outliner
-	UE_API void AddExternalQueryFunction(const FName& QueryName, const TQueryFunction<bool>& InQueryFunction);
+	UE_API void AddExternalQueryFunction(const FName& QueryName, const TConstQueryFunction<bool>& InQueryFunction);
 	UE_API void RemoveExternalQueryFunction(const FName& QueryName);
 
 	// Add an external class query function to the Outliner (uses OR) 
-	UE_API void AddClassQueryFunction(const FName& ClassName, const TQueryFunction<bool>& InClassQueryFunction);
+	UE_API void AddClassQueryFunction(const FName& ClassName, const TConstQueryFunction<bool>& InClassQueryFunction);
 	UE_API void RemoveClassQueryFunction(const FName& ClassName);
 	
 	// Outliner specific functionality
@@ -311,10 +311,10 @@ protected:
 	TMap<FName, TSharedPtr<DataStorage::QueryStack::IQueryNode>> ExternalQueries;
 	
 	// External query functions that are currently active (e.g Filters)
-    TMap<FName, DataStorage::Queries::TQueryFunction<bool>> ExternalQueryFunctions;
+    TMap<FName, DataStorage::Queries::TConstQueryFunction<bool>> ExternalQueryFunctions;
 
 	// External query functions used to filter by class, these are the only filters that use OR instead of AND
-	TMap<FName, DataStorage::Queries::TQueryFunction<bool>> ClassFilters;
+	TMap<FName, DataStorage::Queries::TConstQueryFunction<bool>> ClassFilters;
 
 	// Array to store all filter nodes created from the active query functions
 	TArray<TSharedPtr<DataStorage::QueryStack::FRowFilterNode>> FilterNodes;
diff --git a/Engine/Plugins/Experimental/EditorDataStorageFeatures/Source/TedsOutliner/Public/TedsOutlinerModule.h b/Engine/Plugins/Experimental/EditorDataStorageFeatures/Source/TedsOutliner/Public/TedsOutlinerModule.h
--- a/Engine/Plugins/Experimental/EditorDataStorageFeatures/Source/TedsOutliner/Public/TedsOutlinerModule.h
+++ b/Engine/Plugins/Experimental/EditorDataStorageFeatures/Source/TedsOutliner/Public/TedsOutlinerModule.h
@@ -81,7 +81,7 @@ public:
 	
 	void BuildWorldFilter();
 	
-	void QueryFunction(TQueryContext<RowBatchInfo> Context, TResult<bool>& Result, TConstBatch<FTypedElementWorldColumn> WorldColumns);
+	void QueryFunction(TConstQueryContext<RowBatchInfo> Context, TResult<bool>& Result, TConstBatch<FTypedElementWorldColumn> WorldColumns);
 
 	void SetUserChosenWorld(const TWeakObjectPtr<UWorld>& InWorld);
 
diff --git a/Engine/Plugins/Experimental/EditorDataStorageFeatures/Source/TedsQueryStack/Public/TedsRowFilterNode.h b/Engine/Plugins/Experimental/EditorDataStorageFeatures/Source/TedsQueryStack/Public/TedsRowFilterNode.h
--- a/Engine/Plugins/Experimental/EditorDataStorageFeatures/Source/TedsQueryStack/Public/TedsRowFilterNode.h
+++ b/Engine/Plugins/Experimental/EditorDataStorageFeatures/Source/TedsQueryStack/Public/TedsRowFilterNode.h
@@ -23,7 +23,7 @@ namespace UE::Editor::DataStorage::QueryStack
 	{
 	public:
 		FRowFilterNode(ICoreProvider* Storage, const TSharedPtr<IRowNode>& ParentRowNode, ICoreProvider::EFilterOptions Options, 
-			const Queries::TQueryFunction<bool>& Filter)
+			const Queries::TConstQueryFunction<bool>& Filter)
 			: Storage(Storage)
 			, Filter(Filter)
 			, ParentRowNode(ParentRowNode)
@@ -31,7 +31,7 @@ namespace UE::Editor::DataStorage::QueryStack
 		{
 		}
     
-		FRowFilterNode(ICoreProvider* Storage, const TSharedPtr<IRowNode>& ParentRowNode, const Queries::TQueryFunction<bool>& Filter)
+		FRowFilterNode(ICoreProvider* Storage, const TSharedPtr<IRowNode>& ParentRowNode, const Queries::TConstQueryFunction<bool>& Filter)
 			: FRowFilterNode(
 				Storage, 
 				ParentRowNode, 
@@ -40,30 +40,30 @@ namespace UE::Editor::DataStorage::QueryStack
 		{}
             
 		/*
-		 * Require that the template type cannot already be a TQueryFunction so it will pass to the defined constructor instead of this one.
+		 * Require that the template type cannot already be a TConstQueryFunction so it will pass to the defined constructor instead of this one.
 		 */
 		template< typename FilterFunction UE_REQUIRES(
-			!std::is_same_v<std::remove_cvref_t<FilterFunction>, Queries::TQueryFunction<bool>>) >
+			!std::is_same_v<std::remove_cvref_t<FilterFunction>, Queries::TConstQueryFunction<bool>>) >
 		FRowFilterNode(ICoreProvider* Storage, const TSharedPtr<IRowNode>& ParentRowNode, ICoreProvider::EFilterOptions Options, 
 			FilterFunction&& Filter)
 			: FRowFilterNode(
 				Storage, 
 				ParentRowNode, 
 				Options, 
-				Queries::BuildQueryFunction<bool>(Forward<FilterFunction>(Filter)))
+				Queries::BuildConstQueryFunction<bool>(Forward<FilterFunction>(Filter)))
 		{}
 
 		/*
-		 * Require that the template type cannot already be a TQueryFunction so it will pass to the defined constructor instead of this one.
+		 * Require that the template type cannot already be a TConstQueryFunction so it will pass to the defined constructor instead of this one.
 		 */
 		template< typename FilterFunction UE_REQUIRES(
-			!std::is_same_v<std::remove_cvref_t<FilterFunction>, Queries::TQueryFunction<bool>>) >
+			!std::is_same_v<std::remove_cvref_t<FilterFunction>, Queries::TConstQueryFunction<bool>>) >
 		FRowFilterNode(ICoreProvider* Storage, const TSharedPtr<IRowNode>& ParentRowNode, FilterFunction&& Filter)
 			: FRowFilterNode(
 				Storage, 
 				ParentRowNode, 
 				ICoreProvider::EFilterOptions::Default, 
-				Queries::BuildQueryFunction<bool>(Forward<FilterFunction>(Filter)))
+				Queries::BuildConstQueryFunction<bool>(Forward<FilterFunction>(Filter)))
 		{}
 
 		virtual ~FRowFilterNode() override = default;
@@ -81,7 +81,7 @@ namespace UE::Editor::DataStorage::QueryStack
 	
 	private:
 		ICoreProvider* Storage;
-		Queries::TQueryFunction<bool> Filter;
+		Queries::TConstQueryFunction<bool> Filter;
 		FRowHandleArray Rows;
 		TSharedPtr<IRowNode> ParentRowNode;
 		RevisionId CachedParentRevisionID = UninitializedRevisionId;
diff --git a/Engine/Plugins/Experimental/EditorDataStorageFeatures/Source/TedsTableViewer/Private/TedsFilter.cpp b/Engine/Plugins/Experimental/EditorDataStorageFeatures/Source/TedsTableViewer/Private/TedsFilter.cpp
--- a/Engine/Plugins/Experimental/EditorDataStorageFeatures/Source/TedsTableViewer/Private/TedsFilter.cpp
+++ b/Engine/Plugins/Experimental/EditorDataStorageFeatures/Source/TedsTableViewer/Private/TedsFilter.cpp
@@ -25,15 +25,15 @@ namespace UE::Editor::DataStorage
 				if (bIsClassFilter)
 				{
 					// Class filters are separated since we want to OR them with each other
-					TedsFilterBarPin->AddClassQueryFunction(FilterName, FilterQuery.Get<TQueryFunction<bool>>());
+					TedsFilterBarPin->AddClassQueryFunction(FilterName, FilterQuery.Get<TConstQueryFunction<bool>>());
 				}
 				else if (FilterQuery.IsType<QueryHandle>())
 				{
 					TedsFilterBarPin->AddExternalQuery(FilterName, FilterQuery.Get<QueryHandle>());
 				}
-				else /* if (FilterQuery.IsType<TQueryFunction<bool>>()) */
+				else /* if (FilterQuery.IsType<TConstQueryFunction<bool>>()) */
 				{
-					TedsFilterBarPin->AddExternalQueryFunction(FilterName, FilterQuery.Get<TQueryFunction<bool>>());
+					TedsFilterBarPin->AddExternalQueryFunction(FilterName, FilterQuery.Get<TConstQueryFunction<bool>>());
 				}
 			}
 			else
diff --git a/Engine/Plugins/Experimental/EditorDataStorageFeatures/Source/TedsTableViewer/Private/TedsTableViewerUtils.cpp b/Engine/Plugins/Experimental/EditorDataStorageFeatures/Source/TedsTableViewer/Private/TedsTableViewerUtils.cpp
--- a/Engine/Plugins/Experimental/EditorDataStorageFeatures/Source/TedsTableViewer/Private/TedsTableViewerUtils.cpp
+++ b/Engine/Plugins/Experimental/EditorDataStorageFeatures/Source/TedsTableViewer/Private/TedsTableViewerUtils.cpp
@@ -362,17 +362,11 @@ void UTedsTableViewerFactory::RegisterTables(UE::Editor::DataStorage::ICoreProvi
 	const TableHandle BaseWidgetTable = DataStorage.FindTable(FName(TEXT("Editor_WidgetTable")));
 	if (BaseWidgetTable != InvalidTableHandle)
 	{
-		DataStorage.RegisterTable(
-			BaseWidgetTable,
-			{
-				FTypedElementRowReferenceColumn::StaticStruct()
-			}, TableViewerUtils::TableViewerWidgetTableName);
+		DataStorage.RegisterTable<FTypedElementRowReferenceColumn>(TableViewerUtils::TableViewerWidgetTableName,
+			FTableRegistrationOptions{ .SourceTable = BaseWidgetTable });
 	}
 	
-	DataStorage.RegisterTable(
-		{
-			FTableViewerRowTag::StaticStruct()
-		}, TableViewerUtils::TableViewerRowWidgetTableName);
+	DataStorage.RegisterTable< FTableViewerRowTag>(TableViewerUtils::TableViewerRowWidgetTableName);
 }
 
 void UTedsTableViewerFactory::RegisterHierarchies(UE::Editor::DataStorage::ICoreProvider& DataStorage)
diff --git a/Engine/Plugins/Experimental/EditorDataStorageFeatures/Source/TedsTableViewer/Private/Widgets/STedsFilterBar.cpp b/Engine/Plugins/Experimental/EditorDataStorageFeatures/Source/TedsTableViewer/Private/Widgets/STedsFilterBar.cpp
--- a/Engine/Plugins/Experimental/EditorDataStorageFeatures/Source/TedsTableViewer/Private/Widgets/STedsFilterBar.cpp
+++ b/Engine/Plugins/Experimental/EditorDataStorageFeatures/Source/TedsTableViewer/Private/Widgets/STedsFilterBar.cpp
@@ -292,7 +292,7 @@ namespace UE::Editor::DataStorage
 		FilterNodes.Empty();
 		CompositeFilterNode = CombinedQueryFilterRowNode;
 		// Combines all active external filters using AND
-		for (const TPair<FName, TQueryFunction<bool>>& ExternalQueryFunction : ExternalQueryFunctions)
+		for (const TPair<FName, TConstQueryFunction<bool>>& ExternalQueryFunction : ExternalQueryFunctions)
 		{
 			const TSharedPtr<FRowFilterNode> FilterNode = MakeShared<FRowFilterNode>(Storage, CompositeFilterNode, ExternalQueryFunction.Value);
 			FilterNodes.Add(FilterNode);
@@ -303,7 +303,7 @@ namespace UE::Editor::DataStorage
 		// Combines all active class filters using OR
 		if (!ClassFilters.IsEmpty())
 		{
-			for (const TPair<FName, TQueryFunction<bool>>& ClassFilter : ClassFilters)
+			for (const TPair<FName, TConstQueryFunction<bool>>& ClassFilter : ClassFilters)
 			{
 				ClassFilterNodes.Add(MakeShared<FRowFilterNode>(Storage, InFilterableRowNode, ClassFilter.Value));
 			}
@@ -346,7 +346,7 @@ namespace UE::Editor::DataStorage
 		bMarkDirty = true;
 	}
 
-	void STedsFilterBar::AddExternalQueryFunction(const FName& QueryName, const DataStorage::Queries::TQueryFunction<bool>& InQueryFunction)
+	void STedsFilterBar::AddExternalQueryFunction(const FName& QueryName, const DataStorage::Queries::TConstQueryFunction<bool>& InQueryFunction)
 	{
 		// Store the external query functions as functions instead of RowNodes since the row node to filter on has not been initialized yet
 		ExternalQueryFunctions.Emplace(QueryName, InQueryFunction);
@@ -359,7 +359,7 @@ namespace UE::Editor::DataStorage
 		bMarkDirty = true;
 	}
 
-	void STedsFilterBar::AddClassQueryFunction(const FName& ClassName, const Queries::TQueryFunction<bool>& InClassQueryFunction)
+	void STedsFilterBar::AddClassQueryFunction(const FName& ClassName, const Queries::TConstQueryFunction<bool>& InClassQueryFunction)
 	{
 		ClassFilters.Emplace(ClassName, InClassQueryFunction);
 		bMarkDirty = true;
diff --git a/Engine/Plugins/Experimental/EditorDataStorageFeatures/Source/TedsTableViewer/Private/Widgets/STedsSearchBox.cpp b/Engine/Plugins/Experimental/EditorDataStorageFeatures/Source/TedsTableViewer/Private/Widgets/STedsSearchBox.cpp
--- a/Engine/Plugins/Experimental/EditorDataStorageFeatures/Source/TedsTableViewer/Private/Widgets/STedsSearchBox.cpp
+++ b/Engine/Plugins/Experimental/EditorDataStorageFeatures/Source/TedsTableViewer/Private/Widgets/STedsSearchBox.cpp
@@ -46,7 +46,7 @@ namespace UE::Editor::DataStorage
 			else if (InArgs._InSearchableRowNode)
 			{
 				TSharedPtr<QueryStack::FRowFilterNode> NewSearchNode = MakeShared<QueryStack::FRowFilterNode>(Storage, InArgs._InSearchableRowNode,
-					[Widget = AsWeak()](TQueryContext<SingleRowInfo> Context, const FTypedElementLabelColumn& LabelColumn)
+					[Widget = AsWeak()](TConstQueryContext<SingleRowInfo> Context, const FTypedElementLabelColumn& LabelColumn)
 					{
 						if (TSharedPtr<SWidget> WidgetPtr = Widget.Pin())
 						{
diff --git a/Engine/Plugins/Experimental/EditorDataStorageFeatures/Source/TedsTableViewer/Public/TedsFilter.inl b/Engine/Plugins/Experimental/EditorDataStorageFeatures/Source/TedsTableViewer/Public/TedsFilter.inl
--- a/Engine/Plugins/Experimental/EditorDataStorageFeatures/Source/TedsTableViewer/Public/TedsFilter.inl
+++ b/Engine/Plugins/Experimental/EditorDataStorageFeatures/Source/TedsTableViewer/Public/TedsFilter.inl
@@ -65,7 +65,7 @@ namespace UE::Editor::DataStorage
 		 * @param InFilterQuery QueryFunction to filter for, Batching should be used for improved performance if possible
 		 */
 		FTedsFilterBase(const FName& InFilterName, const FText& InFilterDisplayName, const FText& InFilterToolTip, const FName& InFilterIconName,
-			const TSharedPtr<FFilterCategory>& InCategory, const Queries::TQueryFunction<bool>& InFilterQuery);
+			const TSharedPtr<FFilterCategory>& InCategory, const Queries::TConstQueryFunction<bool>& InFilterQuery);
 
 		/**
 		 * Constructor for a TEDS Filter
@@ -73,7 +73,7 @@ namespace UE::Editor::DataStorage
 		 * @param InFilterDisplayName Display Name to show on the Add Filter Menu and toggle button
 		 * @param InFilterQuery QueryFunction to filter for, Batching should be used for improved performance if possible
 		 */
-		FTedsFilterBase(const FName& InFilterName, const FText& InFilterDisplayName, const Queries::TQueryFunction<bool>& InFilterQuery);
+		FTedsFilterBase(const FName& InFilterName, const FText& InFilterDisplayName, const Queries::TConstQueryFunction<bool>& InFilterQuery);
 		
 		/**
 		 * Constructor for a TEDS Filter
@@ -88,7 +88,7 @@ namespace UE::Editor::DataStorage
 		 */
 		template< typename FilterFunction UE_REQUIRES(
 			!std::is_same_v<std::remove_cvref_t<FilterFunction>, QueryHandle> &&
-			!std::is_same_v<std::remove_cvref_t<FilterFunction>, Queries::TQueryFunction<bool>>) >
+			!std::is_same_v<std::remove_cvref_t<FilterFunction>, Queries::TConstQueryFunction<bool>>) >
 		FTedsFilterBase(const FName& InFilterName, const FText& InFilterDisplayName, const FText& InFilterToolTip, const FName& InFilterIconName,
 			const TSharedPtr<FFilterCategory>& InCategory, FilterFunction&& InFilterQuery)
 			: FTedsFilterBase(
@@ -97,7 +97,7 @@ namespace UE::Editor::DataStorage
 				InFilterToolTip, 
 				InFilterIconName, 
 				InCategory,
-				Queries::BuildQueryFunction<bool>(Forward<FilterFunction>(InFilterQuery)))
+				Queries::BuildConstQueryFunction<bool>(Forward<FilterFunction>(InFilterQuery)))
 		{}
 
 		/**
@@ -110,7 +110,7 @@ namespace UE::Editor::DataStorage
 		 */
 		template< typename FilterFunction UE_REQUIRES(
 			!std::is_same_v<std::remove_cvref_t<FilterFunction>, QueryHandle> &&
-			!std::is_same_v<std::remove_cvref_t<FilterFunction>, Queries::TQueryFunction<bool>>) >
+			!std::is_same_v<std::remove_cvref_t<FilterFunction>, Queries::TConstQueryFunction<bool>>) >
 		FTedsFilterBase(const FName& InFilterName, const FText& InFilterDisplayName, FilterFunction&& InFilterQuery)
 			: FTedsFilterBase(
 				InFilterName, 
@@ -118,7 +118,7 @@ namespace UE::Editor::DataStorage
 				FText::FromString(InFilterName.ToString()), 
 				FName(), 
 				nullptr,
-				Queries::BuildQueryFunction<bool>(Forward<FilterFunction>(InFilterQuery)))
+				Queries::BuildConstQueryFunction<bool>(Forward<FilterFunction>(InFilterQuery)))
 		{}
 
 		/**
@@ -171,7 +171,7 @@ namespace UE::Editor::DataStorage
 
 		const bool bIsClassFilter;
 
-		TVariant<QueryHandle, Queries::TQueryFunction<bool>> FilterQuery;
+		TVariant<QueryHandle, Queries::TConstQueryFunction<bool>> FilterQuery;
 	};
 
 	template<typename FilterType>
@@ -195,7 +195,7 @@ namespace UE::Editor::DataStorage
 
 	template<typename FilterType>
 	FTedsFilterBase<FilterType>::FTedsFilterBase(const FName& InFilterName, const FText& InFilterDisplayName, const FText& InFilterToolTip, 
-		const FName& InFilterIconName, const TSharedPtr<FFilterCategory>& InCategory, const Queries::TQueryFunction<bool>& InFilterQuery)
+		const FName& InFilterIconName, const TSharedPtr<FFilterCategory>& InCategory, const Queries::TConstQueryFunction<bool>& InFilterQuery)
 		: FFilterBase<FilterType>(InCategory)
 		, FilterName(InFilterName)
 		, FilterDisplayName(InFilterDisplayName)
@@ -203,11 +203,11 @@ namespace UE::Editor::DataStorage
 		, FilterIconName(InFilterIconName)
 		, bIsClassFilter(false)
 	{
-		FilterQuery.Set<Queries::TQueryFunction<bool>>(InFilterQuery);
+		FilterQuery.Set<Queries::TConstQueryFunction<bool>>(InFilterQuery);
 	}
 
 	template<typename FilterType>
-	FTedsFilterBase<FilterType>::FTedsFilterBase(const FName& InFilterName, const FText& InFilterDisplayName, const Queries::TQueryFunction<bool>& InFilterQuery)
+	FTedsFilterBase<FilterType>::FTedsFilterBase(const FName& InFilterName, const FText& InFilterDisplayName, const Queries::TConstQueryFunction<bool>& InFilterQuery)
 		: FTedsFilterBase(InFilterName, InFilterDisplayName, FText::FromString(InFilterName.ToString()), FName(), 
 			nullptr, InFilterQuery)
 	{}
@@ -222,12 +222,12 @@ namespace UE::Editor::DataStorage
 		, bIsClassFilter(true)
 	{
 		using namespace UE::Editor::DataStorage::Queries;
-		FilterQuery.Set<TQueryFunction<bool>>(
-			Queries::BuildQueryFunction<bool>([InClass](TQueryContext<RowBatchInfo> Context, TResult<bool>& Result, TConstBatch<FTypedElementClassTypeInfoColumn> TypeInfoColumns)
+		FilterQuery.Set<TConstQueryFunction<bool>>(
+			Queries::BuildConstQueryFunction<bool>([InClass](TConstQueryContext<RowBatchInfo> Context, TResult<bool>& Result, TConstBatch<FTypedElementClassTypeInfoColumn> TypeInfoColumns)
 			{
-				Context.ForEachRow([&Result, InClass](const FTypedElementClassTypeInfoColumn& TypeInfoColumn)
+				Context.ForEachRow([&Result, InClass](RowHandle Row, const FTypedElementClassTypeInfoColumn& TypeInfoColumn)
 				{
-					Result.Add(TypeInfoColumn.TypeInfo->IsChildOf(InClass));
+					Result.Add(Row, TypeInfoColumn.TypeInfo->IsChildOf(InClass));
 				}, TypeInfoColumns);
 			}));
 	}
diff --git a/Engine/Plugins/Experimental/EditorDataStorageFeatures/Source/TedsTableViewer/Public/Widgets/STedsFilterBar.h b/Engine/Plugins/Experimental/EditorDataStorageFeatures/Source/TedsTableViewer/Public/Widgets/STedsFilterBar.h
--- a/Engine/Plugins/Experimental/EditorDataStorageFeatures/Source/TedsTableViewer/Public/Widgets/STedsFilterBar.h
+++ b/Engine/Plugins/Experimental/EditorDataStorageFeatures/Source/TedsTableViewer/Public/Widgets/STedsFilterBar.h
@@ -104,11 +104,11 @@ namespace UE::Editor::DataStorage
 		UE_API void AppendExternalQueries(DataStorage::FQueryDescription& OutQuery);
 
 		// Add an external query function to the FilterBar
-		UE_API void AddExternalQueryFunction(const FName& QueryName, const Queries::TQueryFunction<bool>& InQueryFunction);
+		UE_API void AddExternalQueryFunction(const FName& QueryName, const Queries::TConstQueryFunction<bool>& InQueryFunction);
 		UE_API void RemoveExternalQueryFunction(const FName& QueryName);
 
 		// Add an external class query function to the FilterBar (uses OR) 
-		UE_API void AddClassQueryFunction(const FName& ClassName, const Queries::TQueryFunction<bool>& InClassQueryFunction);
+		UE_API void AddClassQueryFunction(const FName& ClassName, const Queries::TConstQueryFunction<bool>& InClassQueryFunction);
 		UE_API void RemoveClassQueryFunction(const FName& ClassName);
 
 	protected:
@@ -140,10 +140,10 @@ namespace UE::Editor::DataStorage
 		TMap<FName, TSharedPtr<QueryStack::IQueryNode>> ExternalQueries;
 	
 		// External query functions that are currently active (e.g Filters)
-		TMap<FName, Queries::TQueryFunction<bool>> ExternalQueryFunctions;
+		TMap<FName, Queries::TConstQueryFunction<bool>> ExternalQueryFunctions;
 
 		// External query functions used to filter by class, these are the only filters that use OR instead of AND
-		TMap<FName, Queries::TQueryFunction<bool>> ClassFilters;
+		TMap<FName, Queries::TConstQueryFunction<bool>> ClassFilters;
 
 		// Array to store all filter nodes created from the active query functions
 		TArray<TSharedPtr<QueryStack::IRowNode>> FilterNodes;
diff --git a/Engine/Plugins/Experimental/MeshPartition/Source/MeshPartitionEditorUI/Private/MeshPartitionTedsFactory.cpp b/Engine/Plugins/Experimental/MeshPartition/Source/MeshPartitionEditorUI/Private/MeshPartitionTedsFactory.cpp
--- a/Engine/Plugins/Experimental/MeshPartition/Source/MeshPartitionEditorUI/Private/MeshPartitionTedsFactory.cpp
+++ b/Engine/Plugins/Experimental/MeshPartition/Source/MeshPartitionEditorUI/Private/MeshPartitionTedsFactory.cpp
@@ -91,28 +91,30 @@ void UMegaMeshTedsFactory::RegisterTables(Editor::DataStorage::ICoreProvider& Da
 {
 	using namespace Editor::DataStorage;
 	
-	MegaMeshTable = DataStorage.RegisterTable(DataStorage.FindTable(FName(TEXT("Editor_StandardActorTable"))),
-		TTypedElementColumnTypeList<FTypedElementUObjectColumn, FIsMegaMeshObjectTag, FMegaMeshRowParentColumn,
-									FMegaMeshDrawBoundsColumn, FMegaMeshBuildToStatus, FMegaMeshTimingStatistics>(),
-		UMegaMeshEditorUISubsystem::Local::SceneMegaMeshesTableName);
-
-	MegaMeshDefinitionTable = DataStorage.RegisterTable(DataStorage.FindTable(FName(TEXT("Editor_StandardUObjectTable"))),
-		TTypedElementColumnTypeList<FTypedElementUObjectColumn, FIsMegaMeshDefinitionObjectTag>(),
-		UMegaMeshEditorUISubsystem::Local::SceneMegaMeshDefinitionTableName);
-
-	MegaMeshModifierTable = DataStorage.RegisterTable(DataStorage.FindTable(FName(TEXT("Editor_StandardUObjectTable"))),
-		TTypedElementColumnTypeList<FTypedElementUObjectColumn, FTypedElementLabelColumn, FTypedElementLabelHashColumn,
-									FIsMegaMeshModifierTag, FMegaMeshLayerNameColumn, FMegaMeshRowParentColumn,
-									FMegaMeshReferenceColumn, FMegaMeshModifierTiming, FMegaMeshPriorityColumn,
-									FMegaMeshBuildUpToThisLayerColumn, FMegaMeshModifierSortColumn, FParentActorRefColumn>(),
-		UMegaMeshEditorUISubsystem::Local::SceneMegaMeshesModifiersTableName);
-
-	MegaMeshDefinitionLayerTable = DataStorage.RegisterTable(
-		TTypedElementColumnTypeList<FTypedElementRowReferenceColumn, FIsMegaMeshDefinitionLayerTag, FTypedElementLabelColumn,
-									FTypedElementLabelHashColumn, FMegaMeshRowParentColumn, FMegaMeshPriorityColumn,
-									FMegaMeshLayerNameColumn, FMegaMeshDrawBoundsColumn, FMegaMeshBuildUpToThisLayerColumn,
-									FMegaMeshModifierTiming>(),
-		UMegaMeshEditorUISubsystem::Local::LayerVisibilityTableName);
+	MegaMeshTable = DataStorage.RegisterTable<
+		FTypedElementUObjectColumn, FIsMegaMeshObjectTag, FMegaMeshRowParentColumn,
+		FMegaMeshDrawBoundsColumn, FMegaMeshBuildToStatus, FMegaMeshTimingStatistics>(
+			UMegaMeshEditorUISubsystem::Local::SceneMegaMeshesTableName,
+			FTableRegistrationOptions{ .SourceTable = DataStorage.FindTable(FName(TEXT("Editor_StandardActorTable"))) });
+
+	MegaMeshDefinitionTable = DataStorage.RegisterTable<FTypedElementUObjectColumn, FIsMegaMeshDefinitionObjectTag>(
+		UMegaMeshEditorUISubsystem::Local::SceneMegaMeshDefinitionTableName,
+		FTableRegistrationOptions{ .SourceTable = DataStorage.FindTable(FName(TEXT("Editor_StandardUObjectTable"))) });
+
+	MegaMeshModifierTable = DataStorage.RegisterTable<
+		FTypedElementUObjectColumn, FTypedElementLabelColumn, FTypedElementLabelHashColumn,
+		FIsMegaMeshModifierTag, FMegaMeshLayerNameColumn, FMegaMeshRowParentColumn,
+		FMegaMeshReferenceColumn, FMegaMeshModifierTiming, FMegaMeshPriorityColumn,
+		FMegaMeshBuildUpToThisLayerColumn, FMegaMeshModifierSortColumn, FParentActorRefColumn>(
+			UMegaMeshEditorUISubsystem::Local::SceneMegaMeshesModifiersTableName,
+			FTableRegistrationOptions{ .SourceTable = DataStorage.FindTable(FName(TEXT("Editor_StandardUObjectTable"))) });
+
+	MegaMeshDefinitionLayerTable = DataStorage.RegisterTable<
+		FTypedElementRowReferenceColumn, FIsMegaMeshDefinitionLayerTag, FTypedElementLabelColumn,
+		FTypedElementLabelHashColumn, FMegaMeshRowParentColumn, FMegaMeshPriorityColumn,
+		FMegaMeshLayerNameColumn, FMegaMeshDrawBoundsColumn, FMegaMeshBuildUpToThisLayerColumn,
+		FMegaMeshModifierTiming
+		>(UMegaMeshEditorUISubsystem::Local::LayerVisibilityTableName);
 
 	CompatibilityDataStorage.RegisterTypeTableAssociation(AMeshPartition::StaticClass(), MegaMeshTable);
 	CompatibilityDataStorage.RegisterTypeTableAssociation(UMeshPartitionDefinition::StaticClass(), MegaMeshDefinitionTable);
diff --git a/Engine/Plugins/Experimental/MeshPartition/Source/MeshPartitionEditorUI/Private/UI/MeshPartitionSettingsWidget.cpp b/Engine/Plugins/Experimental/MeshPartition/Source/MeshPartitionEditorUI/Private/UI/MeshPartitionSettingsWidget.cpp
--- a/Engine/Plugins/Experimental/MeshPartition/Source/MeshPartitionEditorUI/Private/UI/MeshPartitionSettingsWidget.cpp
+++ b/Engine/Plugins/Experimental/MeshPartition/Source/MeshPartitionEditorUI/Private/UI/MeshPartitionSettingsWidget.cpp
@@ -655,8 +655,8 @@ TSharedRef<UE::Editor::DataStorage::STedsCompositeHierarchyViewer> SMegaMeshSett
 	TArray<TSharedPtr<DataStorage::FTedsFilter>> Filters;
 
 	{
-		TQueryFunction<bool> FilterQuery = 
-			BuildQueryFunction<bool>([](TQueryContext<SingleRowInfo> Context)
+		TConstQueryFunction<bool> FilterQuery = 
+			BuildConstQueryFunction<bool>([](TConstQueryContext<SingleRowInfo> Context)
 			{
 				if (Context.CurrentRowHasColumns<FMegaMeshAffectsFilterBounds>())
 				{
@@ -671,7 +671,7 @@ TSharedRef<UE::Editor::DataStorage::STedsCompositeHierarchyViewer> SMegaMeshSett
 			LOCTEXT("ActorBoundsFilter_Tooltip", "Use the selected Actor's bonuds to filter only those modifiers affecting that region."),
 			FName(),
 			MeshPartitionFilterCategory,
-			FilterQuery
+			MoveTemp(FilterQuery)
 			));
 	}
 
diff --git a/Engine/Source/Runtime/TypedElementFramework/Public/DataStorage/CommonTypes.h b/Engine/Source/Runtime/TypedElementFramework/Public/DataStorage/CommonTypes.h
--- a/Engine/Source/Runtime/TypedElementFramework/Public/DataStorage/CommonTypes.h
+++ b/Engine/Source/Runtime/TypedElementFramework/Public/DataStorage/CommonTypes.h
@@ -3,6 +3,7 @@
 #pragma once
 
 #include "DataStorage/Handles.h"
+#include "DataStorage/MapKey.h"
 
 #include <algorithm>
 #include "CommonTypes.generated.h"
@@ -103,6 +104,60 @@ namespace UE
 			bool operator==(const FDynamicColumnDescription&) const = default;
 		};
 
+		enum class ETableType : uint8
+		{
+			Invalid, //< Not a valid table.
+			Declared, //< Table declared by user.
+			Derivative, //< Table declared by users, based on another table.
+			Variant, //< A dynamically generated variant of a declared or derivative table.
+		};
+
+		/** 
+		 * Foreign keys provide a mechanic to generate a key that unique identifies a given row so it can be
+		 * referenced by external sources.
+		 */
+		using FForeignKey = FMapKey;
+		
+		/**
+		 * Short lived view into information about a table.
+		 * Note: This struct is meant for short term usage only and becomes invalid the moment there's any change to a table.
+		 */
+		struct FTableInfoView
+		{
+			/** List of all dynamically created variants of this table. */
+			TConstArrayView<TableHandle> Variants;
+			/** All declared tables that use this table as a template. */
+			TConstArrayView<TableHandle> Derivatives;
+			/** The name of the table. */
+			FName Name;
+			/** The type of table. */
+			ETableType Type = ETableType::Invalid;
+			/**
+			 * If this is a variant then the Parent represents the table used as the basis for the variant. If this is a derived table then
+			 * the Parent represents the table that was passed in as the template for this table.
+			 */
+			TableHandle Parent = InvalidTableHandle;
+		};
+
+		struct FTableRegistrationOptions
+		{
+			/** If set, use the provided table as a source to copy from. Additional provided rows will be added to the copied layout. */
+			TableHandle SourceTable = InvalidTableHandle;
+			/** 
+			 * Indicates the amount of memory a shard (a.k.a. chunk) of the table will take. Increasing this number means less memory 
+			 * fragmentation and longer continuous time spend in a query callback while smaller sizes will decrease wasted memory. Making
+			 * this too large yields diminishing returns as cache efficiency will plateau while distribution across threads will suffer as
+			 * there are less shards to distribute. Making this too small though will cause increased overhead in managing tables and the
+			 * overhead of calling query callbacks will relatively increase. Consider leaving this to default and only update if profiling
+			 * indicates a benefit. Benefits may include going larger for very large tables since distribution among processors will be high
+			 * but the cost of table management will go down. Decrease for tables that on average only have a small number of entries to
+			 * save on memory.
+			 * Set to 0 will indicate to use a default size or the editor specified override. Any other value will override any default or
+			 * user setting.
+			 */
+			int32 MemorySize = 0;
+		};
+
 		// Standard callbacks.
 
 		using RowCreationCallbackRef = TFunctionRef<void(RowHandle Row)>;
diff --git a/Engine/Source/Runtime/TypedElementFramework/Public/Elements/Framework/TypedElementQueryCapabilities.h b/Engine/Source/Runtime/TypedElementFramework/Public/Elements/Framework/TypedElementQueryCapabilities.h
--- a/Engine/Source/Runtime/TypedElementFramework/Public/Elements/Framework/TypedElementQueryCapabilities.h
+++ b/Engine/Source/Runtime/TypedElementFramework/Public/Elements/Framework/TypedElementQueryCapabilities.h
@@ -11,8 +11,10 @@
 #include "Elements/Framework/TypedElementQueryFunctionArgumentTypes.h"
 #include "Elements/Framework/TypedElementRowHandleArrayView.h"
 #include "Misc/EnumClassFlags.h"
+#include "UObject/Class.h"
 #include "UObject/NameTypes.h"
 
+
 class UScriptStruct;
 
 namespace UE::Editor::DataStorage::Queries
@@ -48,6 +50,8 @@ namespace UE::Editor::DataStorage::Queries
 
 	template<typename ReturnType>
 	class TQueryFunction;
+	template<typename ReturnType>
+	class TConstQueryFunction;
 
 	struct IContextContract;
 	struct IQueryFunctionResponse;
@@ -55,6 +59,9 @@ namespace UE::Editor::DataStorage::Queries
 	template<typename Return, FunctionType Function>
 	TQueryFunction<Return> BuildQueryFunction(Function&& Callback); // forward declaration
 
+	template<typename Return, FunctionType Function>
+	TConstQueryFunction<Return> BuildConstQueryFunction(Function&& Callback); // forward declaration
+
 	// Here only to make sure that the argument shows up in compiler errors.
 
 	template <typename T>
diff --git a/Engine/Source/Runtime/TypedElementFramework/Public/Elements/Framework/TypedElementQueryCapabilities.inl b/Engine/Source/Runtime/TypedElementFramework/Public/Elements/Framework/TypedElementQueryCapabilities.inl
--- a/Engine/Source/Runtime/TypedElementFramework/Public/Elements/Framework/TypedElementQueryCapabilities.inl
+++ b/Engine/Source/Runtime/TypedElementFramework/Public/Elements/Framework/TypedElementQueryCapabilities.inl
@@ -5,6 +5,124 @@
 // This file contains the definitions of all available query capabilities. The macros get expended in various ways to generate interfaces,
 // function forwarders, contexts, etc.
 
+/**
+ * Capability to provide information about tables.
+ */
+CapabilityStart(TableInfo, EContextCapabilityFlags::SupportsSingle | EContextCapabilityFlags::SupportsBatch)
+	/** Checks if the provided table handle is pointing to a valid table. A table handle of value `InvalidTableHandle` is considered invalid. */
+	ConstFunction1(TableInfo, bool, IsTableValid, (TableHandle, Table))
+	/**
+	 * Returns a previously registered table with the provided name or InvalidTableHandle if not found.
+	 * This function will always return the registered version of the table and never any of the dynamically created variants.
+	 */
+	ConstFunction1(TableInfo, TableHandle, FindTable, (const FName&, TableName))
+	/** Finds the table the row is stored in. */
+	ConstFunction1(TableInfo, TableHandle, FindTable, (RowHandle, Row))
+	/** Lists the handles for all tables. */
+	ConstFunction1(TableInfo, void, ListTables, (TFunctionRef<void(TableHandle)>, Callback))
+	/** Lists the handles for all tables of a specific type. */
+	ConstFunction2(TableInfo, void, ListTables, (ETableType, TableType), (TFunctionRef<void(TableHandle)>, Callback))
+	/** Lists all the columns in a table. */
+	ConstFunction2(TableInfo, void, ListTableColumns, (TableHandle, Table), (TFunctionRef<bool(const UScriptStruct*)>, Callback))
+	/** Returns true if the table contains the requested columns, otherwise false. */
+	ConstFunction2(TableInfo, bool, TableHasColumns, (TableHandle, Table), (TConstArrayView<const UScriptStruct*>, Columns))
+	/** Returns various bits of information about the table. */
+	ConstFunction1(TableInfo, FTableInfoView, GetTableInfo, (TableHandle, Table))
+	/** Lists all rows in a table. The callback may be called multiple times, once for each chunk in the table. */
+	ConstFunction2(TableInfo, void, ListTableRows, (TableHandle, Table), (TFunctionRef<void(FRowHandleArrayView)>, Callback))
+	/** Lists all the domain names for foreign keys that have been registered with this table. */
+	ConstFunction3(TableInfo, void, ListTableForeignKeyDomains, (TableHandle, Table), (bool, bIncludeParents), (TFunctionRef<void(const FName&)>, Callback))
+
+#if defined(WithWrappers)
+	/** Returns true if the table contains the requested columns, otherwise false. */
+	template<TColumnType... ColumnTypes>
+	bool TableHasColumns(TableHandle Table) const
+	{
+		return this->TableHasColumns(Table, TConstArrayView<const UScriptStruct*>({ ColumnTypes::StaticStruct()... }));
+	}
+#endif
+CapabilityEnd(TableInfo)
+
+/**
+ * Capability to provide information about the currently active table.
+ */
+CapabilityStart(CurrentTableInfo, EContextCapabilityFlags::SupportsSingle | EContextCapabilityFlags::SupportsBatch)
+	/** Returns the handle to the table that's currently being used in the context. */
+	ConstFunction0(CurrentTableInfo, TableHandle, GetCurrentTable)
+	/** Returns true if the current table contains the requested columns, otherwise false. */
+	ConstFunction1(CurrentTableInfo, bool, CurrentTableHasColumns, (TConstArrayView<const UScriptStruct*>, Columns))
+
+#if defined(WithWrappers)
+	/** Returns true if the current table contains the requested columns, otherwise false. */
+	template<TColumnType... ColumnTypes>
+	bool CurrentTableHasColumns() const
+	{
+		return this->CurrentTableHasColumns(TConstArrayView<const UScriptStruct*>({ ColumnTypes::StaticStruct()... }));
+	}
+#endif
+CapabilityEnd(CurrentTableInfo)
+
+/**
+ * Capability to allow managing tables.
+ */
+CapabilityStart(TableManagement, EContextCapabilityFlags::SupportsSingle | EContextCapabilityFlags::SupportsBatch)
+	/** Register a table by name. */
+	Function2(TableManagement, TableHandle, RegisterTable, (TConstArrayView<const UScriptStruct*>, Columns), (const FName&, TableName))
+	/** Register a table by name with additional configuration options. */
+	Function3(TableManagement, TableHandle, RegisterTable, \
+		(TConstArrayView<const UScriptStruct*>, Columns), (const FName&, TableName), (FTableRegistrationOptions, Options))
+	/**
+	 * Register a function that's used to generate a foreign key for a particular row. Foreign keys create a unique identifier that
+	 * represents the row for a specific domain.
+	 */
+	Function3(TableManagement, void, RegisterTableForeignKey, \
+		(TableHandle, Table), (const FName&, Domain), (const Queries::TConstQueryFunction<FForeignKey>&, KeyConstructor))
+	/**
+	 * Register a function that's used to generate a foreign key for a particular row. Foreign keys create a unique identifier that
+	 * represents the row for a specific domain.
+	 */
+	Function3(TableManagement, void, RegisterTableForeignKey, \
+		(TableHandle, Table), (const FName&, Domain), (Queries::TConstQueryFunction<FForeignKey>&&, KeyConstructor))
+	
+	/** Creates a foreign key for the provided row using the table the row is stored in and the assigned domain. */
+	ConstFunction2(TableManagement, FForeignKey, BuildForeignKey, (const FName&, Domain), (RowHandle, Row))
+
+#if defined(WithWrappers)
+	/** Register a table by name. */
+	template<TColumnType... ColumnTypes>
+	TableHandle RegisterTable(const FName& TableName)
+	{
+		return this->RegisterTable(TConstArrayView<const UScriptStruct*>({ ColumnTypes::StaticStruct()... }), TableName);
+	}
+	/** Register a table by name. */
+	template<TColumnType... ColumnTypes>
+	TableHandle RegisterTable(const FName& TableName, const FTableRegistrationOptions& Options)
+	{
+		return this->RegisterTable(TConstArrayView<const UScriptStruct*>({ ColumnTypes::StaticStruct()... }), TableName, Options);
+	}
+
+	template<FunctionType Function>
+	void RegisterTableForeignKey(TableHandle Table, const FName& Domain, Function&& Callback)
+	{
+		RegisterTableForeignKey(Table, Domain, BuildConstQueryFunction<FForeignKey>(Forward<Function>(Callback)));
+	}
+
+	UE_DEPRECATED(5.8, "Use the version of `RegisterTable` that takes a `FTableRegistrationOptions`.")
+	inline TableHandle RegisterTable(TableHandle SourceTable, TConstArrayView<const UScriptStruct*> ColumnList, const FName& TableName)
+	{
+		return this->RegisterTable(ColumnList, TableName, FTableRegistrationOptions{ .SourceTable = SourceTable });
+	}
+	
+	template<TColumnType... ColumnTypes>
+	UE_DEPRECATED(5.8, "Use the version of `RegisterTable` that takes a `FTableRegistrationOptions`.")
+	TableHandle RegisterTable(TableHandle SourceTable, const FName& TableName)
+	{
+		return this->RegisterTable(TConstArrayView<const UScriptStruct*>({ ColumnTypes::StaticStruct()... }), 
+			TableName, FTableRegistrationOptions{ .SourceTable = SourceTable });
+	}
+#endif
+CapabilityEnd(TableManagement)
+
 /**
  * Capability to provide information about the currently active row.
  */
@@ -55,7 +173,7 @@ CapabilityStart(RowBatchInfo, EContextCapabilityFlags::SupportsBatch)
 
 	/**
 	 * Utility function to iterate over each individual cell in the rows in this batch for the requested columns.
-	 * Example: ForEachRow([](const ColumnA&, const ColumnB&), ArgA, ArgB); 
+	 * Example: ForEachRow([](RowHandle Row, const ColumnA&, const ColumnB&), ArgA, ArgB); 
 	 *		where ArgA is a TBatch<ColumnA> or TConstBatch<ColumnA> (a.k.a. TBatch<const ColumnA>) and 
 	 *		ArgB is a TBatch<ColumnB> or TConstBatch<ColumnB> (a.k.a. TBatch<const ColumnB>) argument in the query callback function.
 	 */
@@ -67,14 +185,14 @@ CapabilityStart(RowBatchInfo, EContextCapabilityFlags::SupportsBatch)
 		{
 			AddressTuple.ApplyBefore([&](std::add_const_t<Args>*&... Pointers)
 				{
-					Callback((*Pointers++)...);
+					Callback(Row, (*Pointers++)...);
 				});
 		}
 	}
 
 	/**
 	 * Utility function to iterate over each individual cell in the rows in this batch for the requested columns.
-	 * Example: ForEachRow([](ColumnA&, const ColumnB&), ArgA, ArgB);
+	 * Example: ForEachRow([](RowHandle Row, ColumnA&, const ColumnB&), ArgA, ArgB);
 	 *		where ArgA is a TBatch<ColumnA> and ArgB is a TConstBatch<ColumnB> (a.k.a. TBatch<const ColumnB>) 
 	 *		or TBatch<ColumnB> argument in the query callback function.
 	 */
@@ -86,7 +204,7 @@ CapabilityStart(RowBatchInfo, EContextCapabilityFlags::SupportsBatch)
 		{
 			AddressTuple.ApplyBefore([&](Args*&... Pointers)
 				{
-					Callback((*Pointers++)...);
+					Callback(Row, (*Pointers++)...);
 				});
 		}
 	}
diff --git a/Engine/Source/Runtime/TypedElementFramework/Public/Elements/Framework/TypedElementQueryCapabilityForwarder.h b/Engine/Source/Runtime/TypedElementFramework/Public/Elements/Framework/TypedElementQueryCapabilityForwarder.h
--- a/Engine/Source/Runtime/TypedElementFramework/Public/Elements/Framework/TypedElementQueryCapabilityForwarder.h
+++ b/Engine/Source/Runtime/TypedElementFramework/Public/Elements/Framework/TypedElementQueryCapabilityForwarder.h
@@ -11,73 +11,68 @@ namespace UE::Editor::DataStorage::Queries::Private
 {
 	// Class fragments use to composite a class that forwards its calls to the query context contract.
 	// If the forwarder contains any virtual calls, the forwarder won't be fully inlined and will add an additional pointer dereference
-	// to each call. To avoid this, the forwarder doesn't directly inherit from the capability interfaces. Static asserts are added
-	// to verify that the implementation is complete.
+	// to each call. To avoid this, the forwarder doesn't directly inherit from the capability interfaces.
 	
 	struct FQueryContext {};
+	struct FConstQueryContext {};
 
-	struct ForwardTesterBase
-	{
-		virtual ~ForwardTesterBase() = default;
-	};
-
-	template<typename... Capabilities> struct TForwarder
+	template<bool bIsConst, typename... Capabilities> struct TForwarder
 	{
 		static_assert(sizeof...(Capabilities) < 0,
 			"Specialization for query context forwards missing. The most likely reason is the addition of a new query "
 			"context capability that wasn't specialized.");
 	};
 
-	template<typename Base, typename...Capabilities>
-	struct TForwarder<Base, Capabilities...> : Base, TForwarder<Capabilities...>{};
+	template<bool bIsConst, typename Base, typename...Capabilities>
+	struct TForwarder<bIsConst, Base, Capabilities...> : Base, TForwarder<bIsConst, Capabilities...>{};
 
-	template<> 
-	struct TForwarder<> : FQueryContext
+	template<bool bIsConst>
+	struct TForwarder<bIsConst> : std::conditional_t<bIsConst, FConstQueryContext, FQueryContext>
 	{
 		TForwarder() = default;
 		explicit TForwarder(IContextContract& Contract) : Contract(&Contract) {}
+		explicit TForwarder(const IContextContract& Contract) requires (bIsConst) : Contract(&Contract) {}
 	
 	protected:
-		IContextContract* Contract = nullptr;
+		std::conditional_t<bIsConst, const IContextContract*, IContextContract*> Contract = nullptr;
 	};
 
+#define CapabilityStart(Capability, Flags) \
+	template<bool bIsConst, typename... Capabilities> \
+	struct TForwarder<bIsConst, Capability , Capabilities...> : TForwarder<bIsConst, Capabilities...> \
+	{ \
+		TForwarder() = default; \
+		explicit TForwarder(IContextContract& Contract) : TForwarder<bIsConst, Capabilities...>(Contract) {} \
+		explicit TForwarder(const IContextContract& Contract) requires (bIsConst) : TForwarder<bIsConst, Capabilities...>(Contract) {}
+
 #define ArgTypeName(Type, Name) Type Name
 #define ArgName(Type, Name) Name
+#define ForwardArg(Type, Name) Forward< Type >( Name )
 
-#define Function0(Capability, Return, Function) Return Function() \
+#define Function0(Capability, Return, Function) Return Function() requires (!bIsConst) \
 	{ return this->Contract-> Function (); }
-#define Function1(Capability, Return, Function, Arg1) Return Function(ArgTypeName Arg1 ) \
-	{ return this->Contract-> Function (ArgName Arg1 ); }
-#define Function2(Capability, Return, Function, Arg1, Arg2) Return Function(ArgTypeName Arg1 , ArgTypeName Arg2 ) \
-	{ return this->Contract-> Function (ArgName Arg1 , ArgName Arg2 ); }
-#define Function3(Capability, Return, Function, Arg1, Arg2, Arg3) Return Function(ArgTypeName Arg1 , ArgTypeName Arg2 , ArgTypeName Arg3 ) \
-	{ return this->Contract-> Function (ArgName Arg1 , ArgName Arg2 , ArgName Arg3); }
-#define Function4(Capability, Return, Function, Arg1, Arg2, Arg3, Arg4) Return Function(ArgTypeName Arg1 , ArgTypeName Arg2 , ArgTypeName Arg3 , ArgTypeName Arg4 ) \
-	{ return this->Contract-> Function (ArgName Arg1 , ArgName Arg2 , ArgName Arg3 , ArgName Arg4); }
+#define Function1(Capability, Return, Function, Arg1) Return Function(ArgTypeName Arg1 ) requires (!bIsConst) \
+	{ return this->Contract-> Function (ForwardArg Arg1 ); }
+#define Function2(Capability, Return, Function, Arg1, Arg2) Return Function(ArgTypeName Arg1 , ArgTypeName Arg2 ) requires (!bIsConst) \
+	{ return this->Contract-> Function (ForwardArg Arg1 , ForwardArg Arg2 ); }
+#define Function3(Capability, Return, Function, Arg1, Arg2, Arg3) Return Function(ArgTypeName Arg1 , ArgTypeName Arg2 , ArgTypeName Arg3 ) requires (!bIsConst) \
+	{ return this->Contract-> Function (ForwardArg Arg1 , ForwardArg Arg2 , ForwardArg Arg3); }
+#define Function4(Capability, Return, Function, Arg1, Arg2, Arg3, Arg4) Return Function(ArgTypeName Arg1 , ArgTypeName Arg2 , ArgTypeName Arg3 , ArgTypeName Arg4 ) requires (!bIsConst) \
+	{ return this->Contract-> Function (ForwardArg Arg1 , ForwardArg Arg2 , ForwardArg Arg3 , ForwardArg Arg4); }
 
 #define ConstFunction0(Capability, Return, Function) Return Function() const \
 	{ return this->Contract-> Function (); }
 #define ConstFunction1(Capability, Return, Function, Arg1) Return Function(ArgTypeName Arg1 ) const \
-	{ return this->Contract-> Function (ArgName Arg1 ); }
+	{ return this->Contract-> Function (ForwardArg Arg1 ); }
 #define ConstFunction2(Capability, Return, Function, Arg1, Arg2) Return Function(ArgTypeName Arg1 , ArgTypeName Arg2 ) const \
-	{ return this->Contract-> Function (ArgName Arg1 , ArgName Arg2 ); }
+	{ return this->Contract-> Function (ForwardArg Arg1 , ForwardArg Arg2 ); }
 #define ConstFunction3(Capability, Return, Function, Arg1, Arg2, Arg3) Return Function(ArgTypeName Arg1 , ArgTypeName Arg2 , ArgTypeName Arg3 ) const \
-	{ return this->Contract-> Function (ArgName Arg1 , ArgName Arg2 , ArgName Arg3); }
+	{ return this->Contract-> Function (ForwardArg Arg1 , ForwardArg Arg2 , ForwardArg Arg3); }
 #define ConstFunction4(Capability, Return, Function, Arg1, Arg2, Arg3, Arg4) Return Function(ArgTypeName Arg1 , ArgTypeName Arg2 , ArgTypeName Arg3 , ArgTypeName Arg4 ) const \
-	{ return this->Contract-> Function (ArgName Arg1 , ArgName Arg2 , ArgName Arg3 , ArgName Arg4); }
-
-#define CapabilityStart(Capability, Flags) \
-	template<typename... Capabilities> \
-	struct TForwarder< Capability , Capabilities...> : TForwarder<Capabilities...> \
-	{ \
-		TForwarder() = default; \
-		explicit TForwarder(IContextContract& Contract) : TForwarder<Capabilities...>(Contract) {}
+	{ return this->Contract-> Function (ForwardArg Arg1 , ForwardArg Arg2 , ForwardArg Arg3 , ForwardArg Arg4); }
 
 #define CapabilityEnd(Capability) \
-	}; \
-	struct Capability##Forwarder : TForwarder< Capability , I##Capability <ForwardTesterBase>> {}; \
-	static_assert(!std::is_abstract_v<Capability##Forwarder>, \
-		"One or more functions are not implemented in the " #Capability " forwarder.");
+	};
 
 #define WithWrappers 1
 	
@@ -85,6 +80,7 @@ namespace UE::Editor::DataStorage::Queries::Private
 
 #undef ArgTypeName
 #undef ArgName
+#undef ForwardArg
 #undef Function0
 #undef Function1
 #undef Function2
diff --git a/Engine/Source/Runtime/TypedElementFramework/Public/Elements/Framework/TypedElementQueryContext.h b/Engine/Source/Runtime/TypedElementFramework/Public/Elements/Framework/TypedElementQueryContext.h
--- a/Engine/Source/Runtime/TypedElementFramework/Public/Elements/Framework/TypedElementQueryContext.h
+++ b/Engine/Source/Runtime/TypedElementFramework/Public/Elements/Framework/TypedElementQueryContext.h
@@ -25,25 +25,172 @@ namespace UE::Editor::DataStorage::Queries
 	 * query callback to improve re-usability or use a predefined query context if the query callback doesn't need to be reused.
 	 */
 	template<ContextCapability... CapabilityTypes>
-	struct TQueryContext final : Private::TForwarder< CapabilityTypes...>
+	struct TQueryContext final : Private::TForwarder<false, CapabilityTypes...>
 	{
 		using Capabilities = TCapabilityStore<CapabilityTypes...>;
 
 		static_assert(Capabilities::bIsSingle || Capabilities::bIsBatch,
 			"One or more capabilities that only support single or batch processing were mixed.");
 
+		template<ContextCapability... InCapabilityTypes>
+		friend struct TForwardingQueryContext;
+		
+		template<ContextCapability... InCapabilityTypes>
+		friend struct TForwardingConstQueryContext;
+
 		TQueryContext() = default;
 		TQueryContext(IContextContract& Contract);
 	};
+	
+	/**
+	 * Template to composite a query context using context capabilities.
+	 * Any query callback that requires interacting with the editor data storage requires a context to get access. Each context capability
+	 * provides access to a different kind of functionality. Some of the capabilities are mutually exclusive and not all places that accept
+	 * a query callback support the same capabilities. It's therefore recommended to only include the capabilities that are needed by the
+	 * query callback to improve re-usability or use a predefined query context if the query callback doesn't need to be reused.
+	 */
+	template<ContextCapability... CapabilityTypes>
+	struct TConstQueryContext final : Private::TForwarder<true, CapabilityTypes...>
+	{
+		using Capabilities = TCapabilityStore<CapabilityTypes...>;
+
+		static_assert(Capabilities::bIsSingle || Capabilities::bIsBatch,
+			"One or more capabilities that only support single or batch processing were mixed.");
+
+		template<ContextCapability... InCapabilityTypes>
+		friend struct TForwardingConstQueryContext;
+
+		TConstQueryContext() = default;
+		TConstQueryContext(const IContextContract& Contract);
+	};
+
+	template<ContextCapability... CapabilityTypes>
+	struct TForwardingQueryContext final : Private::TForwarder<false, CapabilityTypes...>
+	{
+		using Capabilities = TCapabilityStore<CapabilityTypes...>;
+
+		static_assert(Capabilities::bIsSingle || Capabilities::bIsBatch,
+			"One or more capabilities that only support single or batch processing were mixed.");
+
+		TForwardingQueryContext() = default;
+		TForwardingQueryContext(IContextContract& Contract);
+
+		template<ContextCapability... InCapabilityTypes>
+		TForwardingQueryContext(TQueryContext<InCapabilityTypes...> QueryContext);
+
+		operator IContextContract& ();
+		operator const IContextContract& () const;
+	};
+
+	template<ContextCapability... CapabilityTypes>
+	struct TForwardingConstQueryContext final : Private::TForwarder<true, CapabilityTypes...>
+	{
+		using Capabilities = TCapabilityStore<CapabilityTypes...>;
+
+		static_assert(Capabilities::bIsSingle || Capabilities::bIsBatch,
+			"One or more capabilities that only support single or batch processing were mixed.");
+
+		TForwardingConstQueryContext() = default;
+		TForwardingConstQueryContext(const IContextContract& Contract);
+
+		template<ContextCapability... InCapabilityTypes>
+		TForwardingConstQueryContext(TQueryContext<InCapabilityTypes...> QueryContext);
+		template<ContextCapability... InCapabilityTypes>
+		TForwardingConstQueryContext(TConstQueryContext<InCapabilityTypes...> QueryContext);
+
+		operator const IContextContract& () const;
+	};
 
 
 	//
 	// Implementations.
 	//
 	
+	//
+	// TQueryContext
+	//
+
 	template<ContextCapability... CapabilityTypes>
 	TQueryContext<CapabilityTypes...>::TQueryContext(IContextContract& Contract)
-		: Private::TForwarder< CapabilityTypes...>(Contract)
+		: Private::TForwarder<false, CapabilityTypes...>(Contract)
+	{
+	}
+
+	//
+	// TConstQueryContext
+	//
+
+	template<ContextCapability... CapabilityTypes>
+	TConstQueryContext<CapabilityTypes...>::TConstQueryContext(const IContextContract& Contract)
+		: Private::TForwarder<true, CapabilityTypes...>(Contract)
+	{
+	}
+
+	//
+	// TForwardingQueryContext
+	// 
+
+	template<ContextCapability... CapabilityTypes>
+	TForwardingQueryContext<CapabilityTypes...>::TForwardingQueryContext(IContextContract& Contract)
+		: Private::TForwarder<false, CapabilityTypes...>(Contract)
+	{
+	}
+
+	template<ContextCapability... CapabilityTypes>
+	template<ContextCapability... InCapabilityTypes>
+	TForwardingQueryContext<CapabilityTypes...>::TForwardingQueryContext(TQueryContext<InCapabilityTypes...> QueryContext)
+		: Private::TForwarder<false, CapabilityTypes...>(*QueryContext.Contract)
+	{
+		static_assert(
+			(std::is_base_of_v<CapabilityTypes, TCapabilityStore<InCapabilityTypes...>> && ...),
+			"Forwarding query context requires more or different capabilities than the input query context provides.");
+	}
+
+	template<ContextCapability... CapabilityTypes>
+	TForwardingQueryContext<CapabilityTypes...>::operator IContextContract& ()
+	{
+		return *this->Contract;
+	}
+
+	template<ContextCapability... CapabilityTypes>
+	TForwardingQueryContext<CapabilityTypes...>::operator const IContextContract& () const
+	{ 
+		return *this->Contract;
+	}
+
+	//
+	// TForwardingConstQueryContext
+	// 
+
+	template<ContextCapability... CapabilityTypes>
+	TForwardingConstQueryContext<CapabilityTypes...>::TForwardingConstQueryContext(const IContextContract& Contract)
+		: Private::TForwarder<true, CapabilityTypes...>(Contract)
+	{
+	}
+
+	template<ContextCapability... CapabilityTypes>
+	template<ContextCapability... InCapabilityTypes>
+	TForwardingConstQueryContext<CapabilityTypes...>::TForwardingConstQueryContext(TQueryContext<InCapabilityTypes...> QueryContext)
+		: Private::TForwarder<true, CapabilityTypes...>(*QueryContext.Contract)
+	{
+		static_assert(
+			(std::is_base_of_v<CapabilityTypes, TCapabilityStore<InCapabilityTypes...>> && ...),
+			"Forwarding const query context requires more or different capabilities than the input query context provides.");
+	}
+
+	template<ContextCapability... CapabilityTypes>
+	template<ContextCapability... InCapabilityTypes>
+	TForwardingConstQueryContext<CapabilityTypes...>::TForwardingConstQueryContext(TConstQueryContext<InCapabilityTypes...> QueryContext)
+		: Private::TForwarder<true, CapabilityTypes...>(*QueryContext.Contract)
+	{
+		static_assert(
+			(std::is_base_of_v<CapabilityTypes, TCapabilityStore<InCapabilityTypes...>> && ...),
+			"Forwarding const query context requires more or different capabilities than the input query context provides.");
+	}
+
+	template<ContextCapability... CapabilityTypes>
+	TForwardingConstQueryContext<CapabilityTypes...>::operator const IContextContract& () const
 	{
+		return *this->Contract;
 	}
 } // namespace UE::Editor::DataStorage::Queries
diff --git a/Engine/Source/Runtime/TypedElementFramework/Public/Elements/Framework/TypedElementQueryContextImplementation.h b/Engine/Source/Runtime/TypedElementFramework/Public/Elements/Framework/TypedElementQueryContextImplementation.h
--- a/Engine/Source/Runtime/TypedElementFramework/Public/Elements/Framework/TypedElementQueryContextImplementation.h
+++ b/Engine/Source/Runtime/TypedElementFramework/Public/Elements/Framework/TypedElementQueryContextImplementation.h
@@ -38,7 +38,7 @@ namespace UE::Editor::DataStorage::Queries
 	 * failure. This template will satisfy the IContextContract, but will implement placeholder functions for the
 	 * functions of capabilities that are not supported.
 	 */
-	template<typename EnvironmentType, typename... ImplementationTypes>
+	template<bool bIsConst, typename EnvironmentType, typename... ImplementationTypes>
 	class TQueryContextImpl final : public IContextContract
 	{
 	public:
@@ -50,16 +50,21 @@ namespace UE::Editor::DataStorage::Queries
 
 		// Generic access
 
+		static constexpr uint64 GetCapabilityMask();
 		template<ContextCapability RequestedCapability>
 		static constexpr bool SupportsCapability();
 		static bool SupportsCapability(const FName& Capability);
 		static constexpr bool SupportsCapability(int32 CapabilityId);
-
+		static constexpr bool SupportsCapabilities(uint64 CapabilityMask);
+		
+		static void StaticReportCompatibility(uint64 CapabilityMask);
 		static TConstArrayView<FName> StaticGetSupportedCapabilities();
 
+		virtual constexpr bool CheckCompatibility(uint64 CapabilityMask) const override;
+		virtual void ReportCompatibility(uint64 CapabilityMask) const override;
 		virtual TConstArrayView<FName> GetSupportedCapabilities() const override;
 		
-		EnvironmentType& Environment;
+		std::conditional_t<bIsConst, const EnvironmentType&, EnvironmentType&> Environment;
 	};
 }
 // namespace UE::Editor::DataStorage::Queries
diff --git a/Engine/Source/Runtime/TypedElementFramework/Public/Elements/Framework/TypedElementQueryContextImplementation.inl b/Engine/Source/Runtime/TypedElementFramework/Public/Elements/Framework/TypedElementQueryContextImplementation.inl
--- a/Engine/Source/Runtime/TypedElementFramework/Public/Elements/Framework/TypedElementQueryContextImplementation.inl
+++ b/Engine/Source/Runtime/TypedElementFramework/Public/Elements/Framework/TypedElementQueryContextImplementation.inl
@@ -3,6 +3,7 @@
 #include <type_traits>
 #include "Containers/StringFwd.h"
 #include "DataStorage/CommonTypes.h"
+#include "Math/UnrealMathUtility.h"
 
 namespace UE::Editor::DataStorage::Queries
 {
@@ -27,45 +28,69 @@ namespace UE::Editor::DataStorage::Queries
 		using ImplementationType = typename FindImplementationType<RequestedCapability, ImplementationTypes...>::Type;
 	}
 
-	template<typename EnvironmentType, typename... ImplementationTypes>
-	TQueryContextImpl<EnvironmentType, ImplementationTypes...>::TQueryContextImpl(EnvironmentType& Environment)
+	template<bool bIsConst, typename EnvironmentType, typename... ImplementationTypes>
+	TQueryContextImpl<bIsConst, EnvironmentType, ImplementationTypes...>::TQueryContextImpl(EnvironmentType& Environment)
 		: Environment(Environment)
 	{
 	}
 
+	template<bool bIsConst, typename EnvironmentType, typename... ImplementationTypes>
+	constexpr uint64 TQueryContextImpl<bIsConst, EnvironmentType, ImplementationTypes...>::GetCapabilityMask()
+	{
+		return ((uint64(1) << ImplementationTypes::Id) | ...);
+	}
 
-	template<typename EnvironmentType, typename... ImplementationTypes>
+	template<bool bIsConst, typename EnvironmentType, typename... ImplementationTypes>
 	template<ContextCapability RequestedCapability>
-	constexpr bool TQueryContextImpl<EnvironmentType, ImplementationTypes...>::SupportsCapability()
+	constexpr bool TQueryContextImpl<bIsConst, EnvironmentType, ImplementationTypes...>::SupportsCapability()
 	{
 		return (std::is_base_of_v<ImplementsContextCapability<RequestedCapability>, ImplementationTypes> || ...);
 	}
 
-	template<typename EnvironmentType, typename... ImplementationTypes>
-	bool TQueryContextImpl<EnvironmentType, ImplementationTypes...>::SupportsCapability(const FName& Capability)
+	template<bool bIsConst, typename EnvironmentType, typename... ImplementationTypes>
+	bool TQueryContextImpl<bIsConst, EnvironmentType, ImplementationTypes...>::SupportsCapability(const FName& Capability)
 	{
 		return ((Capability == ImplementationTypes::Name) || ...);
 	}
 
-	template<typename EnvironmentType, typename... ImplementationTypes>
-	constexpr bool TQueryContextImpl<EnvironmentType, ImplementationTypes...>::SupportsCapability(int32 CapabilityId)
+	template<bool bIsConst, typename EnvironmentType, typename... ImplementationTypes>
+	constexpr bool TQueryContextImpl<bIsConst, EnvironmentType, ImplementationTypes...>::SupportsCapability(int32 CapabilityId)
 	{
 		return ((CapabilityId == ImplementationTypes::Id) || ...);
 	}
 
+	template<bool bIsConst, typename EnvironmentType, typename... ImplementationTypes>
+	constexpr bool TQueryContextImpl<bIsConst, EnvironmentType, ImplementationTypes...>::SupportsCapabilities(uint64 CapabilityMask)
+	{
+		return (GetCapabilityMask() & CapabilityMask) == CapabilityMask;
+	}
 
 	// Use a macro for this since it's impossible to do with a template since the function pointer that would be needed might not exist.
 	// The macro also has the added benefit that it can print the function name for better reporting.
-#define CallFunction(Capability, ReturnValue, FunctionName, ...) \
+	// CheckArgs and CallArgs are both parenthesized argument packs. CheckArgs uses lvalue argument names for the
+	// requires() expression while CallArgs uses Forward<> for the actual call. Keeping them separate incorrect evaluated requires() as 
+	// false when Forward<T> produces an rvalue for types that are not move-constructible (e.g. TFunctionRef).
+#define ExpandArgPack(...) __VA_ARGS__
+#define PrependCommaIfNeeded(...) __VA_OPT__(,) __VA_ARGS__
+#define CallFunction(ConstFunction, Capability, ReturnValue, FunctionName, CheckArgs, CallArgs) \
 	if constexpr (SupportsCapability<Capability>()) \
 	{ \
 		using ImplementationType = Private::ImplementationType<Capability, ImplementationTypes...>; \
 		constexpr bool bFunctionFound = requires() \
 				{ \
-					{ ImplementationType:: FunctionName (Environment __VA_OPT__(,) __VA_ARGS__ ) } -> UE::CSameAs< ReturnValue >; \
+					{ ImplementationType:: FunctionName (Environment PrependCommaIfNeeded(ExpandArgPack CheckArgs) ) } -> UE::CSameAs< ReturnValue >; \
 				}; \
-		static_assert(bFunctionFound, "Function '" #FunctionName "' in capability '" #Capability "' is missing from the implementation." ); \
-		return ImplementationType:: FunctionName ( Environment __VA_OPT__(,) __VA_ARGS__ ); \
+		if constexpr (!bIsConst || (bIsConst && ConstFunction)) \
+		{ \
+			static_assert(bFunctionFound, "Function '" #FunctionName "' in capability '" #Capability "' is missing from the implementation. This is required by context implementations." ); \
+			return ImplementationType:: FunctionName ( Environment PrependCommaIfNeeded(ExpandArgPack CallArgs) ); \
+		} \
+		else \
+		{ \
+			static_assert(!bFunctionFound, "Non-const function '" #FunctionName "' in capability '" #Capability "' is implemented for a const environment." ); \
+			checkf(false, TEXT("Function '" #FunctionName "' in capability '" #Capability "' cannot be called from a const environment.")); \
+			return UE::MakeDefault< ReturnValue >(); \
+		} \
 	} \
 	else \
 	{ \
@@ -78,61 +103,61 @@ namespace UE::Editor::DataStorage::Queries
 #define ForwardArg(Type, Name) Forward< Type >( Name )
 
 #define FunctionCommon(ReturnType) \
-	template<typename EnvironmentType, typename... ImplementationTypes> \
-	ReturnType TQueryContextImpl<EnvironmentType, ImplementationTypes...>::
+	template<bool bIsConst, typename EnvironmentType, typename... ImplementationTypes> \
+	ReturnType TQueryContextImpl<bIsConst, EnvironmentType, ImplementationTypes...>::
 
 #define CapabilityStart(Capability, Flags)
 
 #define Function0(Capability, Return, Function) \
 	FunctionCommon(Return) Function() \
 	{ \
-		CallFunction(Capability, Return, Function ) \
+		CallFunction(false, Capability, Return, Function, (), ()) \
 	}
 #define Function1(Capability, Return, Function, Arg1) \
 	FunctionCommon(Return) Function(ArgTypeName Arg1) \
 	{ \
-		CallFunction(Capability, Return, Function, ForwardArg Arg1 ) \
+		CallFunction(false, Capability, Return, Function, (ArgName Arg1), (ForwardArg Arg1)) \
 	}
 #define Function2(Capability, Return, Function, Arg1, Arg2) \
 	FunctionCommon(Return) Function(ArgTypeName Arg1, ArgTypeName Arg2) \
 	{ \
-		CallFunction(Capability, Return, Function, ForwardArg Arg1, ForwardArg Arg2 ) \
+		CallFunction(false, Capability, Return, Function, (ArgName Arg1, ArgName Arg2), (ForwardArg Arg1, ForwardArg Arg2)) \
 	}
 #define Function3(Capability, Return, Function, Arg1, Arg2, Arg3) \
 	FunctionCommon(Return) Function(ArgTypeName Arg1, ArgTypeName Arg2, ArgTypeName Arg3) \
 	{ \
-		CallFunction(Capability, Return, Function, ForwardArg Arg1, ForwardArg Arg2, ForwardArg Arg3 ) \
+		CallFunction(false, Capability, Return, Function, (ArgName Arg1, ArgName Arg2, ArgName Arg3), (ForwardArg Arg1, ForwardArg Arg2, ForwardArg Arg3)) \
 	}
 #define Function4(Capability, Return, Function, Arg1, Arg2, Arg3, Arg4) \
 	FunctionCommon(Return) Function(ArgTypeName Arg1, ArgTypeName Arg2, ArgTypeName Arg3, ArgTypeName Arg4) \
 	{ \
-		CallFunction(Capability, Return, Function, ForwardArg Arg1, ForwardArg Arg2, ForwardArg Arg3, ForwardArg Arg4 ) \
+		CallFunction(false, Capability, Return, Function, (ArgName Arg1, ArgName Arg2, ArgName Arg3, ArgName Arg4), (ForwardArg Arg1, ForwardArg Arg2, ForwardArg Arg3, ForwardArg Arg4)) \
 	}
 
 #define ConstFunction0(Capability, Return, Function) \
 	FunctionCommon(Return) Function() const \
 	{ \
-		CallFunction(Capability, Return, Function ) \
+		CallFunction(true, Capability, Return, Function, (), ()) \
 	}
 #define ConstFunction1(Capability, Return, Function, Arg1) \
 	FunctionCommon(Return) Function(ArgTypeName Arg1) const \
 	{ \
-		CallFunction(Capability, Return, Function, ForwardArg Arg1 ) \
+		CallFunction(true, Capability, Return, Function, (ArgName Arg1), (ForwardArg Arg1)) \
 	}
 #define ConstFunction2(Capability, Return, Function, Arg1, Arg2) \
 	FunctionCommon(Return) Function(ArgTypeName Arg1, ArgTypeName Arg2) const \
 	{ \
-		CallFunction(Capability, Return, Function, ForwardArg Arg1, ForwardArg Arg2 ) \
+		CallFunction(true, Capability, Return, Function, (ArgName Arg1, ArgName Arg2), (ForwardArg Arg1, ForwardArg Arg2)) \
 	}
 #define ConstFunction3(Capability, Return, Function, Arg1, Arg2, Arg3) \
 	FunctionCommon(Return) Function(ArgTypeName Arg1, ArgTypeName Arg2, ArgTypeName Arg3) const \
 	{ \
-		CallFunction(Capability, Return, Function, ForwardArg Arg1, ForwardArg Arg2, ForwardArg Arg3 ) \
+		CallFunction(true, Capability, Return, Function, (ArgName Arg1, ArgName Arg2, ArgName Arg3), (ForwardArg Arg1, ForwardArg Arg2, ForwardArg Arg3)) \
 	}
 #define ConstFunction4(Capability, Return, Function, Arg1, Arg2, Arg3, Arg4) \
 	FunctionCommon(Return) Function(ArgTypeName Arg1, ArgTypeName Arg2, ArgTypeName Arg3, ArgTypeName Arg4) const \
 	{ \
-		CallFunction(Capability, Return, Function, ForwardArg Arg1, ForwardArg Arg2, ForwardArg Arg3, ForwardArg Arg4 ) \
+		CallFunction(true, Capability, Return, Function, (ArgName Arg1, ArgName Arg2, ArgName Arg3, ArgName Arg4), (ForwardArg Arg1, ForwardArg Arg2, ForwardArg Arg3, ForwardArg Arg4)) \
 	}
 
 #define CapabilityEnd(Capability)
@@ -142,6 +167,8 @@ namespace UE::Editor::DataStorage::Queries
 #undef ArgTypeName
 #undef ArgName
 #undef ForwardArg
+#undef ExpandArgPack
+#undef PrependCommaIfNeeded
 #undef FunctionCommon
 #undef Function0
 #undef Function1
@@ -160,8 +187,37 @@ namespace UE::Editor::DataStorage::Queries
 	//
 	// Generic access
 	
-	template<typename EnvironmentType, typename... ImplementationTypes>
-	TConstArrayView<FName> TQueryContextImpl<EnvironmentType, ImplementationTypes...>::StaticGetSupportedCapabilities()
+	template<bool bIsConst, typename EnvironmentType, typename... ImplementationTypes>
+	void TQueryContextImpl<bIsConst, EnvironmentType, ImplementationTypes...>::StaticReportCompatibility(uint64 CapabilityMask)
+	{
+		uint64 MissingCapabilitiesMask = CapabilityMask & ~GetCapabilityMask();
+		if (MissingCapabilitiesMask != 0)
+		{
+			TStringBuilder<256> Builder;
+
+			while (MissingCapabilitiesMask)
+			{
+				uint64 CapabilityIndex = FMath::CountTrailingZeros64(MissingCapabilitiesMask);
+				MissingCapabilitiesMask &= MissingCapabilitiesMask - 1;
+
+				Builder.Append(TEXT("Requested query context '"));
+				Builder.Append(IContextContract::GetCapabilityName(static_cast<int32>(CapabilityIndex)).ToString());
+				Builder.Append(TEXT("' is not supported by query context implementation.\n"));
+			}
+
+			Builder.Append(TEXT("The following types are supported by this query context: \n"));
+			for (const FName& CapabilityName : StaticGetSupportedCapabilities())
+			{
+				Builder.Append(TEXT("    "));
+				Builder.Append(CapabilityName.ToString());
+				Builder.AppendChar(TEXT('\n'));
+			}
+			checkf(false, TEXT("%s"), *Builder);
+		}
+	}
+
+	template<bool bIsConst, typename EnvironmentType, typename... ImplementationTypes>
+	TConstArrayView<FName> TQueryContextImpl<bIsConst, EnvironmentType, ImplementationTypes...>::StaticGetSupportedCapabilities()
 	{
 		static_assert(sizeof...(ImplementationTypes) > 0,
 			"At least one context capability needs to be implemented to work with a query context contract.");
@@ -170,8 +226,20 @@ namespace UE::Editor::DataStorage::Queries
 		return Result;
 	}
 
-	template<typename EnvironmentType, typename... ImplementationTypes>
-	TConstArrayView<FName> TQueryContextImpl<EnvironmentType, ImplementationTypes...>::GetSupportedCapabilities() const
+	template<bool bIsConst, typename EnvironmentType, typename... ImplementationTypes>
+	constexpr bool TQueryContextImpl<bIsConst, EnvironmentType, ImplementationTypes...>::CheckCompatibility(uint64 CapabilityMask) const
+	{
+		return SupportsCapabilities(CapabilityMask);
+	}
+
+	template<bool bIsConst, typename EnvironmentType, typename... ImplementationTypes>
+	void TQueryContextImpl<bIsConst, EnvironmentType, ImplementationTypes...>::ReportCompatibility(uint64 CapabilityMask) const
+	{
+		StaticReportCompatibility(CapabilityMask);
+	}
+
+	template<bool bIsConst, typename EnvironmentType, typename... ImplementationTypes>
+	TConstArrayView<FName> TQueryContextImpl<bIsConst, EnvironmentType, ImplementationTypes...>::GetSupportedCapabilities() const
 	{
 		return StaticGetSupportedCapabilities();
 	}
diff --git a/Engine/Source/Runtime/TypedElementFramework/Public/Elements/Framework/TypedElementQueryContract.h b/Engine/Source/Runtime/TypedElementFramework/Public/Elements/Framework/TypedElementQueryContract.h
--- a/Engine/Source/Runtime/TypedElementFramework/Public/Elements/Framework/TypedElementQueryContract.h
+++ b/Engine/Source/Runtime/TypedElementFramework/Public/Elements/Framework/TypedElementQueryContract.h
@@ -49,7 +49,7 @@ namespace UE::Editor::DataStorage::Queries
 			static FName GetCapabilityName(int32 CapabilityId)
 			{
 				FName Result;
-				((Bases<IContextCapability>::Id == CapabilityId && (Result = Bases<IContextCapability>::Name, true)), ...);
+				(void)(((Bases<IContextCapability>::Id == CapabilityId && (Result = Bases<IContextCapability>::Name, true)) || ...));
 				return Result;
 			}
 
@@ -125,6 +125,19 @@ namespace UE::Editor::DataStorage::Queries
 				return Result;
 			}
 
+			template<typename Implementation>
+			static constexpr uint64 GetSupportedCapabilityMask()
+			{
+				uint64 Result = 0;
+				([&]
+				{
+					if constexpr (SupportsCapability<Implementation, Bases<IContextCapability>>())
+					{
+						Result = Result | (uint64(1) << Bases<IContextCapability>::Id);
+					}
+				}(), ...);
+				return Result;
+			}
 		};
 
 		//Used to remove the first template argument as that's a placeholder to allow the X-macros to work with the template.
@@ -163,6 +176,22 @@ namespace UE::Editor::DataStorage::Queries
 	{
 		virtual ~IContextContract() override = default;
 
+		virtual bool constexpr CheckCompatibility(uint64 CapabilityMask) const = 0;
+		template<ContextCapability... Capabilities>
+		constexpr bool CheckCompatibility() const
+		{
+			constexpr uint64 CapabilityMask = ((uint64(1) << Capabilities::Id) | ...);
+			return CheckCompatibility(CapabilityMask);
+		}
+
+		virtual void ReportCompatibility(uint64 CapabilityMask) const = 0;
+		template<ContextCapability... Capabilities>
+		void ReportCompatibility() const
+		{
+			constexpr uint64 CapabilityMask = ((uint64(1) << Capabilities::Id) | ...);
+			ReportCompatibility(CapabilityMask);
+		}
+
 		virtual TConstArrayView<FName> GetSupportedCapabilities() const = 0;
 	};
 
diff --git a/Engine/Source/Runtime/TypedElementFramework/Public/Elements/Framework/TypedElementQueryFunctionArgumentTypes.h b/Engine/Source/Runtime/TypedElementFramework/Public/Elements/Framework/TypedElementQueryFunctionArgumentTypes.h
--- a/Engine/Source/Runtime/TypedElementFramework/Public/Elements/Framework/TypedElementQueryFunctionArgumentTypes.h
+++ b/Engine/Source/Runtime/TypedElementFramework/Public/Elements/Framework/TypedElementQueryFunctionArgumentTypes.h
@@ -3,6 +3,7 @@
 #pragma once
 
 #include "Containers/ArrayView.h"
+#include "DataStorage/Handles.h"
 
 namespace UE::Editor::DataStorage::Queries
 {
@@ -44,7 +45,7 @@ namespace UE::Editor::DataStorage::Queries
 	template<typename ResultType>
 	struct TResult
 	{
-		virtual void Add(ResultType ResultValue) = 0;
+		virtual void Add(RowHandle Row, ResultType ResultValue) = 0;
 	};
 	template<> struct TResult<void>{}; // Used as a placeholder for internal use.
 
diff --git a/Engine/Source/Runtime/TypedElementFramework/Public/Elements/Framework/TypedElementQueryFunctionArguments.h b/Engine/Source/Runtime/TypedElementFramework/Public/Elements/Framework/TypedElementQueryFunctionArguments.h
--- a/Engine/Source/Runtime/TypedElementFramework/Public/Elements/Framework/TypedElementQueryFunctionArguments.h
+++ b/Engine/Source/Runtime/TypedElementFramework/Public/Elements/Framework/TypedElementQueryFunctionArguments.h
@@ -18,7 +18,7 @@ namespace UE::Editor::DataStorage::Queries::Private
 	using UnreferencedType = std::remove_reference_t<std::remove_pointer_t<T>>;
 	template<typename T>
 	using FoundationalType = std::remove_cv_t<UnreferencedType<T>>;
-	template<typename T> concept ContextType = std::is_base_of_v<FQueryContext, T>;
+	template<typename T> concept ContextType = std::is_base_of_v<FQueryContext, T> || std::is_base_of_v<FConstQueryContext, T>;
 
 	enum class EArgumentFlags
 	{
@@ -99,7 +99,7 @@ namespace UE::Editor::DataStorage::Queries::Private
 	template<typename T> requires ContextType<FoundationalType<T>>
 	struct TArgument<T>
 	{
-		static_assert(!std::is_reference_v<T> && !std::is_pointer_v<T>, "The contexts can only be passed in by value.");
+		static_assert(!std::is_reference_v<T> && !std::is_pointer_v<T>, "The context can only be passed in by value.");
 		
 		using BaseType = FoundationalType<T>;
 		using ContextCapabilities = BaseType::Capabilities;
@@ -107,7 +107,7 @@ namespace UE::Editor::DataStorage::Queries::Private
 		constexpr static EArgumentFlags Flags =
 			(ContextCapabilities::bIsSingle ? EArgumentFlags::SingleArgument : EArgumentFlags(0)) |
 			(ContextCapabilities::bIsBatch ? EArgumentFlags::BatchArgument : EArgumentFlags(0)) |
-			(std::is_const_v<UnreferencedType<T>> ? EArgumentFlags::IsConst : EArgumentFlags::IsMutable) |
+			(std::is_base_of_v<FConstQueryContext, BaseType> ? EArgumentFlags::IsConst : EArgumentFlags::IsMutable) |
 			EArgumentFlags::Type_Context; 
 		using ArgumentType = BaseType; // Take by value because the context is a forwarder that only stores a single pointer.
 		using ColumnType = void;
@@ -195,6 +195,9 @@ namespace UE::Editor::DataStorage::Queries::Private
 		template<EArgumentFlags Flags>
 		constexpr static int32 CountFlags();
 
+		template<int32 Index, EArgumentFlags Flags>
+		constexpr static bool HasFlags();
+
 		template<EArgumentFlags Flags>
 		static TConstArrayView<const UScriptStruct*> ListColumns();
 	
@@ -215,6 +218,8 @@ namespace UE::Editor::DataStorage::Queries::Private
 		constexpr static int32 ContextIndex = FirstIndexOfType<EArgumentFlags::Type_Context>();
 		using ContextType = typename TIndexToArgInfo<ContextIndex, Args...>::BaseType;
 		static void SetContext(ArgumentList& Arguments, IContextContract& Contract);
+		static void SetContext(ArgumentList& Arguments, const IContextContract& Contract);
+		constexpr static bool IsContextConst();
 		constexpr static int32 CountContexts();
 
 		constexpr static int32 FlowIndex = FirstIndexOfType<EArgumentFlags::Type_FlowControl>();
@@ -235,6 +240,8 @@ namespace UE::Editor::DataStorage::Queries::Private
 		template<int32 Index>
 		static void SetContextUnguarded(ArgumentList& Arguments, IContextContract& Contract);
 		template<int32 Index>
+		static void SetContextUnguarded(ArgumentList& Arguments, const IContextContract& Contract);
+		template<int32 Index>
 		static void SetFlowControlUnguarded(ArgumentList& Arguments, EFlowControl& FlowControl);
 		template<typename ArgumentType>
 		static void IncrementColumnUnguarded(ArgumentList& Arguments);
diff --git a/Engine/Source/Runtime/TypedElementFramework/Public/Elements/Framework/TypedElementQueryFunctionArguments.inl b/Engine/Source/Runtime/TypedElementFramework/Public/Elements/Framework/TypedElementQueryFunctionArguments.inl
--- a/Engine/Source/Runtime/TypedElementFramework/Public/Elements/Framework/TypedElementQueryFunctionArguments.inl
+++ b/Engine/Source/Runtime/TypedElementFramework/Public/Elements/Framework/TypedElementQueryFunctionArguments.inl
@@ -45,6 +45,14 @@ namespace UE::Editor::DataStorage::Queries::Private
 		return (0 + ... + (EnumHasAllFlags(TArgument<Args>::Flags, Flags) ? 1 : 0));
 	}
 
+	template<typename... Args>
+	template<int32 Index, EArgumentFlags Flags>
+	constexpr bool TArgumentInfo<Args...>::HasFlags()
+	{
+		using NthArg = typename TTupleElement<Index, TTuple<Args...>>::Type;
+		return EnumHasAllFlags(TArgument<NthArg>::Flags, Flags);
+	}
+
 	template<typename... Args>
 	template<typename T>
 	void TArgumentInfo<Args...>::SetResult(ArgumentList& Arguments, TResult<T>& Result)
@@ -70,6 +78,28 @@ namespace UE::Editor::DataStorage::Queries::Private
 		}
 	}
 
+	template<typename... Args>
+	void TArgumentInfo<Args...>::SetContext(ArgumentList& Arguments, const IContextContract& Contract)
+	{
+		if constexpr (ContextIndex >= 0)
+		{
+			SetContextUnguarded<ContextIndex>(Arguments, Contract);
+		}
+	}
+
+	template<typename... Args>
+	constexpr bool TArgumentInfo<Args...>::IsContextConst()
+	{
+		if constexpr (ContextIndex >= 0)
+		{
+			return HasFlags<ContextIndex, EArgumentFlags::IsConst>();
+		}
+		else
+		{
+			return false;
+		}
+	}
+
 	template<typename... Args>
 	constexpr int32 TArgumentInfo<Args...>::CountContexts()
 	{
@@ -201,6 +231,13 @@ namespace UE::Editor::DataStorage::Queries::Private
 		Arguments.template Get<Index>() = typename TTupleElement<Index, ArgumentList>::Type(Contract);
 	}
 
+	template<typename... Args>
+	template<int32 Index>
+	void TArgumentInfo<Args...>::SetContextUnguarded(ArgumentList& Arguments, const IContextContract& Contract)
+	{
+		Arguments.template Get<Index>() = typename TTupleElement<Index, ArgumentList>::Type(Contract);
+	}
+
 	template<typename... Args>
 	template<int32 Index>
 	void TArgumentInfo<Args...>::SetFlowControlUnguarded(ArgumentList& Arguments, EFlowControl& FlowControl)
diff --git a/Engine/Source/Runtime/TypedElementFramework/Public/Elements/Framework/TypedElementQueryFunctions.h b/Engine/Source/Runtime/TypedElementFramework/Public/Elements/Framework/TypedElementQueryFunctions.h
--- a/Engine/Source/Runtime/TypedElementFramework/Public/Elements/Framework/TypedElementQueryFunctions.h
+++ b/Engine/Source/Runtime/TypedElementFramework/Public/Elements/Framework/TypedElementQueryFunctions.h
@@ -23,6 +23,7 @@ namespace UE::Editor::DataStorage::Queries
 		
 		virtual FRowHandleArrayView GetBatchRows() const = 0;
 		virtual void GetConstColumns(TArrayView<const void*> ColumnsData, TConstArrayView<const UScriptStruct*> ColumnTypes) const = 0;
+		// Not used in const queries.
 		virtual void GetMutableColumns(TArrayView<void*> ColumnsData, TConstArrayView<const UScriptStruct*> ColumnTypes) = 0;
 		// Only called when a batch is broken up in individual rows for the current active call.
 		virtual void SetCurrentRow(RowHandle Row) = 0;
@@ -31,19 +32,25 @@ namespace UE::Editor::DataStorage::Queries
 	enum class EFunctionCallConfig
 	{
 		None = 0,
+		/** Verifies if the function is compatible with the provided context. */
+		VerifyCapabilityCompatibility = 1 << 0,
+		/** If the function is not compatible with the provided context it will assert. Requires `VerifyCapabilityCompatibility`. */
+		ReportCapabilityCompatibility = 1 << 1,
 		/** If not all columns could be retrieved, return the default value for the return type, if applicable, and don't call the function. */
-		VerifyColumns = 1 << 0,
+		VerifyColumns = 1 << 2,
 	};
 	ENUM_CLASS_FLAGS(EFunctionCallConfig);
 
 	/** Storage for a function that can be used as part of a query.  */
-	template<typename ReturnType>
+	template<typename ReturnType, bool bIsConst>
 	class TQueryFunctionBase
 	{
 	public:
-		using ContractType = IContextContract&;
+		using ContractType = std::conditional_t<bIsConst, const IContextContract&, IContextContract&>;
 		using FunctionResponseType = IQueryFunctionResponse&;
-		using FunctionSpecializationCallback = bool(*)(FunctionResponseType Response, TArrayView<const void*> ConstColumns, TArrayView<void*> MutableColumns);
+		using FunctionSpecializationCallback = std::conditional_t<bIsConst,
+			bool(*)(FunctionResponseType Response, TArrayView<const void*> ConstColumns),
+			bool(*)(FunctionResponseType Response, TArrayView<const void*> ConstColumns, TArrayView<void*> MutableColumns)>;
 	
 		using WrapperFunctionType = TFunction<EFlowControl(TResult<ReturnType>& Result, ContractType Contract,
 				FunctionResponseType Response, FunctionSpecializationCallback Specialization)>;
@@ -51,7 +58,8 @@ namespace UE::Editor::DataStorage::Queries
 		TConstArrayView<const UScriptStruct*> ConstColumnTypes;
 		TConstArrayView<const UScriptStruct*> MutableColumnTypes;
 		WrapperFunctionType Function;
-		bool bIsSingleRowProcessor;
+		uint64 CapabilityMask = 0;
+		bool bIsSingleRowProcessor = false;
 
 	protected:
 		template<EFunctionCallConfig Config>
@@ -59,10 +67,10 @@ namespace UE::Editor::DataStorage::Queries
 	};
 
 	template<typename ReturnType>
-	class TQueryFunction final : public TQueryFunctionBase<ReturnType>
+	class TQueryFunction final : public TQueryFunctionBase<ReturnType, false>
 	{
 	public:
-		using typename TQueryFunctionBase<ReturnType>::FunctionResponseType;
+		using typename TQueryFunctionBase<ReturnType, false>::FunctionResponseType;
 
 		template<EFunctionCallConfig Config>
 		EFlowControl Call(TResult<ReturnType>& Result, IContextContract& Contract, FunctionResponseType Response) const
@@ -77,10 +85,10 @@ namespace UE::Editor::DataStorage::Queries
 	};
 
 	template<>
-	class TQueryFunction<void> final : public TQueryFunctionBase<void>
+	class TQueryFunction<void> final : public TQueryFunctionBase<void, false>
 	{
 	public:
-		using typename TQueryFunctionBase<void>::FunctionResponseType;
+		using typename TQueryFunctionBase<void, false>::FunctionResponseType;
 
 		template<EFunctionCallConfig Config>
 		EFlowControl Call(IContextContract& Contract, FunctionResponseType Response) const
@@ -95,9 +103,49 @@ namespace UE::Editor::DataStorage::Queries
 		}
 	};
 
+	template<typename ReturnType>
+	class TConstQueryFunction final : public TQueryFunctionBase<ReturnType, true>
+	{
+	public:
+		using typename TQueryFunctionBase<ReturnType, true>::FunctionResponseType;
+
+		template<EFunctionCallConfig Config>
+		EFlowControl Call(TResult<ReturnType>& Result, const IContextContract& Contract, FunctionResponseType Response) const
+		{
+			return this->template CallInternal<Config>(Result, Contract, Response);
+		}
+
+		EFlowControl Call(TResult<ReturnType>& Result, const IContextContract& Contract, FunctionResponseType Response) const
+		{
+			return this->template CallInternal<EFunctionCallConfig::None>(Result, Contract, Response);
+		}
+	};
+
+	template<>
+	class TConstQueryFunction<void> final : public TQueryFunctionBase<void, true>
+	{
+	public:
+		using typename TQueryFunctionBase<void, true>::FunctionResponseType;
+
+		template<EFunctionCallConfig Config>
+		EFlowControl Call(const IContextContract& Contract, FunctionResponseType Response) const
+		{
+			TResult<void> Dummy;
+			return this->template CallInternal<Config>(Dummy, Contract, Response);
+		}
+
+		EFlowControl Call(const IContextContract& Contract, FunctionResponseType Response) const
+		{
+			return this->Call<EFunctionCallConfig::None>(Contract, Response);
+		}
+	};
+
 	template<typename Return, FunctionType Function>
 	TQueryFunction<Return> BuildQueryFunction(Function&& Callback);
 
+	template<typename Return, FunctionType Function>
+	TConstQueryFunction<Return> BuildConstQueryFunction(Function&& Callback);
+
 } // namespace UE::Editor::DataStorage::Queries
 
 #include "Elements/Framework/TypedElementQueryFunctions.inl"
diff --git a/Engine/Source/Runtime/TypedElementFramework/Public/Elements/Framework/TypedElementQueryFunctions.inl b/Engine/Source/Runtime/TypedElementFramework/Public/Elements/Framework/TypedElementQueryFunctions.inl
--- a/Engine/Source/Runtime/TypedElementFramework/Public/Elements/Framework/TypedElementQueryFunctions.inl
+++ b/Engine/Source/Runtime/TypedElementFramework/Public/Elements/Framework/TypedElementQueryFunctions.inl
@@ -80,7 +80,7 @@ namespace UE::Editor::DataStorage::Queries
 				(ArgumentInfo::bIsSingle || ArgumentInfo::bIsBatch);
 			static constexpr bool bIsSingleRowProcessor = ArgumentInfo::bIsSingle;
 
-			template<typename RequestedReturnType>
+			template<typename RequestedReturnType, bool bIsConst>
 			static constexpr void StaticValidation()
 			{
 				static_assert(bIsValidQueryFunction, "The provided function is not compatible with a query callback. "
@@ -88,6 +88,12 @@ namespace UE::Editor::DataStorage::Queries
 
 				static_assert(ArgumentInfo::CountContexts() <= 1, "Only zero or one context can be added as an argument.");
 				static_assert(ArgumentInfo::CountFlowControls() <= 1, "Only zero or one flow control arguments can be added.");
+				
+				if constexpr (bIsConst)
+				{
+					static_assert(ArgumentInfo::IsContextConst(), "Const functions require a const query context (TConstQueryContext).");
+					static_assert(ArgumentInfo::CountMutableColumns() == 0, "Const functions can not bind to mutable columns.");
+				}
 
 				if constexpr (std::is_same_v<RequestedReturnType, void>)
 				{
@@ -130,18 +136,41 @@ namespace UE::Editor::DataStorage::Queries
 			template<typename RequestedReturnType>
 			static constexpr void SetupResult(TQueryFunction<RequestedReturnType>& Result)
 			{
+				if constexpr (!std::is_same_v<void, ContextType>)
+				{
+					Result.CapabilityMask = IContextContract::GetSupportedCapabilityMask<typename ContextType::Capabilities>();
+				}
+				else
+				{
+					Result.CapabilityMask = 0;
+				}
 				Result.ConstColumnTypes = ArgumentInfo::ListConstColumns();
 				Result.MutableColumnTypes = ArgumentInfo::ListMutableColumns();
 				Result.bIsSingleRowProcessor = bIsSingleRowProcessor;
 			}
+
+			template<typename RequestedReturnType>
+			static constexpr void SetupResult(TConstQueryFunction<RequestedReturnType>& Result)
+			{
+				if constexpr (!std::is_same_v<void, ContextType>)
+				{
+					Result.CapabilityMask = IContextContract::GetSupportedCapabilityMask<typename ContextType::Capabilities>();
+				}
+				else
+				{
+					Result.CapabilityMask = 0;
+				}
+				Result.ConstColumnTypes = ArgumentInfo::ListConstColumns();
+				Result.bIsSingleRowProcessor = bIsSingleRowProcessor;
+			}
 		};
 
-		template<typename FunctionInfo, typename ArgumentInfo, typename ReturnType, typename FunctionType>
+		template<bool bIsConst, typename FunctionInfo, typename ArgumentInfo, typename ReturnType, typename FunctionType>
 		EFlowControl CallBody(
 			TResult<ReturnType>& QueryResult,
-			typename TQueryFunctionBase<ReturnType>::ContractType Contract,
-			typename TQueryFunctionBase<ReturnType>::FunctionResponseType Response,
-			typename TQueryFunctionBase<ReturnType>::FunctionSpecializationCallback Specialization,
+			typename TQueryFunctionBase<ReturnType, bIsConst>::ContractType Contract,
+			typename TQueryFunctionBase<ReturnType, bIsConst>::FunctionResponseType Response,
+			typename TQueryFunctionBase<ReturnType, bIsConst>::FunctionSpecializationCallback Specialization,
 			FunctionType&& Callback)
 		{
 			using ArgumentList = typename ArgumentInfo::ArgumentList;
@@ -150,9 +179,14 @@ namespace UE::Editor::DataStorage::Queries
 			const void** ConstColumnsData = static_cast<const void**>(FMemory_Alloca(sizeof(void*) * ConstColumnTypes.Num()));
 			TArrayView<const void*> ConstColumns(ConstColumnsData, ConstColumnTypes.Num());
 
-			TConstArrayView<const UScriptStruct*> MutableColumnTypes = ArgumentInfo::ListMutableColumns();
-			void** MutableColumnsData = static_cast<void**>(FMemory_Alloca(sizeof(void*) * MutableColumnTypes.Num()));
-			TArrayView<void*> MutableColumns(MutableColumnsData, MutableColumnTypes.Num());
+			TConstArrayView<const UScriptStruct*> MutableColumnTypes;
+			TArrayView<void*> MutableColumns;
+			if constexpr (!bIsConst)
+			{
+				MutableColumnTypes = ArgumentInfo::ListMutableColumns();
+				void** MutableColumnsData = static_cast<void**>(FMemory_Alloca(sizeof(void*) * MutableColumnTypes.Num()));
+				MutableColumns = TArrayView<void*>(MutableColumnsData, MutableColumnTypes.Num());
+			}
 
 			EFlowControl Flow = EFlowControl::Continue;
 
@@ -164,14 +198,26 @@ namespace UE::Editor::DataStorage::Queries
 			FRowHandleArrayView Rows = Response.GetBatchRows();
 			if (!Rows.IsEmpty())
 			{
-				Response.GetConstColumns(ConstColumns, ConstColumnTypes);
-				Response.GetMutableColumns(MutableColumns, MutableColumnTypes);
-				bool bSpecializationResult = Specialization(Response, ConstColumns, MutableColumns);
+				bool bSpecializationResult;
+				if constexpr (bIsConst)
+				{
+					Response.GetConstColumns(ConstColumns, ConstColumnTypes);
+					bSpecializationResult = Specialization(Response, ConstColumns);
+				}
+				else
+				{
+					Response.GetConstColumns(ConstColumns, ConstColumnTypes);
+					Response.GetMutableColumns(MutableColumns, MutableColumnTypes);
+					bSpecializationResult = Specialization(Response, ConstColumns, MutableColumns);
+				}
 
 				if (bSpecializationResult)
 				{
 					ArgumentInfo::SetConstColumns(Arguments, ConstColumns);
-					ArgumentInfo::SetMutableColumns(Arguments, MutableColumns);
+					if constexpr (!bIsConst)
+					{
+						ArgumentInfo::SetMutableColumns(Arguments, MutableColumns);
+					}
 
 					if constexpr (FunctionInfo::bIsSingleRowProcessor)
 					{
@@ -186,7 +232,7 @@ namespace UE::Editor::DataStorage::Queries
 							}
 							else
 							{
-								QueryResult.Add(Arguments.ApplyBefore(Callback));
+								QueryResult.Add(*It, Arguments.ApplyBefore(Callback));
 							}
 							ArgumentInfo::IncrementColumns(Arguments);
 						}
@@ -208,11 +254,45 @@ namespace UE::Editor::DataStorage::Queries
 	// TQueryFunctionBase
 	//
 
-	template<typename ReturnType>
+	template<typename ReturnType, bool bIsConst>
 	template<EFunctionCallConfig Config>
-	EFlowControl TQueryFunctionBase<ReturnType>::CallInternal(TResult<ReturnType>& Result, ContractType Contract, FunctionResponseType Response) const
+	EFlowControl TQueryFunctionBase<ReturnType, bIsConst>::CallInternal(TResult<ReturnType>& Result, ContractType Contract, FunctionResponseType Response) const
 	{
-		auto Specialization = [](FunctionResponseType Response, TArrayView<const void*> ConstColumns, TArrayView<void*> MutableColumns)
+		if constexpr (EnumHasAnyFlags(Config, EFunctionCallConfig::VerifyCapabilityCompatibility))
+		{
+			if (!Contract.CheckCompatibility(CapabilityMask))
+			{
+				if constexpr (EnumHasAnyFlags(Config, EFunctionCallConfig::ReportCapabilityCompatibility))
+				{
+					Contract.ReportCompatibility(CapabilityMask);
+				}
+				return EFlowControl::Break;
+			}
+		}
+
+		if constexpr (bIsConst)
+		{
+			auto Specialization = [](FunctionResponseType Response, TArrayView<const void*> ConstColumns)
+				{
+					if constexpr (EnumHasAnyFlags(Config, EFunctionCallConfig::VerifyColumns))
+					{
+						bool bMissingColumn = false;
+						for (const void* Column : ConstColumns)
+						{
+							bMissingColumn = bMissingColumn || Column == nullptr;
+						}
+						return !bMissingColumn;
+					}
+					else
+					{
+						return true;
+					}
+				};
+			return Function(Result, Contract, Response, Specialization);
+		}
+		else
+		{
+			auto Specialization = [](FunctionResponseType Response, TArrayView<const void*> ConstColumns, TArrayView<void*> MutableColumns)
 				{
 					if constexpr (EnumHasAnyFlags(Config, EFunctionCallConfig::VerifyColumns))
 					{
@@ -232,7 +312,8 @@ namespace UE::Editor::DataStorage::Queries
 						return true;
 					}
 				};
-		return Function(Result, Contract, Response, Specialization);
+			return Function(Result, Contract, Response, Specialization);
+		}
 	}
 
 	template<typename Return, FunctionType Function>
@@ -242,7 +323,7 @@ namespace UE::Editor::DataStorage::Queries
 		using ArgumentInfo = typename FunctionInfo::ArgumentInfo;
 		
 		TQueryFunction<Return> Result;
-		FunctionInfo::template StaticValidation<Return>();
+		FunctionInfo::template StaticValidation<Return, false>();
 		FunctionInfo::SetupResult(Result);
 
 		Result.Function =
@@ -252,9 +333,30 @@ namespace UE::Editor::DataStorage::Queries
 				TQueryFunction<Return>::FunctionResponseType Response,
 				TQueryFunction<Return>::FunctionSpecializationCallback Specialization)
 			{
-				return Private::CallBody<FunctionInfo, ArgumentInfo, Return>(QueryResult, Contract, Response, Specialization, LocalCallback);
+				return Private::CallBody<false, FunctionInfo, ArgumentInfo, Return>(QueryResult, Contract, Response, Specialization, LocalCallback);
 			};
 		return Result;
 	}
 
+	template<typename Return, FunctionType Function>
+	TConstQueryFunction<Return> BuildConstQueryFunction(Function&& Callback)
+	{
+		using FunctionInfo = Private::TFunctionInfo<Function>;
+		using ArgumentInfo = typename FunctionInfo::ArgumentInfo;
+
+		TConstQueryFunction<Return> Result;
+		FunctionInfo::template StaticValidation<Return, true>();
+		FunctionInfo::SetupResult(Result);
+
+		Result.Function =
+			[LocalCallback = Forward<Function>(Callback)](
+				TResult<Return>& QueryResult,
+				TConstQueryFunction<Return>::ContractType Contract,
+				TConstQueryFunction<Return>::FunctionResponseType Response,
+				TConstQueryFunction<Return>::FunctionSpecializationCallback Specialization)
+			{
+				return Private::CallBody<true, FunctionInfo, ArgumentInfo, Return>(QueryResult, Contract, Response, Specialization, LocalCallback);
+			};
+		return Result;
+	}
 } // namespace UE::Editor::DataStorage::Queries
diff --git a/Engine/Source/Runtime/TypedElementFramework/Public/Elements/Interfaces/TypedElementDataStorageInterface.h b/Engine/Source/Runtime/TypedElementFramework/Public/Elements/Interfaces/TypedElementDataStorageInterface.h
--- a/Engine/Source/Runtime/TypedElementFramework/Public/Elements/Interfaces/TypedElementDataStorageInterface.h
+++ b/Engine/Source/Runtime/TypedElementFramework/Public/Elements/Interfaces/TypedElementDataStorageInterface.h
@@ -65,7 +65,11 @@ struct FHierarchyRegistrationParams
 	const UScriptStruct* ParentChangedColumn = nullptr;
 };
 	
-class ICoreProvider : public Queries::IDynamicColumnInfo<IModularFeature>
+class ICoreProvider : public 
+	Queries::ITableInfo<
+	Queries::ITableManagement<
+	Queries::IDynamicColumnInfo<
+	IModularFeature>>>
 {
 public:
 	/**
@@ -87,39 +91,6 @@ public:
 	template<typename FactoryT>
 	FactoryT* FindFactory();
 
-	/**
-	 * @section Table management
-	 * 
-	 * @description
-	 * Tables are automatically created by taking an existing table and adding/removing columns. For
-	 * performance its however better to create a table before adding objects to the table. This
-	 * doesn't prevent those objects from having columns added/removed at a later time.
-	 * To make debugging and profiling easier it's also recommended to give tables a name.
-	 */
-
-	/** Creates a new table for with the provided columns. Optionally a name can be given which is useful for retrieval later. */
-	virtual TableHandle RegisterTable(TConstArrayView<const UScriptStruct*> ColumnList, const FName& Name) = 0;
-	template<TColumnType... Columns>
-	TableHandle RegisterTable(const FName& Name);
-	/** 
-	 * Copies the column information from the provided table and creates a new table for with the provided columns. Optionally a 
-	 * name can be given which is useful for retrieval later.
-	 */
-	virtual TableHandle RegisterTable(TableHandle SourceTable,
-		TConstArrayView<const UScriptStruct*> ColumnList, const FName& Name) = 0;
-	template<TColumnType... Columns>
-	TableHandle RegisterTable(TableHandle SourceTable, const FName& Name);
-
-	/** Returns a previously created table with the provided name or InvalidTableHandle if not found. */
-	virtual TableHandle FindTable(const FName& Name) = 0;
-
-	/** Checks if the provided table handle is pointing to a valid table. A table handle of value `InvalidTableHandle` is considered invalid. */
-	virtual bool IsTableValid(TableHandle Table) const = 0;
-	
-	/**
-	 * @section Row management
-	 */
-
 	/** 
 	 * Reserves a row to be assigned to a table at a later point. If the row is no longer needed before it's been assigned
 	 * to a table, it should still be released with RemoveRow.
@@ -146,16 +117,14 @@ public:
 	 * Adds a new row to the provided table. Callers are expected to use the callback to
 	 * initialize the row if needed.
 	 */
-	virtual RowHandle AddRow(TableHandle Table,
-		RowCreationCallbackRef OnCreated) = 0;
+	virtual RowHandle AddRow(TableHandle Table, RowCreationCallbackRef OnCreated) = 0;
 	/** Adds a new row to the provided table using a previously reserved row. */
 	virtual bool AddRow(RowHandle ReservedRow, TableHandle Table) = 0;
 	/**
 	 * Adds a new row to the provided table using a previously reserved row. Callers are expected to use the callback to
 	 * initialize the row if needed.
 	 */
-	virtual bool AddRow(RowHandle ReservedRow, TableHandle Table,
-		RowCreationCallbackRef OnCreated) = 0;
+	virtual bool AddRow(RowHandle ReservedRow, TableHandle Table, RowCreationCallbackRef OnCreated) = 0;
 
 	/**
 	 * Add multiple rows at once. For each new row the OnCreated callback is called. Callers are expected to use the callback to
@@ -167,8 +136,7 @@ public:
 	 * initialize the row if needed. This version uses a set of previously reserved rows. Any row that can't be used will be 
 	 * released.
 	 */
-	virtual bool BatchAddRow(TableHandle Table, TConstArrayView<RowHandle> ReservedHandles,
-		RowCreationCallbackRef OnCreated) = 0;
+	virtual bool BatchAddRow(TableHandle Table, TConstArrayView<RowHandle> ReservedHandles, RowCreationCallbackRef OnCreated) = 0;
 
 	/** Removes a previously reserved or added row. If the row handle is invalid or already removed, nothing happens */
 	virtual void RemoveRow(RowHandle Row) = 0;
@@ -200,18 +168,16 @@ public:
 	{
 		/** Only include rows in the result that have the required columns and pass the filter function */
 		Default = 0,
-		/** Include rows that don't have the required columns and allow them to bypass the filter */
-		IncludeRowsMissingColumns = 1 << 0,
 		/** Invert the filter to return rows that did not pass the filter function */
-		InvertFilter = 1 << 1,
+		InvertFilter = 1 << 0,
 	};
 
 	/** Filters the provided rows using the query filter and stores the results in the provided list. */
 	virtual void FilterRowsBy(
-		FRowHandleArray& Result, FRowHandleArrayView Input, EFilterOptions Options, Queries::TQueryFunction<bool>& Filter) = 0;
+		FRowHandleArray& Result, FRowHandleArrayView Input, EFilterOptions Options, Queries::TConstQueryFunction<bool>& Filter) const = 0;
 	/** Filters the provided rows using the query filter and stores the results in the provided list. */
 	template<typename FilterFunction>
-	void FilterRowsBy(FRowHandleArray& Result, FRowHandleArrayView Input, EFilterOptions Options, FilterFunction&& Filter);
+	void FilterRowsBy(FRowHandleArray& Result, FRowHandleArrayView Input, EFilterOptions Options, FilterFunction&& Filter) const;
 
 	/**
 	 * @section Column management
@@ -526,6 +492,11 @@ public:
 	 * may be called multiple times. If the query handle is invalid or has been deleted nothing happens and the callback won't be called.
 	 */
 	virtual FQueryResult RunQuery(QueryHandle Query, ERunQueryFlags Flags, const Queries::TQueryFunction<void>& Callback) = 0;
+	/**
+	 * Directly runs a query. The callback will be called for batches of matching rows. During a single call to RunQuery the callback
+	 * may be called multiple times. If the query handle is invalid or has been deleted nothing happens and the callback won't be called.
+	 */
+	virtual FQueryResult RunQuery(QueryHandle Query, ERunQueryFlags Flags, const Queries::TConstQueryFunction<void>& Callback) const = 0;
 	/**
 	 * Directly runs a query. The callback will be called for batches of matching rows. During a single call to RunQuery the callback
 	 * may be called multiple times. If the query handle is invalid or has been deleted nothing happens and the callback won't be called.
@@ -538,6 +509,18 @@ public:
 	 */
 	template<typename ResultType, FunctionType Function>
 	FQueryResult RunQuery(QueryHandle Query, ERunQueryFlags Flags, Queries::TResult<ResultType>& Result, Function&& Callback);
+	/**
+	 * Directly runs a query. The callback will be called for batches of matching rows. During a single call to RunQuery the callback
+	 * may be called multiple times. If the query handle is invalid or has been deleted nothing happens and the callback won't be called.
+	 */
+	template<FunctionType Function>
+	FQueryResult RunQuery(QueryHandle Query, ERunQueryFlags Flags, Function&& Callback) const;
+	/**
+	 * Directly runs a query. The callback will be called for batches of matching rows. During a single call to RunQuery the callback
+	 * may be called multiple times. If the query handle is invalid or has been deleted nothing happens and the callback won't be called.
+	 */
+	template<typename ResultType, FunctionType Function>
+	FQueryResult RunQuery(QueryHandle Query, ERunQueryFlags Flags, Queries::TResult<ResultType>& Result, Function&& Callback) const;
 	/**
 	 * Triggers all queries registered under the activation name to run for one update cycle. The activatable queries will be activated at
 	 * start of the cycle and disabled at the end of the cycle and act like regular queries for that cycle. This includes not running
@@ -742,23 +725,10 @@ FactoryT* ICoreProvider::FindFactory()
 	return static_cast<FactoryT*>(FindFactory(FactoryT::StaticClass()));
 }
 
-template<TColumnType... Columns>
-TableHandle ICoreProvider::RegisterTable(const FName& Name)
-{
-	return RegisterTable({ Columns::StaticStruct()... }, Name);
-}
-
-template<TColumnType... Columns>
-TableHandle ICoreProvider::RegisterTable(
-	TableHandle SourceTable, const FName& Name)
-{
-	return RegisterTable(SourceTable, { Columns::StaticStruct()... }, Name);
-}
-
 template<typename FilterFunction>
-void ICoreProvider::FilterRowsBy(FRowHandleArray& Result, FRowHandleArrayView Input, EFilterOptions Options, FilterFunction&& Filter)
+void ICoreProvider::FilterRowsBy(FRowHandleArray& Result, FRowHandleArrayView Input, EFilterOptions Options, FilterFunction&& Filter) const
 {
-	FilterRowsBy(Result, Input, Options, Queries::BuildQueryFunction<bool>(Forward<FilterFunction>(Filter)));
+	FilterRowsBy(Result, Input, Options, Queries::BuildConstQueryFunction<bool>(Forward<FilterFunction>(Filter)));
 }
 
 template<TColumnType Column>
@@ -1066,6 +1036,18 @@ FQueryResult ICoreProvider::RunQuery(QueryHandle Query, ERunQueryFlags Flags, Qu
 	return RunQuery(Query, Flags, Queries::BuildQueryFunction(Result, Forward<Function>(Callback)));
 }
 
+template<FunctionType Function>
+FQueryResult ICoreProvider::RunQuery(QueryHandle Query, ERunQueryFlags Flags, Function&& Callback) const
+{
+	return RunQuery(Query, Flags, Queries::BuildConstQueryFunction<void>(Forward<Function>(Callback)));
+}
+
+template<typename ResultType, FunctionType Function>
+FQueryResult ICoreProvider::RunQuery(QueryHandle Query, ERunQueryFlags Flags, Queries::TResult<ResultType>& Result, Function&& Callback) const
+{
+	return RunQuery(Query, Flags, Queries::BuildConstQueryFunction(Result, Forward<Function>(Callback)));
+}
+
 // Scope template implementations
 
 template<TColumnType T>
diff --git a/Engine/Source/Runtime/TypedElementFramework/Public/Elements/Interfaces/TypedElementQueryStorageInterfaces.h b/Engine/Source/Runtime/TypedElementFramework/Public/Elements/Interfaces/TypedElementQueryStorageInterfaces.h
--- a/Engine/Source/Runtime/TypedElementFramework/Public/Elements/Interfaces/TypedElementQueryStorageInterfaces.h
+++ b/Engine/Source/Runtime/TypedElementFramework/Public/Elements/Interfaces/TypedElementQueryStorageInterfaces.h
@@ -27,7 +27,10 @@ namespace UE::Editor::DataStorage
 	/**
 	 * Base interface for any contexts provided to query callbacks.
 	 */
-	struct ICommonQueryContext : public Queries::IDynamicColumnInfo<Queries::IContextCapability>
+	struct ICommonQueryContext : public 
+		Queries::ITableInfo<
+		Queries::IDynamicColumnInfo<
+		Queries::IContextCapability>>
 	{
 		virtual ~ICommonQueryContext() = default;
 
diff --git a/Engine/Source/Runtime/TypedElementFramework/Tests/Elements/Framework/TypedElementQueryContextMock.h b/Engine/Source/Runtime/TypedElementFramework/Tests/Elements/Framework/TypedElementQueryContextMock.h
--- a/Engine/Source/Runtime/TypedElementFramework/Tests/Elements/Framework/TypedElementQueryContextMock.h
+++ b/Engine/Source/Runtime/TypedElementFramework/Tests/Elements/Framework/TypedElementQueryContextMock.h
@@ -21,7 +21,6 @@ namespace UE::Editor::DataStorage::Queries
 
 #define ArgTypeName(Type, Name) Type Name
 #define ArgName(Type, Name) Name
-#define ForwardArg(Type, Name) Forward< Type >( Name )
 #define MockFuncVar Concat(MockFunction_, __LINE__)
 
 #define CapabilityStart(Capability, Flags)\
@@ -137,7 +136,6 @@ namespace UE::Editor::DataStorage::Queries
 #undef Concat
 #undef ArgTypeName
 #undef ArgName
-#undef ForwardArg
 #undef MockFunction
 #undef MockFuncVar
 	};
@@ -257,7 +255,7 @@ namespace UE::Editor::DataStorage::Queries
 	// Combine all mock implementations together into a final implementation.
 	//
 
-using QueryContextMock = TQueryContextImpl<FMockContextEnvironment
+using QueryContextMock = TQueryContextImpl<false, FMockContextEnvironment
 
 #define Function0(Capability, Return, Function)
 #define Function1(Capability, Return, Function, Arg1)
diff --git a/Engine/Source/Runtime/TypedElementFramework/Tests/Elements/Framework/TypedElementQueryFunctionsTests.cpp b/Engine/Source/Runtime/TypedElementFramework/Tests/Elements/Framework/TypedElementQueryFunctionsTests.cpp
--- a/Engine/Source/Runtime/TypedElementFramework/Tests/Elements/Framework/TypedElementQueryFunctionsTests.cpp
+++ b/Engine/Source/Runtime/TypedElementFramework/Tests/Elements/Framework/TypedElementQueryFunctionsTests.cpp
@@ -7,6 +7,7 @@
 #include "Elements/Framework/TypedElementQueryContextMock.h"
 #include "Elements/Framework/TypedElementQueryFunctions.h"
 #include "Elements/Framework/TypedElementTestColumns.h"
+#include "Math/UnrealMathUtility.h"
 #include "Tests/TestHarnessAdapter.h"
 
 namespace UE::Editor::DataStorage::Queries::Tests
@@ -15,7 +16,7 @@ namespace UE::Editor::DataStorage::Queries::Tests
 	{
 		bool Value = true;
 
-		virtual void Add(bool ResultValue) override
+		virtual void Add(RowHandle Row, bool ResultValue) override
 		{
 			Value = Value && ResultValue;
 		}
@@ -25,7 +26,7 @@ namespace UE::Editor::DataStorage::Queries::Tests
 	{
 		bool Value = false;
 
-		virtual void Add(bool ResultValue) override
+		virtual void Add(RowHandle Row, bool ResultValue) override
 		{
 			Value = Value || ResultValue;
 		}
@@ -35,7 +36,7 @@ namespace UE::Editor::DataStorage::Queries::Tests
 	{
 		int32 Value = 0;
 
-		virtual void Add(int32 ResultValue) override
+		virtual void Add(RowHandle Row, int32 ResultValue) override
 		{
 			Value += ResultValue;
 		}
@@ -133,7 +134,8 @@ namespace UE::Editor::DataStorage::Queries::Tests
 		SECTION("Empty function (void)")
 		{
 			TQueryFunction<void> Result = BuildQueryFunction<void>([](){});
-
+			
+			CHECK_MESSAGE(TEXT("Empty query still has capabilities."), Result.CapabilityMask == 0);
 			CHECK_MESSAGE(TEXT("Empty query still has const columns."), Result.ConstColumnTypes.IsEmpty());
 			CHECK_MESSAGE(TEXT("Empty query still has mutable columns."), Result.MutableColumnTypes.IsEmpty());
 		}
@@ -142,6 +144,7 @@ namespace UE::Editor::DataStorage::Queries::Tests
 		{
 			TQueryFunction<bool> Result = BuildQueryFunction<bool>([]() { return false; });
 
+			CHECK_MESSAGE(TEXT("Empty query still has capabilities."), Result.CapabilityMask == 0);
 			CHECK_MESSAGE(TEXT("Empty query still has const columns."), Result.ConstColumnTypes.IsEmpty());
 			CHECK_MESSAGE(TEXT("Empty query still has mutable columns."), Result.MutableColumnTypes.IsEmpty());
 		}
@@ -151,6 +154,7 @@ namespace UE::Editor::DataStorage::Queries::Tests
 			bool bResult = false;
 			TQueryFunction<bool> Result = BuildQueryFunction<bool>([&bResult]() { return bResult; });
 
+			CHECK_MESSAGE(TEXT("Empty query still has capabilities."), Result.CapabilityMask == 0);
 			CHECK_MESSAGE(TEXT("Empty query still has const columns."), Result.ConstColumnTypes.IsEmpty());
 			CHECK_MESSAGE(TEXT("Empty query still has mutable columns."), Result.MutableColumnTypes.IsEmpty());
 		}
@@ -159,6 +163,7 @@ namespace UE::Editor::DataStorage::Queries::Tests
 		{
 			TQueryFunction<void> Result = BuildQueryFunction<void>([](TQueryContext<SingleRowInfo> Context) {});
 
+			CHECK_MESSAGE(TEXT("Expected exactly one capability."), FMath::CountBits64(Result.CapabilityMask) == 1);
 			CHECK_MESSAGE(TEXT("Query with only context shouldn't have const columns."), Result.ConstColumnTypes.IsEmpty());
 			CHECK_MESSAGE(TEXT("Query with only context shouldn't have mutable columns."), Result.MutableColumnTypes.IsEmpty());
 		}
@@ -170,6 +175,7 @@ namespace UE::Editor::DataStorage::Queries::Tests
 				{
 				});
 
+			CHECK_MESSAGE(TEXT("Query with only columns shouldn't have capabilities."), Result.CapabilityMask == 0);
 			CHECK_MESSAGE(TEXT("Expected exactly one const column."), Result.ConstColumnTypes.Num() == 1);
 			CHECK_MESSAGE(TEXT("Expected exactly one mutable column."), Result.MutableColumnTypes.Num() == 1);
 		}
@@ -181,6 +187,7 @@ namespace UE::Editor::DataStorage::Queries::Tests
 				{
 				});
 
+			CHECK_MESSAGE(TEXT("Expected exactly one capability."), FMath::CountBits64(Result.CapabilityMask) == 1);
 			CHECK_MESSAGE(TEXT("Expected exactly one const column."), Result.ConstColumnTypes.Num() == 1);
 			CHECK_MESSAGE(TEXT("Expected exactly one mutable column."), Result.MutableColumnTypes.Num() == 1);
 
@@ -195,6 +202,7 @@ namespace UE::Editor::DataStorage::Queries::Tests
 				{
 				});
 
+			CHECK_MESSAGE(TEXT("Expected exactly one capability."), FMath::CountBits64(Result.CapabilityMask) == 1);
 			CHECK_MESSAGE(TEXT("Expected exactly one const column."), Result.ConstColumnTypes.Num() == 1);
 			CHECK_MESSAGE(TEXT("Expected exactly one mutable column."), Result.MutableColumnTypes.Num() == 1);
 
@@ -209,6 +217,7 @@ namespace UE::Editor::DataStorage::Queries::Tests
 				{
 				});
 
+			CHECK_MESSAGE(TEXT("Expected exactly one capability."), FMath::CountBits64(Result.CapabilityMask) == 1);
 			CHECK_MESSAGE(TEXT("Expected exactly one const column."), Result.ConstColumnTypes.Num() == 1);
 			CHECK_MESSAGE(TEXT("Expected exactly one mutable column."), Result.MutableColumnTypes.Num() == 1);
 
@@ -216,6 +225,26 @@ namespace UE::Editor::DataStorage::Queries::Tests
 			CHECK_MESSAGE(TEXT("Incorrect mutable column set."), Result.MutableColumnTypes[0] == FTestColumnInt::StaticStruct());
 		}
 
+		SECTION("Check context compatibility without arguments")
+		{
+			TQueryFunction<void> Function = BuildQueryFunction<void>([](){});
+
+			FMockContextEnvironment Environment;
+			QueryContextMock Context(Environment);
+			bool bResult = Context.CheckCompatibility(Function.CapabilityMask);
+			CHECK_MESSAGE(TEXT("Function without context incorrectly found to not be matching."), bResult);
+		}
+
+		SECTION("Check context compatibility with context")
+		{
+			TQueryFunction<void> Function = BuildQueryFunction<void>([](TQueryContext<SingleRowInfo> Context) {});
+
+			FMockContextEnvironment Environment;
+			QueryContextMock Context(Environment);
+			bool bResult = Context.CheckCompatibility(Function.CapabilityMask);
+			CHECK_MESSAGE(TEXT("Function with context incorrectly found to not be matching."), bResult);
+		}
+
 		SECTION("Call without arguments")
 		{
 			TQueryFunction<bool> Function = BuildQueryFunction<bool>(
@@ -283,7 +312,7 @@ namespace UE::Editor::DataStorage::Queries::Tests
 			TQueryFunction<void> Function = BuildQueryFunction<void>(
 				[](TQueryContext<RowBatchInfo> Context, TBatch<FTestColumnString> ColumnsA, TConstBatch<FTestColumnInt> ColumnsB)
 				{
-					Context.ForEachRow([](FTestColumnString& ColumnA, const FTestColumnInt& ColumnB)
+					Context.ForEachRow([](RowHandle Row, FTestColumnString& ColumnA, const FTestColumnInt& ColumnB)
 						{
 							ColumnA.TestString = FString::Printf(TEXT("Test: %i"), ColumnB.TestInt);
 						}, ColumnsA, ColumnsB);
@@ -379,7 +408,7 @@ namespace UE::Editor::DataStorage::Queries::Tests
 			CHECK_MESSAGE(TEXT("Not all Rows were added."), RowCheck.Num() == CheckRange);
 			for (int32 FakeRow = 0; FakeRow < CheckRange; ++FakeRow)
 			{
-				CHECK_MESSAGE(TEXT("One of the rows wasn't called or called with the wrong row handle."), RowCheck[FakeRow] == FakeRow);
+				CHECK_MESSAGE(TEXT("One of the rows wasn't called or calle with the wrong row handle."), RowCheck[FakeRow] == FakeRow);
 			}
 		}
 
@@ -406,7 +435,7 @@ namespace UE::Editor::DataStorage::Queries::Tests
 			TQueryFunction<bool> Function = BuildQueryFunction<bool>(
 				[](TQueryContext<SingleRowInfo> Context, TResult<bool>& Result)
 				{ 
-					Result.Add(true);
+					Result.Add(Context.GetCurrentRow(), true);
 				});
 
 			FMockContextEnvironment Environment;
@@ -429,9 +458,9 @@ namespace UE::Editor::DataStorage::Queries::Tests
 			TQueryFunction<int32> Function = BuildQueryFunction<int32>(
 				[](TQueryContext<RowBatchInfo> Context, TResult<int32>& Result, TConstBatch<FTestColumnInt> Columns)
 				{
-					Context.ForEachRow([&Result](const FTestColumnInt& Column)
+					Context.ForEachRow([&Result](RowHandle Row, const FTestColumnInt& Column)
 						{
-							Result.Add(Column.TestInt);
+							Result.Add(Row, Column.TestInt);
 						}, Columns);
 				});
 
@@ -491,6 +520,49 @@ namespace UE::Editor::DataStorage::Queries::Tests
 			CHECK_MESSAGE(TEXT("Iteration didn't stop on the first row."), LastValue == 1);
 		}
 
+		SECTION("Forward context")
+		{
+			TQueryFunction<void> Function = BuildQueryFunction<void>(
+				[](TQueryContext<SingleRowInfo, CurrentTableInfo> Context, const FTestColumnInt& Column)
+				{
+					Context.GetCurrentRow();
+					Context.GetCurrentTable();
+				});
+
+			FMockContextEnvironment Environment;
+			QueryContextMock Context(Environment);
+
+			bool bCurrentRowCalled = false;
+			bool bCurrentTableCalled = false;
+			Environment.Assign_GetCurrentRow_Const(
+				[&bCurrentRowCalled]()
+				{
+					bCurrentRowCalled = true;
+					return 1;
+				});
+			Environment.Assign_GetCurrentTable_Const(
+				[&bCurrentTableCalled]()
+				{
+					bCurrentTableCalled = true;
+					return 1;
+				});
+
+
+			FTestColumnInt Columns[] = { FTestColumnInt{.TestInt = 1 } };
+
+			FTestQueryFunctionResponse Response;
+			Response.SetRowCount(1);
+			Response.SetConstColumns(Columns);
+
+			// This is an order that would typically happen inside a query callback that calls a utility function.
+			TQueryContext<SingleRowInfo, CurrentTableInfo> ContextRequest(Context);
+			TForwardingQueryContext<SingleRowInfo> ForwardedContext(ContextRequest);
+
+			Function.Call(ForwardedContext, Response);
+
+			CHECK_MESSAGE(TEXT("GetCurrentRow wasn't called."), bCurrentRowCalled);
+			CHECK_MESSAGE(TEXT("GetCurrentTable wasn't called."), bCurrentTableCalled);
+		}
 	}
 } // namespace UE::Editor::DataStorage::Queries::Tests
 
