[TEDS] Renamed a few files to represent the new names in the files.

#Synced-CL 52511355

diff --git a/Engine/Plugins/Experimental/EditorDataStorage/Source/TedsCore/Private/Processors/QueryCapabilityImplementationForwarder.h b/Engine/Plugins/Experimental/EditorDataStorage/Source/TedsCore/Private/Processors/QueryCapabilityImplementationForwarder.h
new file mode 100644
--- /dev/null
+++ b/Engine/Plugins/Experimental/EditorDataStorage/Source/TedsCore/Private/Processors/QueryCapabilityImplementationForwarder.h
@@ -0,0 +1,118 @@
+// Copyright Epic Games, Inc. All Rights Reserved.
+
+#pragma once
+
+#include "DataStorage/CommonTypes.h"
+#include "DataStorage/Queries/Types.h"
+#include "Elements/Framework/TypedElementQueryCapabilities.h"
+#include "Elements/Framework/TypedElementQueryContract.h"
+#include "Elements/Framework/TypedElementQueryFunctionArgumentTypes.h"
+
+namespace UE::Editor::DataStorage::Queries::Private
+{
+	// Because the original context has a different inheritance hierarchy than the new query context implementations can't
+	// be directly used in the various IQueryContext implementation. This wrapper can be used to bridge the two until
+	// the new query context can fully replace the original query context.
+
+	template<typename Environment, typename Interface, typename... Implementations> struct TImplementationForwarder
+	{
+		static_assert(sizeof...(Implementations) < 0,
+			"Specialization for query context forwards missing. The most likely reason is the addition of a new query "
+			"context capability that wasn't specialized.");
+	};
+
+	template<typename Environment, typename Interface>
+	struct TImplementationForwarder<Environment, Interface> : Interface
+	{
+		TImplementationForwarder() = default;
+
+		template<typename... TArgs>
+		explicit TImplementationForwarder(TArgs&&... Args)
+			: ContextEnvironment(Forward<TArgs>(Args)...)
+		{
+		}
+
+		Environment ContextEnvironment;
+	};
+
+	template<typename Implementation, typename Capability>
+	concept ImplementsCapability = std::is_base_of_v<ImplementsContextCapability<Capability>, Implementation>;
+
+#define ArgTypeName(Type, Name) Type Name
+#define ArgName(Type, Name) Name
+#define ForwardArg(Type, Name) Forward< Type >( Name )
+	
+#define CallFunction(Capability, ReturnValue, FunctionName, ...) \
+	if constexpr (ImplementsCapability<Implementation, Capability>) \
+	{ \
+		static_assert( requires() \
+			{ \
+				{ Implementation:: FunctionName (this->ContextEnvironment __VA_OPT__(,) __VA_ARGS__ ) } -> UE::CSameAs< ReturnValue >; \
+			}, "Function '" #FunctionName "' in capability '" #Capability "' is missing from the forwarder." ); \
+		return Implementation:: FunctionName ( this->ContextEnvironment __VA_OPT__(,) __VA_ARGS__ ); \
+	} \
+	else \
+	{ \
+		static_assert(std::is_same_v<void, Capability>, \
+			"Function '" #FunctionName "' in capability '" #Capability "' is missing from the implementation."); \
+		return UE::MakeDefault< ReturnValue >(); \
+	}
+
+#define CapabilityStart(Capability, Flags) \
+	template<typename Environment, typename Interface, typename Implementation, typename... Implementations> \
+		requires ImplementsCapability<Implementation, Queries:: Capability > \
+	struct TImplementationForwarder< Environment, Interface, Implementation, Implementations...> : \
+		Implementation, \
+		TImplementationForwarder<Environment, Interface, Implementations...> \
+	{ \
+		using Super = TImplementationForwarder<Environment, Interface, Implementations...>; \
+		TImplementationForwarder() = default; \
+		template<typename... TArgs> \
+		explicit TImplementationForwarder(TArgs&&... Args) \
+			: Super(Forward<TArgs>(Args)...) \
+		{}
+
+#define Function0(Capability, Return, Function) virtual Return Function() \
+	{ CallFunction(Capability, Return, Function ) }
+#define Function1(Capability, Return, Function, Arg1) virtual Return Function(ArgTypeName Arg1 ) override \
+	{ CallFunction(Capability, Return, Function, ForwardArg Arg1 ) }
+#define Function2(Capability, Return, Function, Arg1, Arg2) virtual Return Function(ArgTypeName Arg1 , ArgTypeName Arg2 ) override \
+	{ CallFunction(Capability, Return, Function, ForwardArg Arg1, ForwardArg Arg2 ) }
+#define Function3(Capability, Return, Function, Arg1, Arg2, Arg3) virtual Return Function(ArgTypeName Arg1 , ArgTypeName Arg2 , ArgTypeName Arg3 ) override \
+	{ CallFunction(Capability, Return, Function, ForwardArg Arg1, ForwardArg Arg2, ForwardArg Arg3 ) }
+#define Function4(Capability, Return, Function, Arg1, Arg2, Arg3, Arg4) virtual Return Function(ArgTypeName Arg1 , ArgTypeName Arg2 , ArgTypeName Arg3 , ArgTypeName Arg4 ) override \
+	{ CallFunction(Capability, Return, Function, ForwardArg Arg1, ForwardArg Arg2, ForwardArg Arg3, ForwardArg Arg4 ) }
+
+#define ConstFunction0(Capability, Return, Function) virtual Return Function() const \
+	{ CallFunction(Capability, Return, Function ) }
+#define ConstFunction1(Capability, Return, Function, Arg1) virtual Return Function(ArgTypeName Arg1 ) const override \
+	{ CallFunction(Capability, Return, Function, ForwardArg Arg1 ) }
+#define ConstFunction2(Capability, Return, Function, Arg1, Arg2) virtual Return Function(ArgTypeName Arg1 , ArgTypeName Arg2 ) const override \
+	{ CallFunction(Capability, Return, Function, ForwardArg Arg1, ForwardArg Arg2 ) }
+#define ConstFunction3(Capability, Return, Function, Arg1, Arg2, Arg3) virtual Return Function(ArgTypeName Arg1 , ArgTypeName Arg2 , ArgTypeName Arg3 ) const override \
+	{ CallFunction(Capability, Return, Function, ForwardArg Arg1, ForwardArg Arg2, ForwardArg Arg3 ) }
+#define ConstFunction4(Capability, Return, Function, Arg1, Arg2, Arg3, Arg4) virtual Return Function(ArgTypeName Arg1 , ArgTypeName Arg2 , ArgTypeName Arg3 , ArgTypeName Arg4 ) const override \
+	{ CallFunction(Capability, Return, Function, ForwardArg Arg1, ForwardArg Arg2, ForwardArg Arg3, ForwardArg Arg4 ) }
+
+#define CapabilityEnd(Capability) \
+	}; \
+	
+#include "Elements/Framework/TypedElementQueryCapabilities.inl"
+
+#undef ArgTypeName
+#undef ArgName
+#undef ForwardArg
+#undef CallFunction
+#undef Function0
+#undef Function1
+#undef Function2
+#undef Function3
+#undef Function4
+#undef ConstFunction0
+#undef ConstFunction1
+#undef ConstFunction2
+#undef ConstFunction3
+#undef ConstFunction4
+#undef CapabilityStart
+#undef CapabilityEnd
+} // namespace UE::Editor::DataStorage::Queries::Private
\ No newline at end of file
diff --git a/Engine/Plugins/Experimental/EditorDataStorage/Source/TedsCore/Private/Processors/TypedElementProcessorAdaptors.cpp b/Engine/Plugins/Experimental/EditorDataStorage/Source/TedsCore/Private/Processors/TypedElementProcessorAdaptors.cpp
--- a/Engine/Plugins/Experimental/EditorDataStorage/Source/TedsCore/Private/Processors/TypedElementProcessorAdaptors.cpp
+++ b/Engine/Plugins/Experimental/EditorDataStorage/Source/TedsCore/Private/Processors/TypedElementProcessorAdaptors.cpp
@@ -8,8 +8,8 @@
 #include "MassExecutionContext.h"
 #include "Misc/EnumClassFlags.h"
 #include "Queries/TypedElementExtendedQueryStore.h"
-#include "Queries/TypedElementQueryContext_MassForwarder.h"
-#include "Processors/TypedElementQueryCapabilityImplemenationForwarder.h"
+#include "Queries/ProcessorContextEnvironment.h"
+#include "Processors/QueryCapabilityImplementationForwarder.h"
 #include "StructUtils/StructTypeBitSet.h"
 #include "TypedElementDatabase.h"
 #include "TypedElementDatabaseEnvironment.h"
diff --git a/Engine/Plugins/Experimental/EditorDataStorage/Source/TedsCore/Private/Processors/TypedElementQueryCapabilityImplemenationForwarder.h b/Engine/Plugins/Experimental/EditorDataStorage/Source/TedsCore/Private/Processors/TypedElementQueryCapabilityImplemenationForwarder.h
deleted file mode 100644
diff --git a/Engine/Plugins/Experimental/EditorDataStorage/Source/TedsCore/Private/Queries/DirectContextEnvironment.cpp b/Engine/Plugins/Experimental/EditorDataStorage/Source/TedsCore/Private/Queries/DirectContextEnvironment.cpp
new file mode 100644
--- /dev/null
+++ b/Engine/Plugins/Experimental/EditorDataStorage/Source/TedsCore/Private/Queries/DirectContextEnvironment.cpp
@@ -0,0 +1,142 @@
+// Copyright Epic Games, Inc. All Rights Reserved.
+
+#include "Queries/DirectContextEnvironment.h"
+#include "TypedElementDatabase.h"
+
+namespace UE::Editor::DataStorage::Queries
+{
+	//
+	// FDirectContextEnvironment
+	//
+
+	FDirectContextEnvironment::FDirectContextEnvironment(FEnvironment& Environment, UEditorDataStorage& DirectContext)
+		: Environment(&Environment)
+		, DirectContext(&DirectContext)
+	{
+	}
+
+	void FDirectContextEnvironment::SetBatch(FRowHandleArrayView InRows)
+	{
+		Rows = InRows;
+	}
+
+	FEnvironment* FDirectContextEnvironment::GetEnvironment()
+	{
+		return Environment;
+	}
+
+	const FEnvironment* FDirectContextEnvironment::GetEnvironment() const
+	{
+		return Environment;
+	}
+
+	FRowHandleArrayView FDirectContextEnvironment::GetBatchRows() const
+	{
+		return Rows;
+	}
+
+	void FDirectContextEnvironment::GetConstColumns(
+		TArrayView<const void*> ColumnsData, TConstArrayView<const UScriptStruct*> ColumnTypes) const
+	{
+		if (!Rows.IsEmpty())
+		{
+			RowHandle Root = *Rows.GetData();
+			const void** ColumnsDataIt = ColumnsData.GetData();
+			for (const UScriptStruct* ColumnType : ColumnTypes)
+			{
+				*ColumnsDataIt++ = DirectContext->GetColumnData(Root, ColumnType);
+			}
+		}
+		else
+		{
+			FMemory::Memzero(ColumnsData.GetData(), ColumnsData.NumBytes());
+		}
+	}
+
+	void FDirectContextEnvironment::GetMutableColumns(
+		TArrayView<void*> ColumnsData, TConstArrayView<const UScriptStruct*> ColumnTypes)
+	{
+		if (!Rows.IsEmpty())
+		{
+			// At the time this is called there are guaranteed to be rows assigned.
+			RowHandle Root = *Rows.GetData();
+			void** ColumnsDataIt = ColumnsData.GetData();
+			for (const UScriptStruct* ColumnType : ColumnTypes)
+			{
+				*ColumnsDataIt++ = DirectContext->GetColumnData(Root, ColumnType);
+			}
+		}
+		else
+		{
+			FMemory::Memzero(ColumnsData.GetData(), ColumnsData.NumBytes());
+		}
+	}
+
+	void FDirectContextEnvironment::SetCurrentRow(RowHandle Row)
+	{
+		CurrentRow = Row;
+	}
+
+
+
+	//
+	// FConstDirectContextEnvironment
+	//
+
+	FConstDirectContextEnvironment::FConstDirectContextEnvironment(const FEnvironment& Environment, const UEditorDataStorage& DirectContext)
+		: Environment(&Environment)
+		, DirectContext(&DirectContext)
+	{
+	}
+
+	void FConstDirectContextEnvironment::SetBatch(FRowHandleArrayView InRows)
+	{
+		Rows = InRows;
+	}
+
+	FEnvironment* FConstDirectContextEnvironment::GetEnvironment()
+	{
+		checkf(false, TEXT("Calling a mutable version of the environment is not supported by a const environment."));
+		return nullptr;
+	}
+
+	const FEnvironment* FConstDirectContextEnvironment::GetEnvironment() const
+	{
+		return Environment;
+	}
+
+	FRowHandleArrayView FConstDirectContextEnvironment::GetBatchRows() const
+	{
+		return Rows;
+	}
+
+	void FConstDirectContextEnvironment::GetConstColumns(
+		TArrayView<const void*> ColumnsData, TConstArrayView<const UScriptStruct*> ColumnTypes) const
+	{
+		if (!Rows.IsEmpty())
+		{
+			RowHandle Root = *Rows.GetData();
+			const void** ColumnsDataIt = ColumnsData.GetData();
+			for (const UScriptStruct* ColumnType : ColumnTypes)
+			{
+				*ColumnsDataIt++ = DirectContext->GetColumnData(Root, ColumnType);
+			}
+		}
+		else
+		{
+			FMemory::Memzero(ColumnsData.GetData(), ColumnsData.NumBytes());
+		}
+	}
+
+	void FConstDirectContextEnvironment::GetMutableColumns(
+		TArrayView<void*> ColumnsData, TConstArrayView<const UScriptStruct*> ColumnTypes)
+	{
+		checkf(false, TEXT("Retrieving mutable columns is not supported from a const environment."));
+		FMemory::Memzero(ColumnsData.GetData(), ColumnsData.NumBytes());
+	}
+
+	void FConstDirectContextEnvironment::SetCurrentRow(RowHandle Row)
+	{
+		CurrentRow = Row;
+	}
+} // namespace UE::Editor::DataStorage::Queries
\ No newline at end of file
diff --git a/Engine/Plugins/Experimental/EditorDataStorage/Source/TedsCore/Private/Queries/DirectContextEnvironment.h b/Engine/Plugins/Experimental/EditorDataStorage/Source/TedsCore/Private/Queries/DirectContextEnvironment.h
new file mode 100644
--- /dev/null
+++ b/Engine/Plugins/Experimental/EditorDataStorage/Source/TedsCore/Private/Queries/DirectContextEnvironment.h
@@ -0,0 +1,75 @@
+// Copyright Epic Games, Inc. All Rights Reserved.
+
+#pragma once
+
+#include "Elements/Framework/TypedElementQueryContextImplementation.h"
+#include "Elements/Framework/TypedElementQueryFunctions.h"
+#include "Elements/Framework/TypedElementRowHandleArray.h"
+#include "TypedElementDatabaseEnvironment.h"
+#include "Queries/QueryContextCapabilities/ContextEnvironment.h"
+#include "Queries/QueryContextCapabilities/EditorDataStorageDynamicColumnInfo.h"
+#include "Queries/QueryContextCapabilities/TableCapabilities.h"
+#include "Queries/QueryContextCapabilities/RowBatch.h"
+#include "Queries/QueryContextCapabilities/SingleRow.h"
+
+class UEditorDataStorage;
+
+namespace UE::Editor::DataStorage::Queries
+{
+	struct FDirectContextEnvironment final : IContextEnvironment
+	{
+		FDirectContextEnvironment() = default;
+		FDirectContextEnvironment(FEnvironment& Environment, UEditorDataStorage& DirectContext);
+
+		void SetBatch(FRowHandleArrayView InRows);
+
+		virtual FEnvironment* GetEnvironment() override;
+		virtual const FEnvironment* GetEnvironment() const override;
+
+		virtual FRowHandleArrayView GetBatchRows() const override;
+		virtual void GetConstColumns(TArrayView<const void*> ColumnsData, TConstArrayView<const UScriptStruct*> ColumnTypes) const override;
+		virtual void GetMutableColumns(TArrayView<void*> ColumnsData, TConstArrayView<const UScriptStruct*> ColumnTypes) override;
+		virtual void SetCurrentRow(RowHandle Row) override;
+
+		FRowHandleArrayView Rows;
+		FEnvironment* Environment = nullptr;
+		UEditorDataStorage* DirectContext = nullptr;
+		RowHandle CurrentRow = InvalidRowHandle;
+	};
+
+	struct FConstDirectContextEnvironment final : IContextEnvironment
+	{
+		/** The provided rows need to be contiguous in the same table. */
+		FConstDirectContextEnvironment() = default;
+		FConstDirectContextEnvironment(const FEnvironment& Environment, const UEditorDataStorage& DirectContext);
+
+		void SetBatch(FRowHandleArrayView InRows);
+
+		virtual FEnvironment* GetEnvironment() override;
+		virtual const FEnvironment* GetEnvironment() const override;
+
+		virtual FRowHandleArrayView GetBatchRows() const override;
+		virtual void GetConstColumns(TArrayView<const void*> ColumnsData, TConstArrayView<const UScriptStruct*> ColumnTypes) const override;
+		virtual void GetMutableColumns(TArrayView<void*> ColumnsData, TConstArrayView<const UScriptStruct*> ColumnTypes) override;
+		virtual void SetCurrentRow(RowHandle Row) override;
+
+		FRowHandleArrayView Rows;
+		const FEnvironment* Environment = nullptr;
+		const UEditorDataStorage* DirectContext = nullptr;
+		RowHandle CurrentRow = InvalidRowHandle;
+	};
+
+	using QueryContext_DirectApi = TQueryContextImpl<false, FDirectContextEnvironment,
+		FCurrentTableInfoContextCapability,
+		FTableInfoContextCapability,
+		FSingleRowInfoContextCapability,
+		FRowBatchInfoContextCapability,
+		FDynamicColumnInfoContextCapability>;
+
+	using QueryContext_ConstDirectApi = TQueryContextImpl<true, FConstDirectContextEnvironment,
+		FCurrentTableInfoContextCapability,
+		FTableInfoContextCapability,
+		FSingleRowInfoContextCapability,
+		FRowBatchInfoContextCapability,
+		FDynamicColumnInfoContextCapability>;
+} // namespace UE::Editor::DataStorage::Queries
diff --git a/Engine/Plugins/Experimental/EditorDataStorage/Source/TedsCore/Private/Queries/ProcessorContextEnvironment.cpp b/Engine/Plugins/Experimental/EditorDataStorage/Source/TedsCore/Private/Queries/ProcessorContextEnvironment.cpp
new file mode 100644
--- /dev/null
+++ b/Engine/Plugins/Experimental/EditorDataStorage/Source/TedsCore/Private/Queries/ProcessorContextEnvironment.cpp
@@ -0,0 +1,125 @@
+// Copyright Epic Games, Inc. All Rights Reserved.
+
+#include "Queries/ProcessorContextEnvironment.h"
+#include "MassExecutionContext.h"
+#include "TypedElementUtils.h"
+
+namespace UE::Editor::DataStorage::Queries
+{
+	//
+	// FProcessorContextEnvironment
+	//
+
+	FProcessorContextEnvironment::FProcessorContextEnvironment(FEnvironment& Environment, FMassExecutionContext& MassContext)
+		: Environment(Environment)
+		, MassContext(MassContext)
+	{
+	}
+
+	FEnvironment* FProcessorContextEnvironment::GetEnvironment()
+	{
+		return &Environment;
+	}
+
+	const FEnvironment* FProcessorContextEnvironment::GetEnvironment() const
+	{
+		return &Environment;
+	}
+
+	FRowHandleArrayView FProcessorContextEnvironment::GetBatchRows() const
+	{
+		TConstArrayView<RowHandle> Rows = MassEntitiesToRowsConversion(MassContext.GetEntities());
+		return FRowHandleArrayView(Rows, FRowHandleArrayView::EFlags::IsUnique);
+	}
+
+	void FProcessorContextEnvironment::GetConstColumns(
+		TArrayView<const void*> ColumnsData, TConstArrayView<const UScriptStruct*> ColumnTypes) const
+	{
+		checkfSlow(ColumnsData.Num() >= ColumnTypes.Num(), TEXT("The list size of requested const column types (%i) doesn't match the "
+			"available space for the results (%i)"), ColumnTypes.Num(), ColumnsData.Num());
+
+		const void** ColumnsDataIt = ColumnsData.GetData();
+		for (const UScriptStruct* ColumnType : ColumnTypes)
+		{
+			checkfSlow(UE::Mass::IsA<FMassFragment>(ColumnType),
+				TEXT("Unable to retrieve const column data for a column type without data: %s"), *ColumnType->GetFullName());
+			*ColumnsDataIt++ = MassContext.GetFragmentView(ColumnType).GetData();
+		}
+	}
+
+	void FProcessorContextEnvironment::GetMutableColumns(
+		TArrayView<void*> ColumnsData, TConstArrayView<const UScriptStruct*> ColumnTypes)
+	{
+		checkfSlow(ColumnsData.Num() >= ColumnTypes.Num(), TEXT("The list size of requested mutable column types (%i) doesn't match the "
+			"available space for the results (%i)"), ColumnTypes.Num(), ColumnsData.Num());
+
+		void** ColumnsDataIt = ColumnsData.GetData();
+		for (const UScriptStruct* ColumnType : ColumnTypes)
+		{
+			checkfSlow(UE::Mass::IsA<FMassFragment>(ColumnType),
+				TEXT("Unable to retrieve mutable column data for a column type without data: %s"), *ColumnType->GetFullName());
+			*ColumnsDataIt++ = MassContext.GetMutableFragmentView(ColumnType).GetData();
+		}
+	}
+
+	void FProcessorContextEnvironment::SetCurrentRow(RowHandle Row)
+	{
+		CurrentRow = Row;
+	}
+
+
+	//
+	// FConstProcessorContextEnvironment
+	//
+
+	FConstProcessorContextEnvironment::FConstProcessorContextEnvironment(
+		const FEnvironment& Environment, const FMassExecutionContext& MassContext)
+		: Environment(Environment)
+		, MassContext(MassContext)
+	{
+	}
+
+	FEnvironment* FConstProcessorContextEnvironment::GetEnvironment()
+	{
+		checkf(false, TEXT("Calling a mutable version of the environment is not supported by a const environment."));
+		return nullptr;
+	}
+
+	const FEnvironment* FConstProcessorContextEnvironment::GetEnvironment() const
+	{
+		return &Environment;
+	}
+
+	FRowHandleArrayView FConstProcessorContextEnvironment::GetBatchRows() const
+	{
+		TConstArrayView<RowHandle> Rows = MassEntitiesToRowsConversion(MassContext.GetEntities());
+		return FRowHandleArrayView(Rows, FRowHandleArrayView::EFlags::IsUnique);
+	}
+
+	void FConstProcessorContextEnvironment::GetConstColumns(
+		TArrayView<const void*> ColumnsData, TConstArrayView<const UScriptStruct*> ColumnTypes) const
+	{
+		checkfSlow(ColumnsData.Num() >= ColumnTypes.Num(), TEXT("The list size of requested const column types (%i) doesn't match the "
+			"available space for the results (%i)"), ColumnTypes.Num(), ColumnsData.Num());
+
+		const void** ColumnsDataIt = ColumnsData.GetData();
+		for (const UScriptStruct* ColumnType : ColumnTypes)
+		{
+			checkfSlow(UE::Mass::IsA<FMassFragment>(ColumnType),
+				TEXT("Unable to retrieve const column data for a column type without data: %s"), *ColumnType->GetFullName());
+			*ColumnsDataIt++ = MassContext.GetFragmentView(ColumnType).GetData();
+		}
+	}
+
+	void FConstProcessorContextEnvironment::GetMutableColumns(
+		TArrayView<void*> ColumnsData, TConstArrayView<const UScriptStruct*> ColumnTypes)
+	{
+		checkf(false, TEXT("Retrieving mutable columns is not supported from a const environment."));
+		FMemory::Memzero(ColumnsData.GetData(), ColumnsData.NumBytes());
+	}
+
+	void FConstProcessorContextEnvironment::SetCurrentRow(RowHandle Row)
+	{
+		CurrentRow = Row;
+	}
+} // namespace UE::Editor::DataStorage::Queries
\ No newline at end of file
diff --git a/Engine/Plugins/Experimental/EditorDataStorage/Source/TedsCore/Private/Queries/ProcessorContextEnvironment.h b/Engine/Plugins/Experimental/EditorDataStorage/Source/TedsCore/Private/Queries/ProcessorContextEnvironment.h
new file mode 100644
--- /dev/null
+++ b/Engine/Plugins/Experimental/EditorDataStorage/Source/TedsCore/Private/Queries/ProcessorContextEnvironment.h
@@ -0,0 +1,64 @@
+// Copyright Epic Games, Inc. All Rights Reserved.
+
+#pragma once
+
+#include "Elements/Framework/TypedElementQueryContextImplementation.h"
+#include "Elements/Framework/TypedElementQueryFunctions.h"
+#include "TypedElementDatabaseEnvironment.h"
+#include "Queries/QueryContextCapabilities/EditorDataStorageDynamicColumnInfo.h"
+#include "Queries/QueryContextCapabilities/TableCapabilities.h"
+#include "Queries/QueryContextCapabilities/RowBatch.h"
+#include "Queries/QueryContextCapabilities/SingleRow.h"
+
+struct FMassExecutionContext;
+
+namespace UE::Editor::DataStorage::Queries
+{
+	struct FProcessorContextEnvironment final : IContextEnvironment
+	{
+		FProcessorContextEnvironment(FEnvironment& Environment, FMassExecutionContext& MassContext);
+
+		virtual FEnvironment* GetEnvironment() override;
+		virtual const FEnvironment* GetEnvironment() const override;
+		
+		virtual FRowHandleArrayView GetBatchRows() const override;
+		virtual void GetConstColumns(TArrayView<const void*> ColumnsData, TConstArrayView<const UScriptStruct*> ColumnTypes) const override;
+		virtual void GetMutableColumns(TArrayView<void*> ColumnsData, TConstArrayView<const UScriptStruct*> ColumnTypes) override;
+		virtual void SetCurrentRow(RowHandle Row) override;
+		
+		FEnvironment& Environment;
+		FMassExecutionContext& MassContext;
+		RowHandle CurrentRow = InvalidRowHandle;
+	};
+
+	struct FConstProcessorContextEnvironment final : IContextEnvironment
+	{
+		FConstProcessorContextEnvironment(const FEnvironment& Environment, const FMassExecutionContext& MassContext);
+
+		virtual FEnvironment* GetEnvironment() override;
+		virtual const FEnvironment* GetEnvironment() const override;
+
+		virtual FRowHandleArrayView GetBatchRows() const override;
+		virtual void GetConstColumns(TArrayView<const void*> ColumnsData, TConstArrayView<const UScriptStruct*> ColumnTypes) const override;
+		virtual void GetMutableColumns(TArrayView<void*> ColumnsData, TConstArrayView<const UScriptStruct*> ColumnTypes) override;
+		virtual void SetCurrentRow(RowHandle Row) override;
+
+		const FEnvironment& Environment;
+		const FMassExecutionContext& MassContext;
+		RowHandle CurrentRow = InvalidRowHandle;
+	};
+
+	using QueryContext_ProcessorApi = TQueryContextImpl<false, FProcessorContextEnvironment,
+		FCurrentTableInfoContextCapability,
+		FTableInfoContextCapability,
+		FSingleRowInfoContextCapability,
+		FRowBatchInfoContextCapability,
+		FDynamicColumnInfoContextCapability>;
+	
+	using QueryContext_ConstProcessorApi = TQueryContextImpl<true, FConstProcessorContextEnvironment,
+		FCurrentTableInfoContextCapability,
+		FTableInfoContextCapability,
+		FSingleRowInfoContextCapability,
+		FRowBatchInfoContextCapability,
+		FDynamicColumnInfoContextCapability>;
+} // namespace UE::Editor::DataStorage::Queries
diff --git a/Engine/Plugins/Experimental/EditorDataStorage/Source/TedsCore/Private/Queries/QueryContextCapabilities/ContextEnvironment.h b/Engine/Plugins/Experimental/EditorDataStorage/Source/TedsCore/Private/Queries/QueryContextCapabilities/ContextEnvironment.h
new file mode 100644
--- /dev/null
+++ b/Engine/Plugins/Experimental/EditorDataStorage/Source/TedsCore/Private/Queries/QueryContextCapabilities/ContextEnvironment.h
@@ -0,0 +1,22 @@
+// Copyright Epic Games, Inc. All Rights Reserved.
+
+#pragma once
+
+#include "DataStorage/CommonTypes.h"
+#include "Elements/Framework/TypedElementQueryFunctions.h"
+
+namespace UE::Editor::DataStorage
+{
+	class FEnvironment;
+} // namespace UE::Editor::DataStorage
+
+namespace UE::Editor::DataStorage::Queries
+{
+	struct IContextEnvironment : IQueryFunctionResponse
+	{
+		virtual ~IContextEnvironment() = default;
+
+		virtual FEnvironment* GetEnvironment() = 0;
+		virtual const FEnvironment* GetEnvironment() const = 0;
+	};
+} // namespace UE::Editor::DataStorage::Queries
diff --git a/Engine/Plugins/Experimental/EditorDataStorage/Source/TedsCore/Private/Queries/QueryContextCapabilities/EditorDataStorageContextEnvironment.h b/Engine/Plugins/Experimental/EditorDataStorage/Source/TedsCore/Private/Queries/QueryContextCapabilities/EditorDataStorageContextEnvironment.h
deleted file mode 100644
diff --git a/Engine/Plugins/Experimental/EditorDataStorage/Source/TedsCore/Private/Queries/QueryContextCapabilities/EditorDataStorageDynamicColumnInfo.cpp b/Engine/Plugins/Experimental/EditorDataStorage/Source/TedsCore/Private/Queries/QueryContextCapabilities/EditorDataStorageDynamicColumnInfo.cpp
--- a/Engine/Plugins/Experimental/EditorDataStorage/Source/TedsCore/Private/Queries/QueryContextCapabilities/EditorDataStorageDynamicColumnInfo.cpp
+++ b/Engine/Plugins/Experimental/EditorDataStorage/Source/TedsCore/Private/Queries/QueryContextCapabilities/EditorDataStorageDynamicColumnInfo.cpp
@@ -2,7 +2,7 @@
 
 #include "Queries/QueryContextCapabilities/EditorDataStorageDynamicColumnInfo.h"
 
-#include "Queries/QueryContextCapabilities/EditorDataStorageContextEnvironment.h"
+#include "Queries/QueryContextCapabilities/ContextEnvironment.h"
 
 namespace UE::Editor::DataStorage::Queries
 {
diff --git a/Engine/Plugins/Experimental/EditorDataStorage/Source/TedsCore/Private/Queries/QueryContextCapabilities/RowBatch.cpp b/Engine/Plugins/Experimental/EditorDataStorage/Source/TedsCore/Private/Queries/QueryContextCapabilities/RowBatch.cpp
--- a/Engine/Plugins/Experimental/EditorDataStorage/Source/TedsCore/Private/Queries/QueryContextCapabilities/RowBatch.cpp
+++ b/Engine/Plugins/Experimental/EditorDataStorage/Source/TedsCore/Private/Queries/QueryContextCapabilities/RowBatch.cpp
@@ -3,8 +3,8 @@
 #include "Queries/QueryContextCapabilities/RowBatch.h"
 
 #include "Queries/QueryContextCapabilities/TableCapabilities.h"
-#include "Queries/TypedElementQueryContext_DirectSingleApi.h"
-#include "Queries/TypedElementQueryContext_MassForwarder.h"
+#include "Queries/DirectContextEnvironment.h"
+#include "Queries/ProcessorContextEnvironment.h"
 #include "TypedElementDatabase.h"
 #include "TypedElementDatabaseEnvironment.h"
 #include "TypedElementUtils.h"
diff --git a/Engine/Plugins/Experimental/EditorDataStorage/Source/TedsCore/Private/Queries/QueryContextCapabilities/RowBatch.h b/Engine/Plugins/Experimental/EditorDataStorage/Source/TedsCore/Private/Queries/QueryContextCapabilities/RowBatch.h
--- a/Engine/Plugins/Experimental/EditorDataStorage/Source/TedsCore/Private/Queries/QueryContextCapabilities/RowBatch.h
+++ b/Engine/Plugins/Experimental/EditorDataStorage/Source/TedsCore/Private/Queries/QueryContextCapabilities/RowBatch.h
@@ -5,7 +5,7 @@
 #include "DataStorage/CommonTypes.h"
 #include "DataStorage/Queries/Types.h"
 #include "Elements/Framework/TypedElementQueryCapabilities.h"
-#include "Queries/QueryContextCapabilities/EditorDataStorageContextEnvironment.h"
+#include "Queries/QueryContextCapabilities/ContextEnvironment.h"
 
 namespace UE::Editor::DataStorage::Queries
 {
diff --git a/Engine/Plugins/Experimental/EditorDataStorage/Source/TedsCore/Private/Queries/QueryContextCapabilities/SingleRow.cpp b/Engine/Plugins/Experimental/EditorDataStorage/Source/TedsCore/Private/Queries/QueryContextCapabilities/SingleRow.cpp
--- a/Engine/Plugins/Experimental/EditorDataStorage/Source/TedsCore/Private/Queries/QueryContextCapabilities/SingleRow.cpp
+++ b/Engine/Plugins/Experimental/EditorDataStorage/Source/TedsCore/Private/Queries/QueryContextCapabilities/SingleRow.cpp
@@ -2,7 +2,7 @@
 
 #include "Queries/QueryContextCapabilities/SingleRow.h"
 
-#include "Queries/QueryContextCapabilities/EditorDataStorageContextEnvironment.h"
+#include "Queries/QueryContextCapabilities/ContextEnvironment.h"
 #include "Queries/QueryContextCapabilities/TableCapabilities.h"
 #include "TypedElementDatabaseEnvironment.h"
 
diff --git a/Engine/Plugins/Experimental/EditorDataStorage/Source/TedsCore/Private/Queries/QueryContextCapabilities/TableCapabilities.cpp b/Engine/Plugins/Experimental/EditorDataStorage/Source/TedsCore/Private/Queries/QueryContextCapabilities/TableCapabilities.cpp
--- a/Engine/Plugins/Experimental/EditorDataStorage/Source/TedsCore/Private/Queries/QueryContextCapabilities/TableCapabilities.cpp
+++ b/Engine/Plugins/Experimental/EditorDataStorage/Source/TedsCore/Private/Queries/QueryContextCapabilities/TableCapabilities.cpp
@@ -4,8 +4,8 @@
 
 #include "MassArchetypeTypes.h"
 #include "MassExecutionContext.h"
-#include "Queries/TypedElementQueryContext_DirectSingleApi.h"
-#include "Queries/TypedElementQueryContext_MassForwarder.h"
+#include "Queries/DirectContextEnvironment.h"
+#include "Queries/ProcessorContextEnvironment.h"
 #include "TypedElementDatabaseEnvironment.h"
 
 namespace UE::Editor::DataStorage::Queries
diff --git a/Engine/Plugins/Experimental/EditorDataStorage/Source/TedsCore/Private/Queries/QueryContextCapabilities/TableCapabilities.h b/Engine/Plugins/Experimental/EditorDataStorage/Source/TedsCore/Private/Queries/QueryContextCapabilities/TableCapabilities.h
--- a/Engine/Plugins/Experimental/EditorDataStorage/Source/TedsCore/Private/Queries/QueryContextCapabilities/TableCapabilities.h
+++ b/Engine/Plugins/Experimental/EditorDataStorage/Source/TedsCore/Private/Queries/QueryContextCapabilities/TableCapabilities.h
@@ -5,7 +5,7 @@
 #include "DataStorage/CommonTypes.h"
 #include "DataStorage/Queries/Types.h"
 #include "Elements/Framework/TypedElementQueryCapabilities.h"
-#include "Queries/QueryContextCapabilities/EditorDataStorageContextEnvironment.h"
+#include "Queries/QueryContextCapabilities/ContextEnvironment.h"
 
 class FName;
 class UScriptStruct;
diff --git a/Engine/Plugins/Experimental/EditorDataStorage/Source/TedsCore/Private/Queries/TypedElementQueryContext_DirectSingleApi.cpp b/Engine/Plugins/Experimental/EditorDataStorage/Source/TedsCore/Private/Queries/TypedElementQueryContext_DirectSingleApi.cpp
deleted file mode 100644
diff --git a/Engine/Plugins/Experimental/EditorDataStorage/Source/TedsCore/Private/Queries/TypedElementQueryContext_DirectSingleApi.h b/Engine/Plugins/Experimental/EditorDataStorage/Source/TedsCore/Private/Queries/TypedElementQueryContext_DirectSingleApi.h
deleted file mode 100644
diff --git a/Engine/Plugins/Experimental/EditorDataStorage/Source/TedsCore/Private/Queries/TypedElementQueryContext_MassForwarder.cpp b/Engine/Plugins/Experimental/EditorDataStorage/Source/TedsCore/Private/Queries/TypedElementQueryContext_MassForwarder.cpp
deleted file mode 100644
diff --git a/Engine/Plugins/Experimental/EditorDataStorage/Source/TedsCore/Private/Queries/TypedElementQueryContext_MassForwarder.h b/Engine/Plugins/Experimental/EditorDataStorage/Source/TedsCore/Private/Queries/TypedElementQueryContext_MassForwarder.h
deleted file mode 100644
diff --git a/Engine/Plugins/Experimental/EditorDataStorage/Source/TedsCore/Private/TypedElementDatabase.cpp b/Engine/Plugins/Experimental/EditorDataStorage/Source/TedsCore/Private/TypedElementDatabase.cpp
--- a/Engine/Plugins/Experimental/EditorDataStorage/Source/TedsCore/Private/TypedElementDatabase.cpp
+++ b/Engine/Plugins/Experimental/EditorDataStorage/Source/TedsCore/Private/TypedElementDatabase.cpp
@@ -18,7 +18,6 @@
 #include "Processors/TypedElementProcessorAdaptors.h"
 #include "ProfilingDebugging/CpuProfilerTrace.h"
 #include "ProfilingDebugging/CountersTrace.h"
-#include "Queries/TypedElementQueryContext_DirectSingleApi.h"
 #include "Stats/Stats.h"
 #include "TickTaskManagerInterface.h"
 #include "TypedElementDatabaseEnvironment.h"
diff --git a/Engine/Plugins/Experimental/EditorDataStorage/Source/TedsCore/Private/TypedElementDatabase.h b/Engine/Plugins/Experimental/EditorDataStorage/Source/TedsCore/Private/TypedElementDatabase.h
--- a/Engine/Plugins/Experimental/EditorDataStorage/Source/TedsCore/Private/TypedElementDatabase.h
+++ b/Engine/Plugins/Experimental/EditorDataStorage/Source/TedsCore/Private/TypedElementDatabase.h
@@ -6,13 +6,13 @@
 #include "Elements/Interfaces/TypedElementDataStorageInterface.h"
 #include "MassArchetypeTypes.h"
 #include "Misc/TVariant.h"
-#include "Processors/TypedElementQueryCapabilityImplemenationForwarder.h"
+#include "Processors/QueryCapabilityImplementationForwarder.h"
 #include "Templates/SharedPointer.h"
 #include "TypedElementDatabaseCommandBuffer.h"
 #include "UObject/ObjectMacros.h"
 #include "Queries/QueryContextCapabilities/EditorDataStorageDynamicColumnInfo.h"
 #include "Queries/QueryContextCapabilities/TableCapabilities.h"
-#include "Queries/TypedElementQueryContext_DirectSingleApi.h"
+#include "Queries/DirectContextEnvironment.h"
 
 #include "TypedElementDatabase.generated.h"
 
