TEDS: Modify the query adding the pie tag to rows to be in the correct phase
TEDS Folders: add the pie tag to folders manually as well, to work around PIE folder registration timing while the integration is in progress

#Synced-CL 52629034

diff --git a/Engine/Plugins/Experimental/EditorDataStorage/Source/TedsCore/Private/Compatibility/TypedElementUObjectWorldQueries.cpp b/Engine/Plugins/Experimental/EditorDataStorage/Source/TedsCore/Private/Compatibility/TypedElementUObjectWorldQueries.cpp
--- a/Engine/Plugins/Experimental/EditorDataStorage/Source/TedsCore/Private/Compatibility/TypedElementUObjectWorldQueries.cpp
+++ b/Engine/Plugins/Experimental/EditorDataStorage/Source/TedsCore/Private/Compatibility/TypedElementUObjectWorldQueries.cpp
@@ -78,7 +78,7 @@ void UObjectWorldDataStorageFactory::RegisterPIEWorldSyncQueries(UE::Editor::Dat
 	DataStorage.RegisterQuery(
 		Select(
 			TEXT("Sync PIE Object Tag"),
-			FProcessor(EQueryTickPhase::PostPhysics, DataStorage.GetQueryTickGroupName(EQueryTickGroups::SyncExternalToDataStorage))
+			FProcessor(EQueryTickPhase::DuringPhysics, DataStorage.GetQueryTickGroupName(EQueryTickGroups::PostUpdate))
 				.SetExecutionMode(EExecutionMode::GameThread),
 			[](IQueryContext& Context, RowHandle Row, const FTypedElementWorldColumn& WorldColumn)
 			{
diff --git a/Engine/Source/Editor/UnrealEd/Private/ActorFolders/TedsActorFolderUtils.cpp b/Engine/Source/Editor/UnrealEd/Private/ActorFolders/TedsActorFolderUtils.cpp
--- a/Engine/Source/Editor/UnrealEd/Private/ActorFolders/TedsActorFolderUtils.cpp
+++ b/Engine/Source/Editor/UnrealEd/Private/ActorFolders/TedsActorFolderUtils.cpp
@@ -63,6 +63,14 @@ namespace UE::Editor::DataStorage::ActorFolders
 	void SetFolderColumns(ICoreProvider* Storage, RowHandle Row, UWorld* World, const FFolder& Folder)
 	{
 		Storage->AddColumn(Row, FTypedElementWorldColumn{.World = TWeakObjectPtr<UWorld>(World)});
+		
+		// This is also done through a processor for all rows with a world, but PIE folders have a specific lifecycle in TEDS that prevents
+		// this from being added the correct time, so we do it explicitly here as well
+		if (World->IsPlayInEditor())
+		{
+			Storage->AddColumn<FPieObjectTag>(Row);
+		}
+		
 		Storage->AddColumn(Row, FFolderCompatibilityColumn{.Folder = Folder});
 
 		if (UActorFolder* ActorFolder = Folder.GetActorFolder())
