[TEDS] Added functionality to retrieve the unique tag that's added to a table.

#Synced-CL 52643848

diff --git a/Engine/Plugins/Experimental/EditorDataStorage/Source/TedsCore/Private/Queries/QueryContextCapabilities/TableCapabilities.cpp b/Engine/Plugins/Experimental/EditorDataStorage/Source/TedsCore/Private/Queries/QueryContextCapabilities/TableCapabilities.cpp
--- a/Engine/Plugins/Experimental/EditorDataStorage/Source/TedsCore/Private/Queries/QueryContextCapabilities/TableCapabilities.cpp
+++ b/Engine/Plugins/Experimental/EditorDataStorage/Source/TedsCore/Private/Queries/QueryContextCapabilities/TableCapabilities.cpp
@@ -70,6 +70,11 @@ namespace UE::Editor::DataStorage::Queries
 		return Environment.GetEnvironment()->GetTableManager().ListForeignKeyDomains(Table, bIncludeParents, Callback);
 	}
 
+	const UScriptStruct* FTableInfoContextCapability::GetTagForTable(const IContextEnvironment& Environment, TableHandle Table)
+	{
+		return Environment.GetEnvironment()->GetTableManager().GetTagForTable(Table);
+	}
+
 
 
 	// 
diff --git a/Engine/Plugins/Experimental/EditorDataStorage/Source/TedsCore/Private/Queries/QueryContextCapabilities/TableCapabilities.h b/Engine/Plugins/Experimental/EditorDataStorage/Source/TedsCore/Private/Queries/QueryContextCapabilities/TableCapabilities.h
--- a/Engine/Plugins/Experimental/EditorDataStorage/Source/TedsCore/Private/Queries/QueryContextCapabilities/TableCapabilities.h
+++ b/Engine/Plugins/Experimental/EditorDataStorage/Source/TedsCore/Private/Queries/QueryContextCapabilities/TableCapabilities.h
@@ -6,6 +6,7 @@
 #include "DataStorage/Queries/Types.h"
 #include "Elements/Framework/TypedElementQueryCapabilities.h"
 #include "Queries/QueryContextCapabilities/ContextEnvironment.h"
+#include "Templates/UnrealTypeTraits.h"
 
 class FName;
 class UScriptStruct;
@@ -38,6 +39,8 @@ namespace UE::Editor::DataStorage::Queries
 		
 		static void ListTableForeignKeyDomains(
 			const IContextEnvironment& Environment, TableHandle Table, bool bIncludeParents, TFunctionRef<void(const FName&)> Callback);
+
+		static const UScriptStruct* GetTagForTable(const IContextEnvironment& Environment, TableHandle Table);
 	};
 
 	class FCurrentTableInfoContextCapability : public ImplementsContextCapability<CurrentTableInfo>
@@ -53,6 +56,16 @@ namespace UE::Editor::DataStorage::Queries
 			const FProcessorContextEnvironment& Environment, TConstArrayView<const UScriptStruct*> Columns);
 		static bool CurrentTableHasColumns(
 			const FConstProcessorContextEnvironment& Environment, TConstArrayView<const UScriptStruct*> Columns);
+
+		template<typename EnvironmentContextType> requires (TIsDerivedFrom<EnvironmentContextType, IContextEnvironment>::Value)
+		static const UScriptStruct* GetTagForCurrentTable(const EnvironmentContextType& Environment)
+		{
+			TableHandle CurrentTable = GetCurrentTable(Environment);
+			checkf(CurrentTable != InvalidTableHandle, TEXT("If this function is reachable there should be a valid table."));
+			const UScriptStruct* Tag = FTableInfoContextCapability::GetTagForTable(Environment, CurrentTable);
+			checkf(Tag, TEXT("The current table doesn't have a tag associated with it."));
+			return Tag;
+		}
 	};
 	
 	class FTableManagementContextCapability : public ImplementsContextCapability<TableManagement>
diff --git a/Engine/Plugins/Experimental/EditorDataStorage/Source/TedsCore/Private/TableManager.cpp b/Engine/Plugins/Experimental/EditorDataStorage/Source/TedsCore/Private/TableManager.cpp
--- a/Engine/Plugins/Experimental/EditorDataStorage/Source/TedsCore/Private/TableManager.cpp
+++ b/Engine/Plugins/Experimental/EditorDataStorage/Source/TedsCore/Private/TableManager.cpp
@@ -82,6 +82,7 @@ namespace UE::Editor::DataStorage
 		}
 		// Add table tag.
 		const UScriptStruct* TableTag = DynamicColumns.GenerateColumn(*FTableDynamicTag::StaticStruct(), TableName).Type;
+		checkf(TableTag, TEXT("Failed to generate table tag for table '%s'"), *TableName.ToString());
 		Composition.Add(TableTag);
 
 		bRegisteringNewTable = true;
@@ -92,8 +93,9 @@ namespace UE::Editor::DataStorage
 			{
 				.Archetype = Archetype,
 				.Name = TableName,
+				.Type = Type,
+				.Tag = TNonNullPtr<const UScriptStruct>(TableTag),
 				.Parent = Options.SourceTable,
-				.Type = Type
 			});
 		NameLookup.Add(TableName, Result);
 		MassLookup.Add(Archetype, Result);
@@ -117,6 +119,19 @@ namespace UE::Editor::DataStorage
 		return Find(MassEntityManager.GetArchetypeForEntity(FMassEntityHandle::FromNumber(Row)));
 	}
 
+	const UScriptStruct* FTableManager::GetTagForTable(TableHandle Table) const
+	{
+		if (IsTableValid(Table))
+		{
+			int32 TableHandleAsIndex = ConvertTableHandleToIndex(Table);
+			return Tables[TableHandleAsIndex].Tag;
+		}
+		else
+		{
+			return nullptr;
+		}
+	}
+
 	void FTableManager::List(TFunctionRef<void(TableHandle)> Callback) const
 	{
 		int32 TableCount = Tables.Num();
@@ -313,8 +328,9 @@ namespace UE::Editor::DataStorage
 				{
 					.Archetype = Archetype,
 					.Name = ParentFName,
+					.Type = ETableType::Variant,
+					.Tag = TNonNullPtr<const UScriptStruct>(TableTag),
 					.Parent = Parent,
-					.Type = ETableType::Variant
 				});
 			TableHandle Variant = MassLookup.Add(Archetype, NewTable);
 			Tables[ConvertTableHandleToIndex(Parent)].Variants.Add(NewTable);
diff --git a/Engine/Plugins/Experimental/EditorDataStorage/Source/TedsCore/Private/TableManager.h b/Engine/Plugins/Experimental/EditorDataStorage/Source/TedsCore/Private/TableManager.h
--- a/Engine/Plugins/Experimental/EditorDataStorage/Source/TedsCore/Private/TableManager.h
+++ b/Engine/Plugins/Experimental/EditorDataStorage/Source/TedsCore/Private/TableManager.h
@@ -53,6 +53,8 @@ namespace UE::Editor::DataStorage
 		TableHandle Find(const FMassArchetypeHandle& Archetype) const;
 		TableHandle Find(RowHandle Row) const;
 		
+		const UScriptStruct* GetTagForTable(TableHandle Table) const;
+		
 		void List(TFunctionRef<void(TableHandle)> Callback) const;
 		void List(ETableType TableType, TFunctionRef<void(TableHandle)> Callback) const;
 		
@@ -77,10 +79,11 @@ namespace UE::Editor::DataStorage
 			TArray<TableHandle> Derivatives;
 			FMassArchetypeHandle Archetype;
 			FName Name;
+			ETableType Type;
+			TNonNullPtr<const UScriptStruct> Tag;
 			// If this is a dynamically generated variant then the Parent represents the table used as the basis for the variant. 
 			// If this is a derived copy then the Parent represents the table that was passed in as the template for this table.
 			TableHandle Parent;
-			ETableType Type;
 		};
 
 		void OnNewArchetype(const FMassArchetypeHandle& Archetype);
diff --git a/Engine/Source/Runtime/TypedElementFramework/Public/Elements/Framework/TypedElementQueryCapabilities.inl b/Engine/Source/Runtime/TypedElementFramework/Public/Elements/Framework/TypedElementQueryCapabilities.inl
--- a/Engine/Source/Runtime/TypedElementFramework/Public/Elements/Framework/TypedElementQueryCapabilities.inl
+++ b/Engine/Source/Runtime/TypedElementFramework/Public/Elements/Framework/TypedElementQueryCapabilities.inl
@@ -32,6 +32,8 @@ CapabilityStart(TableInfo, EContextCapabilityFlags::SupportsSingle | EContextCap
 	ConstFunction2(TableInfo, void, ListTableRows, (TableHandle, Table), (TFunctionRef<void(FRowHandleArrayView)>, Callback))
 	/** Lists all the domain names for foreign keys that have been registered with this table. */
 	ConstFunction3(TableInfo, void, ListTableForeignKeyDomains, (TableHandle, Table), (bool, bIncludeParents), (TFunctionRef<void(const FName&)>, Callback))
+	/** Get the unique tag that identifies the requested table. If the table is not found a nullptr will be returned. */
+	ConstFunction1(TableInfo, const UScriptStruct*, GetTagForTable, (TableHandle, Table))
 
 #if defined(WithWrappers)
 	/** Returns true if the table contains the requested columns, otherwise false. */
@@ -51,6 +53,8 @@ CapabilityStart(CurrentTableInfo, EContextCapabilityFlags::SupportsSingle | ECon
 	ConstFunction0(CurrentTableInfo, TableHandle, GetCurrentTable)
 	/** Returns true if the current table contains the requested columns, otherwise false. */
 	ConstFunction1(CurrentTableInfo, bool, CurrentTableHasColumns, (TConstArrayView<const UScriptStruct*>, Columns))
+	/** Get the unique tag that identifies the current table. */
+	ConstFunction0(CurrentTableInfo, const UScriptStruct*, GetTagForCurrentTable)
 
 #if defined(WithWrappers)
 	/** Returns true if the current table contains the requested columns, otherwise false. */
@@ -289,4 +293,4 @@ CapabilityStart(DynamicColumnInfo, EContextCapabilityFlags::SupportsSingle | ECo
 		return Description.TemplateType ? this->FindDynamicColumnType(*Description.TemplateType, Description.Identifier) : nullptr;
 	}
 #endif
-CapabilityEnd(DynamicColumnInfo)
\ No newline at end of file
+CapabilityEnd(DynamicColumnInfo)
