[TEDS] Relations

Adds relations to TEDS. A relation connects two rows (Subject and Object) with
a named type. Each role can be configured with exclusivity and destruction
policies. Implementation is based on Mass Relations.

Relations can optionally be hierarchical, enabling tree queries like
IsDescendantOf, GetAncestors, and TraverseDescendants. Two hierarchy modes are
available:

- WalkOnly: parent-chain traversal, O(depth). Simple and always correct.
- IntervalEncoded: nested-set intervals for O(1) IsDescendantOf. Intervals are
assigned on creation and automatically rebalanced at frame end when they
become stale (e.g. after reparenting). Stale nodes fall back to WalkOnly
until rebalanced.

Hierarchies can use either the existing Legacy (adjacency-list column) backend
or the new Relations backend. Both are transparent to user code:
AccessesHierarchy, SetParentRow, and GetParentChangedColumnType work uniformly
regardless of backend.

Relation change-notification columns (SubjectChanged / ObjectChanged) can be
enabled per relation type. These tag columns are stamped on participant rows
when a relation is created or destroyed, and automatically cleaned up at frame
end.

Relation predicates can be used in queries (e.g.
Where().IsDescendantOf(Type, Row)) with per-entity filtering.

The TEDS Debugger QueryEditor surfaces relations-backed hierarchies, and
SRowDetails now displays relation links with clickable navigation between
related rows.

395 automation tests (Relations + Hierarchy).

#Synced-CL 52664489

diff --git a/Engine/Plugins/Experimental/EditorDataStorage/Source/TedsCore/Private/Hierarchy/EditorDataStorageHierarchyImplementation.cpp b/Engine/Plugins/Experimental/EditorDataStorage/Source/TedsCore/Private/Hierarchy/EditorDataStorageHierarchyImplementation.cpp
--- a/Engine/Plugins/Experimental/EditorDataStorage/Source/TedsCore/Private/Hierarchy/EditorDataStorageHierarchyImplementation.cpp
+++ b/Engine/Plugins/Experimental/EditorDataStorage/Source/TedsCore/Private/Hierarchy/EditorDataStorageHierarchyImplementation.cpp
@@ -6,11 +6,112 @@
 
 namespace UE::Editor::DataStorage
 {
-	// Interpretation for the opaque FHierarchyHandle
+	// ──────────────────────────────────────────────────────────────────────────
+	// Concrete hierarchy backend declarations (private to this translation unit)
+	// ──────────────────────────────────────────────────────────────────────────
+
+	/**
+	 * Adjacency-list (legacy) hierarchy backend.
+	 * Stores parent/child data in columns directly on participant rows.
+	 */
+	class FAdjacencyListBackedHierarchyAccess final : public FTedsHierarchyAccessInterface
+	{
+	public:
+		struct FConstructParams
+		{
+			const UScriptStruct* ChildTag = nullptr;
+			const UScriptStruct* ParentTag = nullptr;
+			const UScriptStruct* HierarchyData = nullptr;
+			const UScriptStruct* UnresolvedParentColumn = nullptr;
+			const UScriptStruct* ParentChangedColumn = nullptr;
+		};
+
+		explicit FAdjacencyListBackedHierarchyAccess(const FConstructParams& Params);
+
+		const UScriptStruct* GetChildTagType() const override;
+		const UScriptStruct* GetParentTagType() const override;
+		const UScriptStruct* GetHierarchyDataColumnType() const override;
+		const UScriptStruct* GetUnresolvedParentColumnType() const override;
+		const UScriptStruct* GetParentChangedColumnType() const override;
+
+		void SetParentRow(ICoreProvider* CoreProvider, RowHandle Target, RowHandle Parent) const override;
+		void SetUnresolvedParent(ICoreProvider* CoreProvider, RowHandle Target, FMapKey ParentId, FName MappingDomain) const override;
+		RowHandle GetParentRow(const ICoreProvider* CoreProvider, RowHandle Target) const override;
+
+		bool HasChildren(const ICoreProvider& Context, RowHandle Row) const override;
+		void WalkDepthFirst(
+			const ICoreProvider& Context,
+			RowHandle Row,
+			TFunction<void(const ICoreProvider& Context, RowHandle Owner, RowHandle Target)> OnVisitedFn,
+			ICoreProvider::ETraversalOrder TraversalOrder) const override;
+		bool IterateChildren(
+			const ICoreProvider& Context,
+			RowHandle Row,
+			TFunctionRef<bool(const ICoreProvider& Context, RowHandle Child)> OnVisitedFn) const override;
+
+		TFunction<RowHandle(const void*, const UScriptStruct*)> CreateParentExtractionFunction() const override;
+
+	protected:
+		RowHandle GetParent(const ICommonQueryContext& Context, RowHandle Target) const override;
+
+	private:
+		const UScriptStruct* ChildTag = nullptr;
+		const UScriptStruct* ParentTag = nullptr;
+		const UScriptStruct* HierarchyDataColumnType = nullptr;
+		const UScriptStruct* UnresolvedParentColumnType = nullptr;
+		const UScriptStruct* ParentChangedColumn = nullptr;
+	};
+
+	/**
+	 * Relations-backed hierarchy backend.
+	 * Stores hierarchy data on dedicated relation rows, using interval encoding for O(1) descent queries.
+	 */
+	class FTedsRelationBackedHierarchyAccess final : public FTedsHierarchyAccessInterface
+	{
+	public:
+		FTedsRelationBackedHierarchyAccess(RelationTypeHandle InRelationType,
+			const UScriptStruct* InUnresolvedParentColumnType,
+			const UScriptStruct* InParentChangedColumn = nullptr);
+
+		// Relations backend doesn't use tag/data columns on participant rows.
+		const UScriptStruct* GetChildTagType() const override { return nullptr; }
+		const UScriptStruct* GetParentTagType() const override { return nullptr; }
+		const UScriptStruct* GetHierarchyDataColumnType() const override { return nullptr; }
+		const UScriptStruct* GetUnresolvedParentColumnType() const override { return UnresolvedParentColumnType; }
+		const UScriptStruct* GetParentChangedColumnType() const override { return ParentChangedColumnType; }
+
+		void SetParentRow(ICoreProvider* CoreProvider, RowHandle Target, RowHandle Parent) const override;
+		void SetUnresolvedParent(ICoreProvider* CoreProvider, RowHandle Target, FMapKey ParentId, FName MappingDomain) const override;
+		RowHandle GetParentRow(const ICoreProvider* CoreProvider, RowHandle Target) const override;
+
+		bool HasChildren(const ICoreProvider& Context, RowHandle Row) const override;
+		void WalkDepthFirst(
+			const ICoreProvider& Context,
+			RowHandle Row,
+			TFunction<void(const ICoreProvider& Context, RowHandle Owner, RowHandle Target)> OnVisitedFn,
+			ICoreProvider::ETraversalOrder TraversalOrder) const override;
+		bool IterateChildren(
+			const ICoreProvider& Context,
+			RowHandle Row,
+			TFunctionRef<bool(const ICoreProvider& Context, RowHandle Child)> OnVisitedFn) const override;
+
+		TFunction<RowHandle(const void*, const UScriptStruct*)> CreateParentExtractionFunction() const override;
+
+	private:
+		RelationTypeHandle HierarchyRelationType;
+		const UScriptStruct* UnresolvedParentColumnType = nullptr;
+		const UScriptStruct* ParentChangedColumnType = nullptr;
+	};
+
+	// ──────────────────────────────────────────────────────────────────────────
+
+	// Interpretation for the opaque FHierarchyHandle.
+	// Backend field distinguishes Legacy (column-on-row) from Relations (relation-row) handles.
 	struct FHierarchyHandleInterpretation
 	{
 		int32 Index;
 		bool IsSet;
+		uint8 Backend;
 	};
 
 	static int32 GetHandleIndex(const FHierarchyHandle& Handle)
@@ -19,15 +120,22 @@ namespace UE::Editor::DataStorage
 		return Interp.IsSet ? Interp.Index : INDEX_NONE;
 	}
 
-	static FHierarchyHandle CreateHandle(int32 Index)
+	static EHierarchyBackend GetHandleBackend(const FHierarchyHandle& Handle)
+	{
+		const FHierarchyHandleInterpretation& Interp = reinterpret_cast<const FHierarchyHandleInterpretation&>(Handle);
+		return Interp.IsSet ? static_cast<EHierarchyBackend>(Interp.Backend) : EHierarchyBackend::Legacy;
+	}
+
+	static FHierarchyHandle CreateHandle(int32 Index, EHierarchyBackend Backend = EHierarchyBackend::Legacy)
 	{
 		FHierarchyHandle Handle;
-		
+
 		FHierarchyHandleInterpretation& Interp = reinterpret_cast<FHierarchyHandleInterpretation&>(Handle);
 		new (&Interp) FHierarchyHandleInterpretation();
 		Interp.Index = Index;
 		Interp.IsSet = true;
-		
+		Interp.Backend = static_cast<uint8>(Backend);
+
 		return Handle;
 	}
 
@@ -113,8 +221,16 @@ namespace UE::Editor::DataStorage
 
 	FHierarchyHandle FTedsHierarchyRegistrar::RegisterHierarchy(ICoreProvider* InProvider, const FHierarchyRegistrationParams& RegistrationParams)
 	{
-		// Create an access interface
-		
+		if (RegistrationParams.Backend == FHierarchyRegistrationParams::EBackend::Relations)
+		{
+			return RegisterHierarchyRelations(InProvider, RegistrationParams);
+		}
+
+		return RegisterHierarchyLegacy(InProvider, RegistrationParams);
+	}
+
+	FHierarchyHandle FTedsHierarchyRegistrar::RegisterHierarchyLegacy(ICoreProvider* InProvider, const FHierarchyRegistrationParams& RegistrationParams)
+	{
 		// Generate the tags and HierarchyData column per Hierarchy type
 		const UScriptStruct* HierarchyData = [InProvider, &RegistrationParams]()
 		{
@@ -125,7 +241,7 @@ namespace UE::Editor::DataStorage
 			};
 			return InProvider->GenerateDynamicColumn(Desc);
 		}();
-		
+
 		const UScriptStruct* ParentTag = [InProvider, &RegistrationParams]()
 		{
 			FDynamicColumnDescription Desc
@@ -135,7 +251,7 @@ namespace UE::Editor::DataStorage
 			};
 			return InProvider->GenerateDynamicColumn(Desc);
 		}();
-		
+
 		const UScriptStruct* ChildTag = [InProvider, &RegistrationParams]()
 		{
 			FDynamicColumnDescription Desc
@@ -155,16 +271,16 @@ namespace UE::Editor::DataStorage
 			};
 			return InProvider->GenerateDynamicColumn(Desc);
 		}();
-		
+
 		const UScriptStruct* ParentChangedColumn = nullptr;
-		
+
 		// If the old deprecated ParentChangedColumn arg was provided, keep using it for backwards compatibility
 		PRAGMA_DISABLE_DEPRECATION_WARNINGS
 		if (RegistrationParams.ParentChangedColumn)
 		{
 			ParentChangedColumn = RegistrationParams.ParentChangedColumn;
 		}
-		PRAGMA_ENABLE_DEPRECATION_WARNINGS				
+		PRAGMA_ENABLE_DEPRECATION_WARNINGS
 		else if (RegistrationParams.bEnableParentChangedColumn)
 		{
 			ParentChangedColumn = [InProvider, &RegistrationParams]()
@@ -177,10 +293,10 @@ namespace UE::Editor::DataStorage
 				return InProvider->GenerateDynamicColumn(Desc);
 			}();
 		}
-		 
+
 		TUniquePtr<FTedsHierarchyAccessInterface> AccessInterface = [&]()
 		{
-			FTedsHierarchyAccessInterface::FConstructParams Params
+			FAdjacencyListBackedHierarchyAccess::FConstructParams Params
 			{
 				.ChildTag = ChildTag,
 				.ParentTag = ParentTag,
@@ -188,9 +304,9 @@ namespace UE::Editor::DataStorage
 				.UnresolvedParentColumn = UnresolvedParentColumn,
 				.ParentChangedColumn = ParentChangedColumn
 			};
-			return MakeUnique<FTedsHierarchyAccessInterface>(Params);
+			return MakeUnique<FAdjacencyListBackedHierarchyAccess>(Params);
 		}();
-		
+
 		FRegisteredHierarchy RegisteredHierarchy
 		{
 			.Name = RegistrationParams.Name,
@@ -199,17 +315,86 @@ namespace UE::Editor::DataStorage
 			.HierarchyData = HierarchyData,
 			.UnresolvedParentColumn = UnresolvedParentColumn,
 			.ParentChangedColumn = ParentChangedColumn,
-			.AccessInterface = MoveTemp(AccessInterface)
+			.AccessInterface = MoveTemp(AccessInterface),
+			.Backend = EHierarchyBackend::Legacy
 		};
 
 		int32 Index = RegisteredHierarchies.Num();
 		RegisteredHierarchies.Emplace(MoveTemp(RegisteredHierarchy));
 
-		FHierarchyHandle HierarchyHandle = CreateHandle(Index);
-		
+		FHierarchyHandle HierarchyHandle = CreateHandle(Index, EHierarchyBackend::Legacy);
+
 		RegisterObservers(*InProvider, HierarchyHandle);
 		RegisterProcessors(*InProvider, HierarchyHandle);
-		
+
+		return HierarchyHandle;
+	}
+
+	FHierarchyHandle FTedsHierarchyRegistrar::RegisterHierarchyRelations(ICoreProvider* InProvider, const FHierarchyRegistrationParams& RegistrationParams)
+	{
+		// Register a hierarchical relation type with exclusive Object (one parent per child)
+		FRelationRegistrationParams RelationParams;
+		RelationParams.Name = RegistrationParams.Name;
+		RelationParams.Traits.HierarchyMode = EHierarchyMode::IntervalEncoded;
+		RelationParams.Traits.Subject.bExclusive = false; // A child can only have one parent, but this is enforced by Object exclusivity
+		RelationParams.Traits.Object.bExclusive = true;   // Each child has at most one parent
+		// CleanUp: only the relation row is destroyed when a parent row is removed.
+		// Children are NOT cascade-destroyed — they become roots, matching legacy hierarchy behaviour.
+		RelationParams.Traits.Object.DestructionPolicy = FTedsRelationRoleTraits::EDestructionPolicy::CleanUp;
+		// Propagate the parent-changed column opt-in: Subject = child, so bEnableSubjectChangedColumn
+		// is the correct side -- the child row gets stamped when its parent relation changes.
+		RelationParams.bEnableSubjectChangedColumn = RegistrationParams.bEnableParentChangedColumn;
+
+		const RelationTypeHandle RelationType = InProvider->RegisterRelationType(RelationParams);
+		if (RelationType == InvalidRelationTypeHandle)
+		{
+			return FHierarchyHandle();
+		}
+
+		// If requested, retrieve the generated subject-changed column to use as ParentChangedColumn.
+		const UScriptStruct* ParentChangedColumn =
+			RegistrationParams.bEnableParentChangedColumn
+				? InProvider->GetRelationSubjectChangedColumn(RelationType)
+				: nullptr;
+
+		// Generate a dynamic unresolved parent column for deferred relation creation.
+		// Same pattern as Legacy — column on Subject row, auto-cleaned on row destruction.
+		const UScriptStruct* UnresolvedParentColumn = [InProvider, &RegistrationParams]()
+		{
+			FDynamicColumnDescription Desc
+			{
+				.TemplateType = FEditorDataHierarchyUnresolvedParent_Template::StaticStruct(),
+				.Identifier = RegistrationParams.Name
+			};
+			return InProvider->GenerateDynamicColumn(Desc);
+		}();
+
+		// Create a Relations-backed access interface, forwarding the ParentChangedColumn.
+		TUniquePtr<FTedsHierarchyAccessInterface> AccessInterface =
+			MakeUnique<FTedsRelationBackedHierarchyAccess>(RelationType, UnresolvedParentColumn, ParentChangedColumn);
+
+		FRegisteredHierarchy RegisteredHierarchy
+		{
+			.Name = RegistrationParams.Name,
+			.ChildTag = nullptr,
+			.ParentTag = nullptr,
+			.HierarchyData = nullptr,
+			.UnresolvedParentColumn = UnresolvedParentColumn,
+			.ParentChangedColumn = ParentChangedColumn,
+			.AccessInterface = MoveTemp(AccessInterface),
+			.Backend = EHierarchyBackend::Relations,
+			.HierarchyRelationType = RelationType
+		};
+
+		int32 Index = RegisteredHierarchies.Num();
+		RegisteredHierarchies.Emplace(MoveTemp(RegisteredHierarchy));
+
+		FHierarchyHandle HierarchyHandle = CreateHandle(Index, EHierarchyBackend::Relations);
+
+		// Register processors for unresolved parent resolution.
+		// No observers needed — the relation system handles lifecycle.
+		RegisterProcessors(*InProvider, HierarchyHandle);
+
 		return HierarchyHandle;
 	}
 
@@ -219,7 +404,7 @@ namespace UE::Editor::DataStorage
 		{
 			if (RegisteredHierarchies[HierarchyIndex].Name == Name)
 			{
-				return CreateHandle(HierarchyIndex);
+				return CreateHandle(HierarchyIndex, RegisteredHierarchies[HierarchyIndex].Backend);
 			}
 		}
 		return FHierarchyHandle();
@@ -378,15 +563,14 @@ namespace UE::Editor::DataStorage
 
 		const FRegisteredHierarchy* Hierarchy = FindRegisteredHierarchy(Handle);
 		checkf(Hierarchy, TEXT("Couldn't find hierarchy that was just registered, processors for the hierarchy could not be registered successfully"));
-		
-		CoreProvider.RegisterQuery(
-		Select(
+
+		auto ProcessorQuery = Select(
 			TEXT("Resolve hierarchy rows"),
 			FProcessor(EQueryTickPhase::PrePhysics, CoreProvider.GetQueryTickGroupName(EQueryTickGroups::Default)),
 			[this, Handle](IQueryContext& Context, RowHandle Row)
 			{
 				const FRegisteredHierarchy* Hierarchy = FindRegisteredHierarchy(Handle);
-				
+
 				FEditorDataHierarchyUnresolvedParent_Template* UnresolvedParentColumn;
 				{
 					const int RowIndexInChunk = Hack_FindRowIndexInChunk(Context, Row);
@@ -395,17 +579,27 @@ namespace UE::Editor::DataStorage
 					uint8* HierarchyDataComponentObjectAddress = static_cast<uint8*>(HierarchyDataComponentArrayStart) + (RowIndexInChunk * Hierarchy->UnresolvedParentColumn->GetStructureSize());
 					UnresolvedParentColumn = reinterpret_cast<FEditorDataHierarchyUnresolvedParent_Template*>(HierarchyDataComponentObjectAddress);
 				}
-				
+
 				RowHandle ParentRow = Context.LookupMappedRow(UnresolvedParentColumn->MappingDomain, UnresolvedParentColumn->ParentId);
 				if (Context.IsRowAvailable(ParentRow))
 				{
+					// Routes through HierarchyAccessInterface -- each backend's SetParentRow
+					// handles the correct storage (Legacy: adjacency-list column; Relations: CreateRelation).
 					Context.SetParentRow(Row, ParentRow);
 					Context.RemoveColumns(Row, {Hierarchy->UnresolvedParentColumn});
 				}
-			})
-		.AccessesHierarchy(Hierarchy->Name)
-		.ReadWrite(Hierarchy->UnresolvedParentColumn)
-		.Compile());
+			});
+
+		// AccessesHierarchy wires HierarchyAccessInterface into the query context so
+		// Context.SetParentRow dispatches through the correct backend implementation.
+		// SetupHierarchyDependencies now skips null HierarchyDataColumnType, so this
+		// works for both Legacy and Relations backends without adding invalid requirements.
+		ProcessorQuery.AccessesHierarchy(Hierarchy->Name);
+
+		CoreProvider.RegisterQuery(
+			ProcessorQuery
+			.ReadWrite(Hierarchy->UnresolvedParentColumn)
+			.Compile());
 	}
 
 	const FTedsHierarchyRegistrar::FRegisteredHierarchy* FTedsHierarchyRegistrar::FindRegisteredHierarchy(const FHierarchyHandle& Handle)
@@ -418,48 +612,49 @@ namespace UE::Editor::DataStorage
 		return nullptr;
 	}
 
-	FTedsHierarchyAccessInterface::FTedsHierarchyAccessInterface(const FConstructParams& Params)
+	FAdjacencyListBackedHierarchyAccess::FAdjacencyListBackedHierarchyAccess(const FConstructParams& Params)
 		: ChildTag(Params.ChildTag)
 		, ParentTag(Params.ParentTag)
 		, HierarchyDataColumnType(Params.HierarchyData)
 		, UnresolvedParentColumnType(Params.UnresolvedParentColumn)
 		, ParentChangedColumn(Params.ParentChangedColumn)
 	{
-		check(HierarchyDataColumnType->IsChildOf(FEditorDataHierarchyData_Template::StaticStruct()));
+		// HierarchyData may be nullptr if GenerateDynamicColumn fails; the check below guards against that.
+		check(!HierarchyDataColumnType || HierarchyDataColumnType->IsChildOf(FEditorDataHierarchyData_Template::StaticStruct()));
 	}
 
-	const UScriptStruct* FTedsHierarchyAccessInterface::GetChildTagType() const
+	const UScriptStruct* FAdjacencyListBackedHierarchyAccess::GetChildTagType() const
 	{
 		return ChildTag;
 	}
 
-	const UScriptStruct* FTedsHierarchyAccessInterface::GetParentTagType() const
+	const UScriptStruct* FAdjacencyListBackedHierarchyAccess::GetParentTagType() const
 	{
 		return ParentTag;
 	}
 
-	const UScriptStruct* FTedsHierarchyAccessInterface::GetHierarchyDataColumnType() const
+	const UScriptStruct* FAdjacencyListBackedHierarchyAccess::GetHierarchyDataColumnType() const
 	{
 		return HierarchyDataColumnType;
 	}
 
-	const UScriptStruct* FTedsHierarchyAccessInterface::GetUnresolvedParentColumnType() const
+	const UScriptStruct* FAdjacencyListBackedHierarchyAccess::GetUnresolvedParentColumnType() const
 	{
 		return UnresolvedParentColumnType;
 	}
 
-	const UScriptStruct* FTedsHierarchyAccessInterface::GetParentChangedColumnType() const
+	const UScriptStruct* FAdjacencyListBackedHierarchyAccess::GetParentChangedColumnType() const
 	{
 		return ParentChangedColumn;
 	}
 
-	bool FTedsHierarchyAccessInterface::HasChildren(const ICoreProvider& Context, RowHandle Row) const
+	bool FAdjacencyListBackedHierarchyAccess::HasChildren(const ICoreProvider& Context, RowHandle Row) const
 	{
 		const FEditorDataHierarchyData_Template* Hierarchy = static_cast<const FEditorDataHierarchyData_Template*>(Context.GetColumnData(Row, HierarchyDataColumnType));
 		return Hierarchy && !Hierarchy->Children.IsEmpty();
 	}
 
-	void FTedsHierarchyAccessInterface::WalkDepthFirst(
+	void FAdjacencyListBackedHierarchyAccess::WalkDepthFirst(
 		const ICoreProvider& Context,
 		RowHandle Row,
 		TFunction<void(const ICoreProvider& Context, RowHandle Owner,RowHandle Target)> OnVisitedFn,
@@ -507,7 +702,7 @@ namespace UE::Editor::DataStorage
 		}
 	}
 
-	bool FTedsHierarchyAccessInterface::IterateChildren(const ICoreProvider& Context, RowHandle Row, 
+	bool FAdjacencyListBackedHierarchyAccess::IterateChildren(const ICoreProvider& Context, RowHandle Row, 
 		TFunctionRef<bool(const ICoreProvider& Context, RowHandle Child)> OnVisitedFn) const
 	{
 		if (!Context.IsRowAvailable(Row))
@@ -529,7 +724,7 @@ namespace UE::Editor::DataStorage
 		return true;
 	}
 
-	void FTedsHierarchyAccessInterface::SetParentRow(ICoreProvider* CoreProvider, RowHandle Target, RowHandle Parent) const
+	void FAdjacencyListBackedHierarchyAccess::SetParentRow(ICoreProvider* CoreProvider, RowHandle Target, RowHandle Parent) const
 	{
 		// The target is not assigned (ie. neither assigned to a table nor reserved)
 		// It will never be a valid row so appease TEDS philosophy of not crashing with invalid rows
@@ -537,6 +732,7 @@ namespace UE::Editor::DataStorage
 		{
 			return;
 		}
+
 		CoreProvider->RemoveColumn(Target, UnresolvedParentColumnType);
 
 		auto AddHierarchyDataColumn = [this](ICoreProvider* CoreProvider, RowHandle Row, RowHandle ParentRow)
@@ -645,7 +841,7 @@ namespace UE::Editor::DataStorage
 		}
 	}
 
-	void FTedsHierarchyAccessInterface::SetUnresolvedParent(ICoreProvider* CoreProvider, RowHandle Target, FMapKey ParentId, FName MappingDomain) const
+	void FAdjacencyListBackedHierarchyAccess::SetUnresolvedParent(ICoreProvider* CoreProvider, RowHandle Target, FMapKey ParentId, FName MappingDomain) const
 	{
 		const UScriptStruct* UnresolvedParentColumn = GetUnresolvedParentColumnType();
 		
@@ -666,7 +862,7 @@ namespace UE::Editor::DataStorage
 			);
 	}
 
-	RowHandle FTedsHierarchyAccessInterface::GetParentRow(const ICoreProvider* CoreProvider, RowHandle Target) const
+	RowHandle FAdjacencyListBackedHierarchyAccess::GetParentRow(const ICoreProvider* CoreProvider, RowHandle Target) const
 	{
 		const UScriptStruct* HierarchyDataColumn = GetHierarchyDataColumnType();
 		const void* ColumnData = CoreProvider->GetColumnData(Target, HierarchyDataColumn);
@@ -678,7 +874,7 @@ namespace UE::Editor::DataStorage
 		return InvalidRowHandle;
 	}
 
-	TFunction<RowHandle(const void*, const UScriptStruct*)> FTedsHierarchyAccessInterface::CreateParentExtractionFunction() const
+	TFunction<RowHandle(const void*, const UScriptStruct*)> FAdjacencyListBackedHierarchyAccess::CreateParentExtractionFunction() const
 	{
 		return [this](const void* ColumnData, const UScriptStruct* ColumnType)->RowHandle
 		{
@@ -745,7 +941,12 @@ namespace UE::Editor::DataStorage
 		Context.PushCommand(MoveTemp(Command));	
 	}
 
-	RowHandle FTedsHierarchyAccessInterface::GetParent(const ICommonQueryContext& Context, RowHandle Target) const
+	RowHandle FTedsHierarchyAccessInterface::GetParent(const ICommonQueryContext& /*Context*/, RowHandle /*Target*/) const
+	{
+		return InvalidRowHandle;
+	}
+
+	RowHandle FAdjacencyListBackedHierarchyAccess::GetParent(const ICommonQueryContext& Context, RowHandle Target) const
 	{
 		// NOTE: It is not clear to the user that they will need to have registered the HierarchyDataColumnType as at least ReadOnly
 		// for this function to return the Parent.
@@ -812,26 +1013,142 @@ namespace UE::Editor::DataStorage
 
 	RowHandle FTedsHierarchyAccessInterface::GetParentRow(const IDirectQueryContext& Context, RowHandle Target) const
 	{
-		// NOTE: It is not clear to the user that they will need to have registered the HierarchyDataColumnType as at least ReadOnly
-		// for this function to return the Parent.
-		// This can be improved with a function on the the context (and other contexts) to check for the access requirements of a
-		// query and warn or ensure about a missing requirement - though this does not currently exist at time of writing.
-		const FEditorDataHierarchyData_Template* ParentHierarchyDataComponent;
+		return GetParent(static_cast<const ICommonQueryContext&>(Context), Target);
+	}
+
+	//
+	// FTedsRelationBackedHierarchyAccess
+	//
+
+	FTedsRelationBackedHierarchyAccess::FTedsRelationBackedHierarchyAccess(
+		RelationTypeHandle InRelationType,
+		const UScriptStruct* InUnresolvedParentColumnType,
+		const UScriptStruct* InParentChangedColumn)
+		: HierarchyRelationType(InRelationType)
+		, UnresolvedParentColumnType(InUnresolvedParentColumnType)
+		, ParentChangedColumnType(InParentChangedColumn)
+	{
+	}
+
+	void FTedsRelationBackedHierarchyAccess::SetParentRow(ICoreProvider* CoreProvider, RowHandle Target, RowHandle Parent) const
+	{
+		check(CoreProvider != nullptr);
+
+		// Destroy existing parent relation first (C1: explicit destroy-before-create for re-parenting)
+		RowHandle CurrentParent = CoreProvider->GetRelationObject(HierarchyRelationType, Target);
+		if (CurrentParent != InvalidRowHandle)
 		{
-			const int RowIndexInChunk = Hack_FindRowIndexInChunk(Context, Target);
-			if (RowIndexInChunk == INDEX_NONE)
+			CoreProvider->DestroyRelation(HierarchyRelationType, Target, CurrentParent);
+		}
+
+		// Create new parent relation (or orphan if Parent == InvalidRowHandle)
+		if (Parent != InvalidRowHandle)
+		{
+			CoreProvider->CreateRelation(HierarchyRelationType, Target, Parent);
+		}
+	}
+
+	void FTedsRelationBackedHierarchyAccess::SetUnresolvedParent(
+		ICoreProvider* CoreProvider, RowHandle Target, FMapKey ParentId, FName MappingDomain) const
+	{
+		check(CoreProvider != nullptr);
+
+		// Try to resolve immediately
+		RowHandle ParentRow = CoreProvider->LookupMappedRow(MappingDomain, ParentId);
+		if (CoreProvider->IsRowAvailable(ParentRow))
+		{
+			SetParentRow(CoreProvider, Target, ParentRow);
+			return;
+		}
+
+		// Store unresolved parent data as a column on the Target (child) row.
+		// The column auto-cleans when the Target row is destroyed.
+		// ResolveHierarchy picks these up each frame and attempts resolution.
+		const UScriptStruct* UnresolvedColumn = GetUnresolvedParentColumnType();
+		if (UnresolvedColumn)
+		{
+			CoreProvider->AddColumn(Target, UnresolvedColumn);
+			if (void* ColumnData = CoreProvider->GetColumnData(Target, UnresolvedColumn))
 			{
-				return InvalidRowHandle;
+				FEditorDataHierarchyUnresolvedParent_Template* Unresolved =
+					static_cast<FEditorDataHierarchyUnresolvedParent_Template*>(ColumnData);
+				Unresolved->ParentId = ParentId;
+				Unresolved->MappingDomain = MappingDomain;
 			}
-			const void* HierarchyDataComponentArrayStart = Context.GetColumn(HierarchyDataColumnType);
-			if (HierarchyDataComponentArrayStart == nullptr)
+		}
+	}
+
+	RowHandle FTedsRelationBackedHierarchyAccess::GetParentRow(const ICoreProvider* CoreProvider, RowHandle Target) const
+	{
+		check(CoreProvider != nullptr);
+		return CoreProvider->GetRelationObject(HierarchyRelationType, Target);
+	}
+
+	bool FTedsRelationBackedHierarchyAccess::HasChildren(const ICoreProvider& Context, RowHandle Row) const
+	{
+		return Context.HasRelationSubject(HierarchyRelationType, Row);
+	}
+
+	void FTedsRelationBackedHierarchyAccess::WalkDepthFirst(
+		const ICoreProvider& Context,
+		RowHandle Row,
+		TFunction<void(const ICoreProvider& Context, RowHandle Owner, RowHandle Target)> OnVisitedFn,
+		ICoreProvider::ETraversalOrder TraversalOrder) const
+	{
+		if (Row == InvalidRowHandle || !OnVisitedFn)
+		{
+			return;
+		}
+
+		// Match Legacy behavior: visit the root node as well as all descendants.
+		// Legacy's WalkDepthFirst visits root first (PreOrder) or last (PostOrder).
+		if (TraversalOrder == ICoreProvider::ETraversalOrder::PreOrder)
+		{
+			OnVisitedFn(Context, InvalidRowHandle /*no parent*/, Row);
+		}
+
+		Context.TraverseDescendants(HierarchyRelationType, Row,
+			[&Context, &OnVisitedFn](RowHandle Current, RowHandle Parent, int32 /*Depth*/)
 			{
-				return InvalidRowHandle;
+				OnVisitedFn(Context, Parent, Current);
+			},
+			TraversalOrder);
+
+		if (TraversalOrder == ICoreProvider::ETraversalOrder::PostOrder)
+		{
+			OnVisitedFn(Context, InvalidRowHandle /*no parent*/, Row);
+		}
+	}
+
+	bool FTedsRelationBackedHierarchyAccess::IterateChildren(
+		const ICoreProvider& Context,
+		RowHandle Row,
+		TFunctionRef<bool(const ICoreProvider& Context, RowHandle Child)> OnVisitedFn) const
+	{
+		if (!Context.IsRowAvailable(Row))
+		{
+			return false;
+		}
+
+		TArray<RowHandle> Children;
+		Context.GetRelationSubjects(HierarchyRelationType, Row, Children);
+		for (RowHandle Child : Children)
+		{
+			if (!OnVisitedFn(Context, Child))
+			{
+				return false;
 			}
-			const uint8* HierarchyDataComponentObjectAddress = static_cast<const uint8*>(HierarchyDataComponentArrayStart) + (RowIndexInChunk * HierarchyDataColumnType->GetStructureSize());
-			ParentHierarchyDataComponent = reinterpret_cast<const FEditorDataHierarchyData_Template*>(HierarchyDataComponentObjectAddress);
 		}
+		return true;
+	}
 
-		return ParentHierarchyDataComponent->Parent;
+	TFunction<RowHandle(const void*, const UScriptStruct*)> FTedsRelationBackedHierarchyAccess::CreateParentExtractionFunction() const
+	{
+		// Relations backend doesn't store parent data in a column on the participant row.
+		// Return a no-op function that always returns InvalidRowHandle.
+		return [](const void*, const UScriptStruct*) -> RowHandle
+		{
+			return InvalidRowHandle;
+		};
 	}
 }
diff --git a/Engine/Plugins/Experimental/EditorDataStorage/Source/TedsCore/Private/Hierarchy/EditorDataStorageHierarchyImplementation.h b/Engine/Plugins/Experimental/EditorDataStorage/Source/TedsCore/Private/Hierarchy/EditorDataStorageHierarchyImplementation.h
--- a/Engine/Plugins/Experimental/EditorDataStorage/Source/TedsCore/Private/Hierarchy/EditorDataStorageHierarchyImplementation.h
+++ b/Engine/Plugins/Experimental/EditorDataStorage/Source/TedsCore/Private/Hierarchy/EditorDataStorageHierarchyImplementation.h
@@ -8,11 +8,17 @@
 
 namespace UE::Editor::DataStorage
 {
+	enum class EHierarchyBackend : uint8
+	{
+		Legacy = 0,
+		Relations = 1,
+	};
+
 	class FTedsHierarchyAccessInterface;
 	class FTedsHierarchyRegistrar final
 	{
 	public:
-		
+
 		~FTedsHierarchyRegistrar();
 		/**
 		 * Registers a hierarchy type with TEDS
@@ -32,16 +38,18 @@ namespace UE::Editor::DataStorage
 		 * Iterate over all registered hierarchyes
 		 */
 		void ListHierarchyNames(TFunctionRef<void(const FName& HierarchyName)> Callback) const;
-		
+
 		/**
 		 * Try to resolve all child rows with unresolved parent columns
 		 */
 		void ResolveHierarchy(ICoreProvider& CoreProvider, FHierarchyHandle HierarchyHandle);
-		
+
 	private:
+		FHierarchyHandle RegisterHierarchyLegacy(ICoreProvider* InProvider, const FHierarchyRegistrationParams& Params);
+		FHierarchyHandle RegisterHierarchyRelations(ICoreProvider* InProvider, const FHierarchyRegistrationParams& Params);
 		void RegisterObservers(ICoreProvider& CoreProvider, const FHierarchyHandle& Handle);
 		void RegisterProcessors(ICoreProvider& CoreProvider, const FHierarchyHandle& Handle);
-		
+
 		struct FRegisteredHierarchy
 		{
 			FName Name;
@@ -52,52 +60,56 @@ namespace UE::Editor::DataStorage
 			const UScriptStruct* ParentChangedColumn;
 			TUniquePtr<FTedsHierarchyAccessInterface> AccessInterface;
 			QueryHandle UnresolvedChildRowsQueryHandle = InvalidQueryHandle;
+			EHierarchyBackend Backend = EHierarchyBackend::Legacy;
+			RelationTypeHandle HierarchyRelationType = InvalidRelationTypeHandle;
 		};
 
 		const FRegisteredHierarchy* FindRegisteredHierarchy(const FHierarchyHandle& Handle);
-	
+
 	private:
 		TArray<FRegisteredHierarchy> RegisteredHierarchies;
 	};
 
-	class FTedsHierarchyAccessInterface final
+	/**
+	 * Abstract interface for accessing hierarchy operations.
+	 * Concrete backends (adjacency-list and relations) are defined in the .cpp.
+	 *
+	 * The non-virtual context-dispatch overloads are generic and live here because
+	 * they delegate to the virtual ICoreProvider* overloads via context unwrapping.
+	 */
+	class FTedsHierarchyAccessInterface
 	{
 	public:
-		struct FConstructParams
-		{
-			const UScriptStruct* ChildTag;
-			const UScriptStruct* ParentTag;
-			const UScriptStruct* HierarchyData;
-			const UScriptStruct* UnresolvedParentColumn;
-			const UScriptStruct* ParentChangedColumn;
-		};
-		
-		explicit FTedsHierarchyAccessInterface(const FConstructParams& Params);
-		
-		const UScriptStruct* GetChildTagType() const;
-		const UScriptStruct* GetParentTagType() const;
-		const UScriptStruct* GetHierarchyDataColumnType() const;
-		const UScriptStruct* GetUnresolvedParentColumnType() const;
-		const UScriptStruct* GetParentChangedColumnType() const;
-
-
-		void SetParentRow(ICoreProvider* CoreProvider, RowHandle Target, RowHandle Parent) const;
-		void SetUnresolvedParent(ICoreProvider* CoreProvider, RowHandle Target, FMapKey ParentId, FName MappingDomain) const;
-		RowHandle GetParentRow(const ICoreProvider* CoreProvider, RowHandle Target) const;
-		
-		bool HasChildren(const ICoreProvider& Context, RowHandle Row) const;
-		void WalkDepthFirst(
+		virtual ~FTedsHierarchyAccessInterface() = default;
+
+		virtual const UScriptStruct* GetChildTagType() const = 0;
+		virtual const UScriptStruct* GetParentTagType() const = 0;
+		virtual const UScriptStruct* GetHierarchyDataColumnType() const = 0;
+		virtual const UScriptStruct* GetUnresolvedParentColumnType() const = 0;
+		virtual const UScriptStruct* GetParentChangedColumnType() const = 0;
+
+		virtual void SetParentRow(ICoreProvider* CoreProvider, RowHandle Target, RowHandle Parent) const = 0;
+		virtual void SetUnresolvedParent(ICoreProvider* CoreProvider, RowHandle Target, FMapKey ParentId, FName MappingDomain) const = 0;
+		virtual RowHandle GetParentRow(const ICoreProvider* CoreProvider, RowHandle Target) const = 0;
+
+		virtual bool HasChildren(const ICoreProvider& Context, RowHandle Row) const = 0;
+		virtual void WalkDepthFirst(
 			const ICoreProvider& Context,
 			RowHandle Row,
 			TFunction<void(const ICoreProvider& Context, RowHandle Owner, RowHandle Target)> OnVisitedFn,
-			ICoreProvider::ETraversalOrder TraversalOrder = ICoreProvider::ETraversalOrder::PreOrder) const;
-		bool IterateChildren(
+			ICoreProvider::ETraversalOrder TraversalOrder = ICoreProvider::ETraversalOrder::PreOrder) const = 0;
+
+		virtual bool IterateChildren(
 			const ICoreProvider& Context,
 			RowHandle Row,
-			TFunctionRef<bool(const ICoreProvider& Context, RowHandle Child)> OnVisitedFn) const;
-		
-		TFunction<RowHandle(const void*, const UScriptStruct*)> CreateParentExtractionFunction() const;
-		
+			TFunctionRef<bool(const ICoreProvider& Context, RowHandle Child)> OnVisitedFn) const
+		{
+			return false;
+		}
+
+		virtual TFunction<RowHandle(const void*, const UScriptStruct*)> CreateParentExtractionFunction() const = 0;
+
+		// Non-virtual context-dispatch overloads — generic for all backends.
 		void SetParentRow(IQueryContext& Context, RowHandle Target, RowHandle Parent) const;
 		void SetUnresolvedParent(IQueryContext& Context, RowHandle Target, FMapKey ParentId, FName MappingDomain) const;
 		RowHandle GetParentRow(const IQueryContext& Context, RowHandle Target) const;
@@ -108,16 +120,13 @@ namespace UE::Editor::DataStorage
 		void SetUnresolvedParent(IDirectQueryContext& Context, RowHandle Target, FMapKey ParentId, FName MappingDomain) const;
 		RowHandle GetParentRow(const IDirectQueryContext& Context, RowHandle Target) const;
 
-	private:
-
+	protected:
+		// Generic context-dispatch helpers shared by all backends.
 		void SetParent(ICommonQueryContext& Context, RowHandle Target, RowHandle Parent) const;
 		void SetUnresolvedParent(ICommonQueryContext& Context, RowHandle Target, FMapKey ParentId, FName MappingDomain) const;
-		RowHandle GetParent(const ICommonQueryContext& Context, RowHandle Target) const;
-		
-		const UScriptStruct* ChildTag;
-		const UScriptStruct* ParentTag;
-		const UScriptStruct* HierarchyDataColumnType;
-		const UScriptStruct* UnresolvedParentColumnType;
-		const UScriptStruct* ParentChangedColumn;
+
+		// Context-aware parent lookup. Default returns InvalidRowHandle; override in column-backed backends.
+		virtual RowHandle GetParent(const ICommonQueryContext& Context, RowHandle Target) const;
 	};
-}
+
+} // namespace UE::Editor::DataStorage
diff --git a/Engine/Plugins/Experimental/EditorDataStorage/Source/TedsCore/Private/Processors/TypedElementProcessorAdaptors.cpp b/Engine/Plugins/Experimental/EditorDataStorage/Source/TedsCore/Private/Processors/TypedElementProcessorAdaptors.cpp
--- a/Engine/Plugins/Experimental/EditorDataStorage/Source/TedsCore/Private/Processors/TypedElementProcessorAdaptors.cpp
+++ b/Engine/Plugins/Experimental/EditorDataStorage/Source/TedsCore/Private/Processors/TypedElementProcessorAdaptors.cpp
@@ -11,6 +11,7 @@
 #include "Queries/ProcessorContextEnvironment.h"
 #include "Processors/QueryCapabilityImplementationForwarder.h"
 #include "StructUtils/StructTypeBitSet.h"
+#include "Relations/EditorDataStorageRelationsImplementation.h"
 #include "TypedElementDatabase.h"
 #include "TypedElementDatabaseEnvironment.h"
 #include "TypedElementUtils.h"
@@ -20,36 +21,262 @@
 
 namespace UE::Editor::DataStorage
 {
+	namespace Private
+	{
+		/** Check whether a query has any relation or hierarchy conditions that need per-entity filtering. */
+		static bool HasRelationConditions(const FQueryDescription& Description)
+		{
+			return !Description.RelationConditions.IsEmpty()
+				|| !Description.DescendantOfConditions.IsEmpty()
+				|| !Description.AncestorOfConditions.IsEmpty();
+		}
+
+		/** Evaluate all relation/hierarchy conditions for a single row. Returns true if the row passes all conditions. */
+		static bool PassesRelationConditions(RowHandle Row, const FQueryDescription& Description,
+			FMassEntityManager& EntityManager, const FTedsRelationAdapter& RelationAdapter,
+			const FTedsHierarchicalRelationManager* HierarchyManager)
+		{
+			// Check relation conditions (IsSubjectOf / IsObjectOf)
+			for (const FQueryDescription::FRelationCondition& Condition : Description.RelationConditions)
+			{
+				if (Condition.RelationType == InvalidRelationTypeHandle)
+				{
+					continue;
+				}
+
+				if (Condition.Role == FQueryDescription::ERelationRole::Subject)
+				{
+					if (Condition.TargetRow != InvalidRowHandle)
+					{
+						// IsSubjectOf(Type, TargetObject) — check if this row is subject with that specific object
+						if (!RelationAdapter.HasMassRelationData(Condition.RelationType) ||
+								!EntityManager.GetRelationManager().IsSubjectOfRelation(
+									RelationAdapter.GetMassHandle(Condition.RelationType),
+									FMassEntityHandle::FromNumber(Row),
+									FMassEntityHandle::FromNumber(Condition.TargetRow)))
+						{
+							return false;
+						}
+					}
+					else
+					{
+						// IsSubjectOf(Type) — check if this row is a subject of any relation of this type
+						if (!RelationAdapter.HasMassRelationData(Condition.RelationType) ||
+							!EntityManager.GetRelationManager().HasRelationObjects(
+								RelationAdapter.GetMassHandle(Condition.RelationType),
+								FMassEntityHandle::FromNumber(Row)))
+						{
+							return false;
+						}
+					}
+				}
+				else // ERelationRole::Object
+				{
+					if (Condition.TargetRow != InvalidRowHandle)
+					{
+						// IsObjectOf(Type, TargetSubject) — check if this row is object with that specific subject
+						if (!RelationAdapter.HasMassRelationData(Condition.RelationType) ||
+								!EntityManager.GetRelationManager().IsSubjectOfRelation(
+									RelationAdapter.GetMassHandle(Condition.RelationType),
+									FMassEntityHandle::FromNumber(Condition.TargetRow),
+									FMassEntityHandle::FromNumber(Row)))
+						{
+							return false;
+						}
+					}
+					else
+					{
+						// IsObjectOf(Type) — check if this row is an object of any relation of this type
+						if (!RelationAdapter.HasMassRelationData(Condition.RelationType) ||
+							!EntityManager.GetRelationManager().HasRelationSubjects(
+								RelationAdapter.GetMassHandle(Condition.RelationType),
+								FMassEntityHandle::FromNumber(Row)))
+						{
+							return false;
+						}
+					}
+				}
+			}
+
+			// Check descendant-of conditions
+			if (!Description.DescendantOfConditions.IsEmpty())
+			{
+				if (!HierarchyManager)
+				{
+					return false;
+				}
+				for (const FQueryDescription::FHierarchyCondition& Condition : Description.DescendantOfConditions)
+				{
+					if (Condition.RelationType == InvalidRelationTypeHandle || Condition.ReferenceRow == InvalidRowHandle)
+					{
+						continue;
+					}
+					if (Condition.bIncludeSelf && Row == Condition.ReferenceRow)
+					{
+						continue; // Passes — self-inclusion
+					}
+					if (!HierarchyManager->IsDescendantOf(EntityManager, RelationAdapter, Condition.RelationType, Row, Condition.ReferenceRow))
+					{
+						return false;
+					}
+				}
+			}
+
+			// Check ancestor-of conditions
+			if (!Description.AncestorOfConditions.IsEmpty())
+			{
+				if (!HierarchyManager)
+				{
+					return false;
+				}
+				for (const FQueryDescription::FHierarchyCondition& Condition : Description.AncestorOfConditions)
+				{
+					if (Condition.RelationType == InvalidRelationTypeHandle || Condition.ReferenceRow == InvalidRowHandle)
+					{
+						continue;
+					}
+					if (Condition.bIncludeSelf && Row == Condition.ReferenceRow)
+					{
+						continue; // Passes — self-inclusion
+					}
+					// IsAncestorOf(Type, Descendant) means: is Descendant a descendant of Row?
+					if (!HierarchyManager->IsDescendantOf(EntityManager, RelationAdapter, Condition.RelationType, Condition.ReferenceRow, Row))
+					{
+						return false;
+					}
+				}
+			}
+
+
+
+			return true;
+		}
+
+		/**
+		 * Evaluate relation conditions for all entities in the current chunk and invoke the callback
+		 * once per contiguous sub-range of passing entities. This performs entity-level relation filtering
+		 * entirely within the TEDS adaptor layer, without requiring Mass-level EntityCondition support.
+		 */
+		static void ForEachRelationFilteredSubrange(
+			FMassExecutionContext& Context,
+			const FQueryDescription& Description,
+			FMassEntityManager& EntityManager,
+			const FTedsRelationAdapter& RelationAdapter,
+			const FTedsHierarchicalRelationManager* HierarchyManager,
+			TFunctionRef<void(int32 SubRangeStart, int32 SubRangeLength)> Callback)
+		{
+			TConstArrayView<FMassEntityHandle> Entities = Context.GetEntities();
+			const int32 Num = Entities.Num();
+
+			FMemMark Mark(FMemStack::Get());
+			TArray<bool, TMemStackAllocator<>> Passes;
+			Passes.SetNumUninitialized(Num);
+			for (int32 i = 0; i < Num; ++i)
+			{
+				Passes[i] = PassesRelationConditions(
+					Entities[i].AsNumber(), Description, EntityManager, RelationAdapter, HierarchyManager);
+			}
+
+			int32 RangeStart = INDEX_NONE;
+			for (int32 i = 0; i <= Num; ++i)
+			{
+				const bool bPasses = (i < Num) && Passes[i];
+				if (bPasses && RangeStart == INDEX_NONE)
+				{
+					RangeStart = i;
+				}
+				else if (!bPasses && RangeStart != INDEX_NONE)
+				{
+					Callback(RangeStart, i - RangeStart);
+					RangeStart = INDEX_NONE;
+				}
+			}
+		}
+	} // namespace Private
+
 	namespace Processors::Private
 	{
+		/**
+		 * Deferred command to change a row's relation object.
+		 * Removes the existing relation (if any) then optionally creates a new one.
+		 * Executes after the processor chunk loop, under ICoreProvider's normal write path.
+		 */
+		struct FSetRelationObjectCommand
+		{
+			RelationTypeHandle Type = InvalidRelationTypeHandle;
+			RowHandle Target        = InvalidRowHandle;
+			RowHandle NewObject     = InvalidRowHandle;   // InvalidRowHandle = remove only
+			ICoreProvider* CoreProvider = nullptr;
+
+			void operator()()
+			{
+				if (!CoreProvider || !CoreProvider->IsRowAssigned(Target))
+				{
+					return;
+				}
+
+				// Remove existing relation, if any
+				const RowHandle OldObject = CoreProvider->GetRelationObject(Type, Target);
+				if (OldObject != InvalidRowHandle)
+				{
+					CoreProvider->DestroyRelation(Type, Target, OldObject);
+				}
+
+				// Create new relation, if requested
+				if (NewObject != InvalidRowHandle && CoreProvider->IsRowAssigned(NewObject))
+				{
+					CoreProvider->CreateRelation(Type, Target, NewObject);
+				}
+			}
+		};
+
 		struct FMassContextCommon
 		{
 			FMassExecutionContext& Context;
 			FEnvironment& Environment;
-	
+			int32 SubRangeStart = 0;
+			int32 SubRangeLength = -1;
+
 			FMassContextCommon(FMassExecutionContext& InContext, FEnvironment& InEnvironment)
 				: Context(InContext)
 				, Environment(InEnvironment)
 			{}
-	
+
+			void SetSubRange(int32 InStart, int32 InLength)
+			{
+				SubRangeStart = InStart;
+				SubRangeLength = InLength;
+			}
+
+			void ClearSubRange()
+			{
+				SubRangeStart = 0;
+				SubRangeLength = -1;
+			}
+
+			bool HasSubRange() const { return SubRangeLength >= 0; }
+
 			uint32 GetRowCount() const
 			{
-				return Context.GetNumEntities();
+				return HasSubRange() ? static_cast<uint32>(SubRangeLength) : Context.GetNumEntities();
 			}
 
 			TConstArrayView<RowHandle> GetRowHandles() const
 			{
-				return MassEntitiesToRowsConversion(Context.GetEntities());
+				TConstArrayView<RowHandle> All = MassEntitiesToRowsConversion(Context.GetEntities());
+				return HasSubRange() ? All.Mid(SubRangeStart, SubRangeLength) : All;
 			}
-	
+
 			const void* GetColumn(const UScriptStruct* ColumnType) const
 			{
-				return Context.GetFragmentView(ColumnType).GetData();
+				const uint8* Base = reinterpret_cast<const uint8*>(Context.GetFragmentView(ColumnType).GetData());
+				return HasSubRange() ? Base + SubRangeStart * ColumnType->GetStructureSize() : Base;
 			}
 
 			void* GetMutableColumn(const UScriptStruct* ColumnType)
 			{
-				return Context.GetMutableFragmentView(ColumnType).GetData();
+				uint8* Base = reinterpret_cast<uint8*>(Context.GetMutableFragmentView(ColumnType).GetData());
+				return HasSubRange() ? Base + SubRangeStart * ColumnType->GetStructureSize() : Base;
 			}
 
 			void GetColumns(TArrayView<char*> RetrievedAddresses,
@@ -63,7 +290,7 @@ namespace UE::Editor::DataStorage
 
 				GetColumnsUnguarded(ColumnTypes.Num(), RetrievedAddresses.GetData(), ColumnTypes.GetData(), AccessTypes.GetData());
 			}
-	
+
 			void GetColumnsUnguarded(
 				int32 TypeCount,
 				char** RetrievedAddresses,
@@ -73,9 +300,16 @@ namespace UE::Editor::DataStorage
 				for (int32 Index = 0; Index < TypeCount; ++Index)
 				{
 					checkf(ColumnTypes->IsValid(), TEXT("Attempting to retrieve a column that is not available."));
-					*RetrievedAddresses = *AccessTypes == EQueryAccessType::ReadWrite
-						? reinterpret_cast<char*>(Context.GetMutableFragmentView(ColumnTypes->Get()).GetData())
-						: const_cast<char*>(reinterpret_cast<const char*>(Context.GetFragmentView(ColumnTypes->Get()).GetData()));
+					if (*AccessTypes == EQueryAccessType::ReadWrite)
+					{
+						uint8* Base = reinterpret_cast<uint8*>(Context.GetMutableFragmentView(ColumnTypes->Get()).GetData());
+						*RetrievedAddresses = reinterpret_cast<char*>(HasSubRange() ? Base + SubRangeStart * ColumnTypes->Get()->GetStructureSize() : Base);
+					}
+					else
+					{
+						const uint8* Base = reinterpret_cast<const uint8*>(Context.GetFragmentView(ColumnTypes->Get()).GetData());
+						*RetrievedAddresses = const_cast<char*>(reinterpret_cast<const char*>(HasSubRange() ? Base + SubRangeStart * ColumnTypes->Get()->GetStructureSize() : Base));
+					}
 
 					++RetrievedAddresses;
 					++ColumnTypes;
@@ -1107,6 +1341,15 @@ namespace UE::Editor::DataStorage
 			bool HasColumn(RowHandle Row, const UScriptStruct* ColumnType) const override { return Implementation.HasColumn(Row, ColumnType); }
 			float GetDeltaTimeSeconds() const override { return Implementation.GetDeltaTimeSeconds(); }
 			virtual void SetParentRow(RowHandle Target, RowHandle Parent) override { if (HierarchyAccessInterface) { HierarchyAccessInterface->SetParentRow(*static_cast<IDirectQueryContext*>(this), Target, Parent); } }
+			virtual void SetRelation(RelationTypeHandle Type, RowHandle Target, RowHandle NewObject) override
+			{
+				FSetRelationObjectCommand Command;
+				Command.Type = Type;
+				Command.Target = Target;
+				Command.NewObject = NewObject;
+				Command.CoreProvider = GetMutableDataStorageFeature<ICoreProvider>(StorageFeatureName);
+				ICommonQueryContext::PushCommand(MoveTemp(Command));
+			}
 			virtual void SetUnresolvedParent(RowHandle Target, FMapKey ParentId, FName MappingDomain) override
 			{
 				if (HierarchyAccessInterface)
@@ -1172,6 +1415,15 @@ namespace UE::Editor::DataStorage
 			void RemoveColumns(TConstArrayView<RowHandle> Rows, TConstArrayView<const UScriptStruct*> ColumnTypes) override { Implementation.RemoveColumns(Rows, ColumnTypes); }
 			float GetDeltaTimeSeconds() const override { return Implementation.GetDeltaTimeSeconds(); }
 			virtual void SetParentRow(RowHandle Target, RowHandle Parent) override { Implementation.SetParent(*static_cast<ISubqueryContext*>(this), Target, Parent); }
+			virtual void SetRelation(RelationTypeHandle Type, RowHandle Target, RowHandle NewObject) override
+			{
+				FSetRelationObjectCommand Command;
+				Command.Type = Type;
+				Command.Target = Target;
+				Command.NewObject = NewObject;
+				Command.CoreProvider = GetMutableDataStorageFeature<ICoreProvider>(StorageFeatureName);
+				ICommonQueryContext::PushCommand(MoveTemp(Command));
+			}
 			virtual void SetUnresolvedParent(RowHandle Target, FMapKey ParentId, FName MappingDomain) override { Implementation.SetUnresolvedParent(*static_cast<ISubqueryContext*>(this), Target, ParentId, MappingDomain); }
 			virtual RowHandle GetParentRow(RowHandle Target) const override { return Implementation.GetParent(*static_cast<const ISubqueryContext*>(this), Target); }
 			void PushCommand(void (*CommandFunction)(void*), void* CommandData) override { return Implementation.PushCommand(CommandFunction, CommandData); }
@@ -1331,6 +1583,15 @@ namespace UE::Editor::DataStorage
 			void RemoveColumns(TConstArrayView<RowHandle> Rows, TConstArrayView<const UScriptStruct*> ColumnTypes) override { Implementation.RemoveColumns(Rows, ColumnTypes); }
 			float GetDeltaTimeSeconds() const override { return Implementation.GetDeltaTimeSeconds(); }
 			virtual void SetParentRow(RowHandle Target, RowHandle Parent) override { Implementation.SetParent(*static_cast<IQueryContext*>(this), Target, Parent); }
+			virtual void SetRelation(RelationTypeHandle Type, RowHandle Target, RowHandle NewObject) override
+			{
+				FSetRelationObjectCommand Command;
+				Command.Type = Type;
+				Command.Target = Target;
+				Command.NewObject = NewObject;
+				Command.CoreProvider = GetMutableDataStorageFeature<ICoreProvider>(StorageFeatureName);
+				ICommonQueryContext::PushCommand(MoveTemp(Command));
+			}
 			virtual void SetUnresolvedParent(RowHandle Target, FMapKey ParentId, FName MappingDomain) override { Implementation.SetUnresolvedParent(*static_cast<IQueryContext*>(this), Target, ParentId, MappingDomain); }
 			virtual RowHandle GetParentRow(RowHandle Target) const override { return Implementation.GetParent(*static_cast<const IQueryContext*>(this), Target); }
 			void PushCommand(void (*CommandFunction)(void*), void* Context) override { return Implementation.PushCommand(CommandFunction, Context); }
@@ -1343,6 +1604,9 @@ namespace UE::Editor::DataStorage
 			FQueryResult RunSubquery(int32 SubqueryIndex, UE::Editor::DataStorage::SubqueryCallbackRef Callback) override { return Implementation.RunSubquery(SubqueryIndex, Callback); }
 			FQueryResult RunSubquery(int32 SubqueryIndex, RowHandle Row, UE::Editor::DataStorage::SubqueryCallbackRef Callback) override { return Implementation.RunSubquery(SubqueryIndex, Row, Callback); }
 
+			void SetSubRange(int32 InStart, int32 InLength) { Implementation.SetSubRange(InStart, InLength); }
+			void ClearSubRange() { Implementation.ClearSubRange(); }
+
 		protected:
 			void* EmplaceObjectInScratch(const FEmplaceObjectParams& Params) override { return Implementation.EmplaceObjectInScratch(Params.ObjectSize, Params.Alignment, Params.Construct, Params.Destroy, Params.SourceObject); }
 
@@ -1580,7 +1844,7 @@ bool FTypedElementQueryProcessorData::PrepareCachedDependenciesOnQuery(
 UE::Editor::DataStorage::FQueryResult FTypedElementQueryProcessorData::Execute(
 	UE::Editor::DataStorage::DirectQueryCallbackRef& Callback,
 	UE::Editor::DataStorage::FQueryDescription& Description,
-	FMassEntityQuery& NativeQuery, 
+	FMassEntityQuery& NativeQuery,
 	FMassEntityManager& EntityManager,
 	UE::Editor::DataStorage::FEnvironment& Environment,
 	UE::Editor::DataStorage::EDirectQueryExecutionFlags ExecutionFlags)
@@ -1589,23 +1853,45 @@ UE::Editor::DataStorage::FQueryResult FTypedElementQueryProcessorData::Execute(
 
 	FQueryResult Result;
 	Result.Completed = FQueryResult::ECompletion::Fully;
-	
+
 	if (EnumHasAnyFlags(ExecutionFlags, EDirectQueryExecutionFlags::AllowBoundQueries) || !Description.Callback.Function)
 	{
-		FMassExecutionContext Context(EntityManager);
-		Context.SetDeferredCommandBuffer(Environment.GetSharedMassCommandBuffer());
-		Context.SetFlushDeferredCommands(false);
-		auto ExecuteFunction = [&Result, &Callback, &Description, &Environment](FMassExecutionContext& Context)
+		const bool bHasRelationConditions = UE::Editor::DataStorage::Private::HasRelationConditions(Description);
+		const FTedsRelationAdapter* RelationAdapter = bHasRelationConditions ? &Environment.GetRelationAdapter() : nullptr;
+		const FTedsHierarchicalRelationManager* HierarchyManager = bHasRelationConditions ? &Environment.GetHierarchicalRelationManager() : nullptr;
+
+		auto ExecuteFunction = [&Result, &Callback, &Description, &Environment, &EntityManager,
+			bHasRelationConditions, RelationAdapter, HierarchyManager](FMassExecutionContext& Context)
 			{
 				using namespace Processors::Private;
-				// No need to cache any subsystem dependencies as these are not accessible from a direct query.
 				FMassDirectContextForwarder QueryContext(Context, Environment, Description);
-				Callback(Description, QueryContext);
-				Result.Count += Context.GetNumEntities();
+
+				if (bHasRelationConditions)
+				{
+					UE::Editor::DataStorage::Private::ForEachRelationFilteredSubrange(
+						Context, Description, EntityManager, *RelationAdapter, HierarchyManager,
+						[&](int32 SubStart, int32 SubLength)
+						{
+							QueryContext.Implementation.SetSubRange(SubStart, SubLength);
+							Callback(Description, QueryContext);
+							Result.Count += SubLength;
+						});
+					QueryContext.Implementation.ClearSubRange();
+				}
+				else
+				{
+					Callback(Description, QueryContext);
+					Result.Count += Context.GetNumEntities();
+				}
 			};
+
+		FMassExecutionContext Context(EntityManager);
+		Context.SetDeferredCommandBuffer(Environment.GetSharedMassCommandBuffer());
+		Context.SetFlushDeferredCommands(false);
+
 		if (EnumHasAnyFlags(ExecutionFlags, EDirectQueryExecutionFlags::ParallelizeChunks))
 		{
-			FMassEntityQuery::EParallelExecutionFlags Mode =
+			const FMassEntityQuery::EParallelExecutionFlags Mode =
 				EnumHasAnyFlags(ExecutionFlags, EDirectQueryExecutionFlags::AutoBalanceParallelChunkProcessing)
 				? FMassEntityQuery::EParallelExecutionFlags::AutoBalance
 				: FMassEntityQuery::EParallelExecutionFlags::Default;
@@ -1701,21 +1987,43 @@ UE::Editor::DataStorage::FQueryResult FTypedElementQueryProcessorData::Execute(
 
 	if (EnumHasAnyFlags(ExecutionFlags, ERunQueryFlags::AllowBoundQueries) || !Description.Callback.Function)
 	{
-		FMassExecutionContext Context(EntityManager);
-		Context.SetDeferredCommandBuffer(Environment.GetSharedMassCommandBuffer());
-		Context.SetFlushDeferredCommands(false);
-		auto ExecuteFunction = [&Result, &Callback, &Environment](FMassExecutionContext& Context)
+		const bool bHasRelationConditions = UE::Editor::DataStorage::Private::HasRelationConditions(Description);
+		const FTedsRelationAdapter* RelationAdapter = bHasRelationConditions ? &Environment.GetRelationAdapter() : nullptr;
+		const FTedsHierarchicalRelationManager* HierarchyManager = bHasRelationConditions ? &Environment.GetHierarchicalRelationManager() : nullptr;
+
+		auto ExecuteFunction = [&Result, &Callback, &Environment, &Description, &EntityManager,
+			bHasRelationConditions, RelationAdapter, HierarchyManager](FMassExecutionContext& Context)
 			{
 				using namespace Processors::Private;
 				FConstProcessorContextEnvironment MassContextEnvironment(Environment, Context);
 				QueryContext_ConstProcessorApi QueryContext(MassContextEnvironment);
-				// This assumes that there already has been a check that the provided function is compatible with the query.
-				Callback.Call(QueryContext, MassContextEnvironment);
-				Result.Count += Context.GetNumEntities();
+
+				if (bHasRelationConditions)
+				{
+					UE::Editor::DataStorage::Private::ForEachRelationFilteredSubrange(
+						Context, Description, EntityManager, *RelationAdapter, HierarchyManager,
+						[&](int32 SubStart, int32 SubLength)
+						{
+							MassContextEnvironment.SetSubRange(SubStart, SubLength);
+							Callback.Call(QueryContext, MassContextEnvironment);
+							Result.Count += SubLength;
+						});
+					MassContextEnvironment.ClearSubRange();
+				}
+				else
+				{
+					Callback.Call(QueryContext, MassContextEnvironment);
+					Result.Count += Context.GetNumEntities();
+				}
 			};
+
+		FMassExecutionContext Context(EntityManager);
+		Context.SetDeferredCommandBuffer(Environment.GetSharedMassCommandBuffer());
+		Context.SetFlushDeferredCommands(false);
+
 		if (EnumHasAnyFlags(ExecutionFlags, ERunQueryFlags::ParallelizeChunks))
 		{
-			FMassEntityQuery::EParallelExecutionFlags Mode =
+			const FMassEntityQuery::EParallelExecutionFlags Mode =
 				EnumHasAnyFlags(ExecutionFlags, ERunQueryFlags::AutoBalanceParallelChunkProcessing)
 				? FMassEntityQuery::EParallelExecutionFlags::AutoBalance
 				: FMassEntityQuery::EParallelExecutionFlags::Default;
@@ -1820,28 +2128,49 @@ UE::Editor::DataStorage::FQueryResult FTypedElementQueryProcessorData::Execute(
 void FTypedElementQueryProcessorData::Execute(FMassEntityManager& EntityManager, FMassExecutionContext& Context)
 {
 	using namespace UE::Editor::DataStorage;
-	
+
 	FExtendedQuery* StoredQuery = QueryStore->GetMutable(ParentQuery);
 
 	checkf(StoredQuery, TEXT("A query callback was registered for execution without an associated query. Processor: [%s]"), *GetProcessorName());
-	
+
 	FQueryDescription& Description = StoredQuery->Description;
-	auto ExecuteFunction = [this, &Description](FMassExecutionContext& Context)
+
+	const bool bHasRelationConditions = UE::Editor::DataStorage::Private::HasRelationConditions(Description);
+	const FTedsRelationAdapter* RelationAdapter = bHasRelationConditions ? &Environment->GetRelationAdapter() : nullptr;
+	const FTedsHierarchicalRelationManager* HierarchyManager = bHasRelationConditions ? &Environment->GetHierarchicalRelationManager() : nullptr;
+
+	auto ExecuteFunction = [this, &Description, &EntityManager,
+		bHasRelationConditions, RelationAdapter, HierarchyManager](FMassExecutionContext& Context)
 		{
 			if (PrepareCachedDependenciesOnQuery(Description, Context))
 			{
 				Processors::Private::FMassContextForwarder QueryContext(Description, Context, *QueryStore, *Environment);
-				Description.Callback.Function(Description, QueryContext);
+
+				if (bHasRelationConditions)
+				{
+					UE::Editor::DataStorage::Private::ForEachRelationFilteredSubrange(
+						Context, Description, EntityManager, *RelationAdapter, HierarchyManager,
+						[&](int32 SubStart, int32 SubLength)
+						{
+							QueryContext.SetSubRange(SubStart, SubLength);
+							Description.Callback.Function(Description, QueryContext);
+						});
+					QueryContext.ClearSubRange();
+				}
+				else
+				{
+					Description.Callback.Function(Description, QueryContext);
+				}
 			}
 		};
 
-	if (StoredQuery->Description.Callback.ExecutionMode != EExecutionMode::ThreadedChunks)
+	if (StoredQuery->Description.Callback.ExecutionMode == EExecutionMode::ThreadedChunks)
 	{
-		NativeQuery.ForEachEntityChunk(Context, ExecuteFunction);
+		NativeQuery.ParallelForEachEntityChunk(Context, ExecuteFunction);
 	}
 	else
 	{
-		NativeQuery.ParallelForEachEntityChunk(Context, ExecuteFunction);
+		NativeQuery.ForEachEntityChunk(Context, ExecuteFunction);
 	}
 }
 
diff --git a/Engine/Plugins/Experimental/EditorDataStorage/Source/TedsCore/Private/Queries/ProcessorContextEnvironment.cpp b/Engine/Plugins/Experimental/EditorDataStorage/Source/TedsCore/Private/Queries/ProcessorContextEnvironment.cpp
--- a/Engine/Plugins/Experimental/EditorDataStorage/Source/TedsCore/Private/Queries/ProcessorContextEnvironment.cpp
+++ b/Engine/Plugins/Experimental/EditorDataStorage/Source/TedsCore/Private/Queries/ProcessorContextEnvironment.cpp
@@ -29,6 +29,10 @@ namespace UE::Editor::DataStorage::Queries
 	FRowHandleArrayView FProcessorContextEnvironment::GetBatchRows() const
 	{
 		TConstArrayView<RowHandle> Rows = MassEntitiesToRowsConversion(MassContext.GetEntities());
+		if (HasSubRange())
+		{
+			Rows = Rows.Mid(SubRangeStart, SubRangeLength);
+		}
 		return FRowHandleArrayView(Rows, FRowHandleArrayView::EFlags::IsUnique);
 	}
 
@@ -43,7 +47,8 @@ namespace UE::Editor::DataStorage::Queries
 		{
 			checkfSlow(UE::Mass::IsA<FMassFragment>(ColumnType),
 				TEXT("Unable to retrieve const column data for a column type without data: %s"), *ColumnType->GetFullName());
-			*ColumnsDataIt++ = MassContext.GetFragmentView(ColumnType).GetData();
+			const uint8* Base = reinterpret_cast<const uint8*>(MassContext.GetFragmentView(ColumnType).GetData());
+			*ColumnsDataIt++ = HasSubRange() ? Base + SubRangeStart * ColumnType->GetStructureSize() : Base;
 		}
 	}
 
@@ -58,7 +63,8 @@ namespace UE::Editor::DataStorage::Queries
 		{
 			checkfSlow(UE::Mass::IsA<FMassFragment>(ColumnType),
 				TEXT("Unable to retrieve mutable column data for a column type without data: %s"), *ColumnType->GetFullName());
-			*ColumnsDataIt++ = MassContext.GetMutableFragmentView(ColumnType).GetData();
+			uint8* Base = reinterpret_cast<uint8*>(MassContext.GetMutableFragmentView(ColumnType).GetData());
+			*ColumnsDataIt++ = HasSubRange() ? Base + SubRangeStart * ColumnType->GetStructureSize() : Base;
 		}
 	}
 
@@ -93,6 +99,10 @@ namespace UE::Editor::DataStorage::Queries
 	FRowHandleArrayView FConstProcessorContextEnvironment::GetBatchRows() const
 	{
 		TConstArrayView<RowHandle> Rows = MassEntitiesToRowsConversion(MassContext.GetEntities());
+		if (HasSubRange())
+		{
+			Rows = Rows.Mid(SubRangeStart, SubRangeLength);
+		}
 		return FRowHandleArrayView(Rows, FRowHandleArrayView::EFlags::IsUnique);
 	}
 
@@ -107,7 +117,8 @@ namespace UE::Editor::DataStorage::Queries
 		{
 			checkfSlow(UE::Mass::IsA<FMassFragment>(ColumnType),
 				TEXT("Unable to retrieve const column data for a column type without data: %s"), *ColumnType->GetFullName());
-			*ColumnsDataIt++ = MassContext.GetFragmentView(ColumnType).GetData();
+			const uint8* Base = reinterpret_cast<const uint8*>(MassContext.GetFragmentView(ColumnType).GetData());
+			*ColumnsDataIt++ = HasSubRange() ? Base + SubRangeStart * ColumnType->GetStructureSize() : Base;
 		}
 	}
 
diff --git a/Engine/Plugins/Experimental/EditorDataStorage/Source/TedsCore/Private/Queries/ProcessorContextEnvironment.h b/Engine/Plugins/Experimental/EditorDataStorage/Source/TedsCore/Private/Queries/ProcessorContextEnvironment.h
--- a/Engine/Plugins/Experimental/EditorDataStorage/Source/TedsCore/Private/Queries/ProcessorContextEnvironment.h
+++ b/Engine/Plugins/Experimental/EditorDataStorage/Source/TedsCore/Private/Queries/ProcessorContextEnvironment.h
@@ -20,15 +20,21 @@ namespace UE::Editor::DataStorage::Queries
 
 		virtual FEnvironment* GetEnvironment() override;
 		virtual const FEnvironment* GetEnvironment() const override;
-		
+
 		virtual FRowHandleArrayView GetBatchRows() const override;
 		virtual void GetConstColumns(TArrayView<const void*> ColumnsData, TConstArrayView<const UScriptStruct*> ColumnTypes) const override;
 		virtual void GetMutableColumns(TArrayView<void*> ColumnsData, TConstArrayView<const UScriptStruct*> ColumnTypes) override;
 		virtual void SetCurrentRow(RowHandle Row) override;
-		
+
+		void SetSubRange(int32 InStart, int32 InLength) { SubRangeStart = InStart; SubRangeLength = InLength; }
+		void ClearSubRange() { SubRangeStart = 0; SubRangeLength = -1; }
+		bool HasSubRange() const { return SubRangeLength >= 0; }
+
 		FEnvironment& Environment;
 		FMassExecutionContext& MassContext;
 		RowHandle CurrentRow = InvalidRowHandle;
+		int32 SubRangeStart = 0;
+		int32 SubRangeLength = -1;
 	};
 
 	struct FConstProcessorContextEnvironment final : IContextEnvironment
@@ -43,9 +49,15 @@ namespace UE::Editor::DataStorage::Queries
 		virtual void GetMutableColumns(TArrayView<void*> ColumnsData, TConstArrayView<const UScriptStruct*> ColumnTypes) override;
 		virtual void SetCurrentRow(RowHandle Row) override;
 
+		void SetSubRange(int32 InStart, int32 InLength) { SubRangeStart = InStart; SubRangeLength = InLength; }
+		void ClearSubRange() { SubRangeStart = 0; SubRangeLength = -1; }
+		bool HasSubRange() const { return SubRangeLength >= 0; }
+
 		const FEnvironment& Environment;
 		const FMassExecutionContext& MassContext;
 		RowHandle CurrentRow = InvalidRowHandle;
+		int32 SubRangeStart = 0;
+		int32 SubRangeLength = -1;
 	};
 
 	using QueryContext_ProcessorApi = TQueryContextImpl<false, FProcessorContextEnvironment,
diff --git a/Engine/Plugins/Experimental/EditorDataStorage/Source/TedsCore/Private/Queries/QueryContextCapabilities/RowBatch.cpp b/Engine/Plugins/Experimental/EditorDataStorage/Source/TedsCore/Private/Queries/QueryContextCapabilities/RowBatch.cpp
--- a/Engine/Plugins/Experimental/EditorDataStorage/Source/TedsCore/Private/Queries/QueryContextCapabilities/RowBatch.cpp
+++ b/Engine/Plugins/Experimental/EditorDataStorage/Source/TedsCore/Private/Queries/QueryContextCapabilities/RowBatch.cpp
@@ -49,7 +49,8 @@ namespace UE::Editor::DataStorage::Queries
 	{
 		checkf(UE::Mass::IsA<FMassFragment>(ColumnType), TEXT("Trying to retrieve a data pointer for column '%s' which doesn't have data."),
 			ColumnType ? *ColumnType->GetName() : TEXT("<null>"));
-		return Environment.MassContext.GetFragmentView(ColumnType).GetData();
+		const uint8* Base = reinterpret_cast<const uint8*>(Environment.MassContext.GetFragmentView(ColumnType).GetData());
+		return Environment.HasSubRange() ? Base + Environment.SubRangeStart * ColumnType->GetStructureSize() : Base;
 	}
 
 	void* FRowBatchInfoContextCapability::GetMutableColumnBatchAddress(
@@ -57,7 +58,8 @@ namespace UE::Editor::DataStorage::Queries
 	{
 		checkf(UE::Mass::IsA<FMassFragment>(ColumnType), TEXT("Trying to retrieve a data pointer for column '%s' which doesn't have data."),
 			ColumnType ? *ColumnType->GetName() : TEXT("<null>"));
-		return Environment.MassContext.GetMutableFragmentView(ColumnType).GetData();
+		uint8* Base = reinterpret_cast<uint8*>(Environment.MassContext.GetMutableFragmentView(ColumnType).GetData());
+		return Environment.HasSubRange() ? Base + Environment.SubRangeStart * ColumnType->GetStructureSize() : Base;
 	}
 
 
@@ -77,7 +79,8 @@ namespace UE::Editor::DataStorage::Queries
 	{
 		checkf(UE::Mass::IsA<FMassFragment>(ColumnType), TEXT("Trying to retrieve a data pointer for column '%s' which doesn't have data."),
 			ColumnType ? *ColumnType->GetName() : TEXT("<null>"));
-		return Environment.MassContext.GetFragmentView(ColumnType).GetData();
+		const uint8* Base = reinterpret_cast<const uint8*>(Environment.MassContext.GetFragmentView(ColumnType).GetData());
+		return Environment.HasSubRange() ? Base + Environment.SubRangeStart * ColumnType->GetStructureSize() : Base;
 	}
 
 	
diff --git a/Engine/Plugins/Experimental/EditorDataStorage/Source/TedsCore/Private/Queries/TypedElementExtendedQueryStore.cpp b/Engine/Plugins/Experimental/EditorDataStorage/Source/TedsCore/Private/Queries/TypedElementExtendedQueryStore.cpp
--- a/Engine/Plugins/Experimental/EditorDataStorage/Source/TedsCore/Private/Queries/TypedElementExtendedQueryStore.cpp
+++ b/Engine/Plugins/Experimental/EditorDataStorage/Source/TedsCore/Private/Queries/TypedElementExtendedQueryStore.cpp
@@ -757,15 +757,18 @@ namespace UE::Editor::DataStorage
 				*Query.Hierarchy.ToString());
 			if (HierarchyAccessInterface)
 			{
-				// Must bind the HierarchyDataColumn for Read to allow for GetParent
-				// Note: It would be better to do late binding to avoid the order dependency on query creation and hierarchy registration
-				// but this would require resetting a query's requirements after it has been bound to Mass or swapping out a new query
-				// when the hierarchy is registered
+				// Must bind the HierarchyDataColumn for Read to allow for GetParent.
+				// Relations-backed hierarchies store data on relation rows rather than
+				// participant rows, so GetHierarchyDataColumnType() returns null for them.
+				// Skip adding a null type -- the query still works for those backends.
 				// Note: Optional used to allow a query to match rows that are not currently participating in a relationship (ie. childless root objects)
 				const UScriptStruct* DataColumnType = HierarchyAccessInterface->GetHierarchyDataColumnType();
-				Query.SelectionTypes.Emplace(DataColumnType);
-				Query.SelectionAccessTypes.Emplace(EQueryAccessType::OptionalReadOnly);
-				Query.SelectionMetaData.Emplace(FColumnMetaData(DataColumnType, FColumnMetaData::EFlags::None));
+				if (DataColumnType)
+				{
+					Query.SelectionTypes.Emplace(DataColumnType);
+					Query.SelectionAccessTypes.Emplace(EQueryAccessType::OptionalReadOnly);
+					Query.SelectionMetaData.Emplace(FColumnMetaData(DataColumnType, FColumnMetaData::EFlags::None));
+				}
 			}
 		}
 		return true;
diff --git a/Engine/Plugins/Experimental/EditorDataStorage/Source/TedsCore/Private/TableManager.cpp b/Engine/Plugins/Experimental/EditorDataStorage/Source/TedsCore/Private/TableManager.cpp
--- a/Engine/Plugins/Experimental/EditorDataStorage/Source/TedsCore/Private/TableManager.cpp
+++ b/Engine/Plugins/Experimental/EditorDataStorage/Source/TedsCore/Private/TableManager.cpp
@@ -8,6 +8,7 @@
 #include "EditorDataStorageSettings.h"
 #include "Elements/Framework/TypedElementQueryContract.h"
 #include "MassEntityManager.h"
+#include "MassEntityRelations.h"
 #include "MassEntityTypes.h"
 #include "TypedElementUtils.h"
 
@@ -306,19 +307,49 @@ namespace UE::Editor::DataStorage
 		if (!bRegisteringNewTable)
 		{
 			const UScriptStruct* TableTag = nullptr;
+			const UScriptStruct* RelationTag = nullptr;
 
 			FMassArchetypeCompositionDescriptor Composition = MassEntityManager.GetArchetypeComposition(Archetype);
 			const FMassElementBitSet& Columns = Composition.GetElementsBitSet();
-			Columns.ExportTypes([&TableTag](const UScriptStruct* TagType)
+			Columns.ExportTypes([&TableTag, &RelationTag](const UScriptStruct* Type)
 				{
-					if (TagType->IsChildOf(FTableDynamicTag::StaticStruct()))
+					if (Type->IsChildOf(FTableDynamicTag::StaticStruct()))
 					{
-						TableTag = TagType;
+						TableTag = Type;
 						return false;
 					}
+					if (Type->IsChildOf(FMassRelation::StaticStruct()))
+					{
+						RelationTag = Type;
+					}
 					return true;
 				});
 
+			// Relation entity archetypes are created by Mass's FRelationManager and
+			// don't carry TEDS table tags. Look up the parent table using the per-
+			// relation-type tag name -- both the tag and the registered table share the
+			// same "TedsRelation_{name}" name, so no transformation is needed.
+			if (!TableTag && RelationTag)
+			{
+				FName ParentFName = FName(RelationTag->GetName());
+				TableHandle Parent = Find(ParentFName);
+				if (IsTableValid(Parent))
+				{
+					int32 NewTable = Tables.Emplace(FTableInfo
+						{
+							.Archetype = Archetype,
+							.Name = ParentFName,
+							.Type = ETableType::Variant,
+							.Tag = RelationTag,
+							.Parent = Parent
+						});
+					MassLookup.Add(Archetype, NewTable);
+					Tables[ConvertTableHandleToIndex(Parent)].Variants.Add(NewTable);
+					UE_LOG(LogEditorDataStorage, VeryVerbose, TEXT("New relation table variant for '%s'"), *ParentFName.ToString());
+				}
+				return;
+			}
+
 			checkf(TableTag, TEXT("Expected a table tag to be on a variant archetype. This function shouldn't be called on new archetypes."));
 			const FString& ParentName = TableTag->GetMetaData(TEXT("EditorDataStorage_DerivedFromDynamicTemplate"));
 			FName  ParentFName = FName(ParentName);
diff --git a/Engine/Plugins/Experimental/EditorDataStorage/Source/TedsCore/Private/TypedElementDatabase.cpp b/Engine/Plugins/Experimental/EditorDataStorage/Source/TedsCore/Private/TypedElementDatabase.cpp
--- a/Engine/Plugins/Experimental/EditorDataStorage/Source/TedsCore/Private/TypedElementDatabase.cpp
+++ b/Engine/Plugins/Experimental/EditorDataStorage/Source/TedsCore/Private/TypedElementDatabase.cpp
@@ -5,6 +5,8 @@
 #include <cinttypes>
 #include "DataStorage/Features.h"
 #include "DataStorage/Debug/Log.h"
+#include "Relations/EditorDataStorageRelationColumns.h"
+#include "Relations/EditorDataStorageRelationsImplementation.h"
 #include "Editor.h"
 #include "Elements/Interfaces/TypedElementDataStorageFactory.h"
 #include "Elements/Interfaces/TypedElementDataStorageInterface.h"
@@ -12,7 +14,10 @@
 #include "GlobalLock.h"
 #include "Math/UnrealMathUtility.h"
 #include "MassEntityEditorSubsystem.h"
+#include "MassEntityManager.h"
+#include "MassEntityRelations.h"
 #include "MassEntityTypes.h"
+#include "MassRelationManager.h"
 #include "MassProcessor.h"
 #include "MassSubsystemAccess.h"
 #include "Processors/TypedElementProcessorAdaptors.h"
@@ -232,6 +237,13 @@ void UEditorDataStorage::Initialize()
 		}
 	}
 
+	// Register the shared relations table. All relation entities use this base table
+	// (or variants of it when they acquire additional columns like FIntervalEncodedHierarchyMetadata).
+	{
+		const UScriptStruct* RelationFragment = FMassRelationFragment::StaticStruct();
+		Environment->GetTableManager().Register({&RelationFragment, 1}, TEXT("TEDSRelations"));
+	}
+
 	UE_LOGF(LogEditorDataStorage, Log, "Initialized TEDS Core.");
 }
 
@@ -1219,6 +1231,660 @@ bool UEditorDataStorage::IterateChildren(
 	return false;
 }
 
+//
+// Relations Interface Implementation
+//
+
+UE::Editor::DataStorage::RelationTypeHandle UEditorDataStorage::RegisterRelationType(
+	const UE::Editor::DataStorage::FRelationRegistrationParams& Params)
+{
+	using namespace UE::Editor::DataStorage;
+
+	if (!ActiveEditorEntityManager)
+	{
+		return InvalidRelationTypeHandle;
+	}
+
+	const RelationTypeHandle Handle = Environment->GetRelationAdapter().RegisterRelationType(
+		*ActiveEditorEntityManager, Environment->GetTableManager(), Params);
+
+	if (Handle == InvalidRelationTypeHandle)
+	{
+		return InvalidRelationTypeHandle;
+	}
+
+	if (Params.Traits.HierarchyMode != EHierarchyMode::Disabled)
+	{
+		Environment->GetHierarchicalRelationManager().RegisterHierarchyMode(Handle, Params.Traits.HierarchyMode, Params.Traits.IntervalGap);
+	}
+
+	// Generate optional change-notification columns and register FrameEnd cleanup processors.
+	auto GenerateAndRegisterCleanup = [this, Handle](bool bEnable, const UScriptStruct* TemplateType, const FName& Identifier)
+		-> const UScriptStruct*
+	{
+		if (!bEnable)
+		{
+			return nullptr;
+		}
+		using namespace UE::Editor::DataStorage::Queries;
+		FDynamicColumnDescription Desc{ .TemplateType = TemplateType, .Identifier = Identifier };
+		const UScriptStruct* Column = GenerateDynamicColumn(Desc);
+		if (Column)
+		{
+			// Remove the change column from all matching rows at FrameEnd.
+			RegisterQuery(
+				Select(
+					*FString::Printf(TEXT("Cleanup RelationChanged %s"), *Identifier.ToString()),
+					FPhaseAmble(FPhaseAmble::ELocation::Postamble, EQueryTickPhase::FrameEnd),
+					[Column](IQueryContext& Context, const RowHandle* Rows)
+					{
+						Context.RemoveColumns(
+							TConstArrayView<RowHandle>(Rows, Context.GetRowCount()),
+							{Column});
+					})
+				.Where()
+					.All(Column)
+				.Compile());
+		}
+		return Column;
+	};
+
+	const UScriptStruct* SubjectChanged = GenerateAndRegisterCleanup(
+		Params.bEnableSubjectChangedColumn,
+		FTedsRelationSubjectChanged_Template::StaticStruct(),
+		Params.Name);
+	const UScriptStruct* ObjectChanged = GenerateAndRegisterCleanup(
+		Params.bEnableObjectChangedColumn,
+		FTedsRelationObjectChanged_Template::StaticStruct(),
+		Params.Name);
+
+	if (SubjectChanged || ObjectChanged)
+	{
+		Environment->GetRelationAdapter().SetChangedColumns(Handle, SubjectChanged, ObjectChanged);
+	}
+
+	return Handle;
+}
+
+UE::Editor::DataStorage::RelationTypeHandle UEditorDataStorage::FindRelationType(const FName& InName) const
+{
+	using namespace UE::Editor::DataStorage;
+
+	if (!ActiveEditorEntityManager)
+	{
+		return InvalidRelationTypeHandle;
+	}
+
+	return Environment->GetRelationAdapter().FindRelationType(InName);
+}
+
+bool UEditorDataStorage::IsValidRelationType(UE::Editor::DataStorage::RelationTypeHandle Type) const
+{
+	if (!ActiveEditorEntityManager)
+	{
+		return false;
+	}
+
+	return Environment->GetRelationAdapter().IsValidRelationType(Type);
+}
+
+const UE::Editor::DataStorage::FTedsRelationTraits* UEditorDataStorage::GetRelationTypeTraits(
+	UE::Editor::DataStorage::RelationTypeHandle Type) const
+{
+	if (!ActiveEditorEntityManager)
+	{
+		return nullptr;
+	}
+
+	return Environment->GetRelationAdapter().GetTraits(Type);
+}
+
+void UEditorDataStorage::ListRelationTypes(
+	TFunctionRef<void(UE::Editor::DataStorage::RelationTypeHandle, const FName&)> Callback) const
+{
+	if (ActiveEditorEntityManager)
+	{
+		Environment->GetRelationAdapter().ListTypes(Callback);
+	}
+}
+
+const UScriptStruct* UEditorDataStorage::GetRelationSubjectChangedColumn(
+	UE::Editor::DataStorage::RelationTypeHandle Type) const
+{
+	return Environment ? Environment->GetRelationAdapter().GetSubjectChangedColumn(Type) : nullptr;
+}
+
+const UScriptStruct* UEditorDataStorage::GetRelationObjectChangedColumn(
+	UE::Editor::DataStorage::RelationTypeHandle Type) const
+{
+	return Environment ? Environment->GetRelationAdapter().GetObjectChangedColumn(Type) : nullptr;
+}
+
+UE::Editor::DataStorage::RowHandle UEditorDataStorage::CreateRelation(
+	UE::Editor::DataStorage::RelationTypeHandle Type,
+	UE::Editor::DataStorage::RowHandle Subject,
+	UE::Editor::DataStorage::RowHandle Object)
+{
+	using namespace UE::Editor::DataStorage;
+
+	if (!ActiveEditorEntityManager
+		|| !Environment->GetRelationAdapter().IsValidRelationType(Type)
+		|| !IsRowAssigned(Subject)
+		|| !IsRowAssigned(Object))
+	{
+		return InvalidRowHandle;
+	}
+
+	if (Subject == Object)
+	{
+		ensureMsgf(false, TEXT("CreateRelation: self-relations (Subject == Object) are not supported."));
+		return InvalidRowHandle;
+	}
+
+	const UE::Mass::FTypeHandle MassHandle = Environment->GetRelationAdapter().GetMassHandle(Type);
+
+	// Pre-check exclusivity (S3): reject if exclusive and relation already exists
+	const FTedsRelationTraits* Traits = Environment->GetRelationAdapter().GetTraits(Type);
+	if (Traits)
+	{
+		if (Traits->Subject.bExclusive)
+		{
+			// bExclusive on Subject means this Subject may have at most one Object across all relations
+			// of this type. Reject if the Subject already has ANY Object, not just the specific pair.
+			if (HasRelationObject(Type, Subject))
+			{
+				return InvalidRowHandle;
+			}
+		}
+	}
+
+	// Ensure Mass's RelationsDataMap is populated for this type before creating instances.
+	// FTypeManager::RegisterType only calls OnNewTypeRegistered (which populates RelationsDataMap)
+	// after bBuiltInTypesRegistered is set. If TEDS registered the type before Mass finished
+	// built-in processing, the data won't exist yet. We trigger it on first use.
+	if (!Environment->GetRelationAdapter().HasMassRelationData(Type))
+	{
+		ActiveEditorEntityManager->OnNewTypeRegistered(MassHandle);
+		Environment->GetRelationAdapter().SetHasMassRelationData(Type);
+	}
+
+	// Delegate to Mass
+	const FMassEntityHandle MassRelationEntity = ActiveEditorEntityManager->GetRelationManager()
+		.CreateRelationInstance(MassHandle, FMassEntityHandle::FromNumber(Subject), FMassEntityHandle::FromNumber(Object));
+
+	const RowHandle RelationRow = MassRelationEntity.AsNumber();
+
+	// If this is a hierarchical relation, add the appropriate metadata column and initialize it.
+	// Note: AddColumn causes a per-entity archetype migration. This is acceptable because:
+	// - Relation entities are internal (not user-visible participant rows)
+	// - Hierarchy changes are not a hot-path operation
+	// Future optimization (S2): Mass's FRelationManager::CreateRelationInstances has a hardcoded
+	// archetype (FMassRelationFragment + tag). To avoid per-creation migration, Mass would need
+	// an API extension to accept additional fragments at creation time. That change belongs in a
+	// separate CL to the MassEntity module.
+	if (RelationRow != InvalidRowHandle && Traits && Traits->HierarchyMode != EHierarchyMode::Disabled )
+	{
+		const UScriptStruct* MetadataStruct = (Traits->HierarchyMode == EHierarchyMode::IntervalEncoded)
+			? FIntervalEncodedHierarchyMetadata::StaticStruct()
+			: FWalkOnlyHierarchyMetadata::StaticStruct();
+		AddColumn(RelationRow, MetadataStruct);
+		Environment->GetHierarchicalRelationManager().InitializeHierarchyMetadata(*this, Type, RelationRow, Subject, Object);
+	}
+
+	// Stamp change-notification columns on participant rows if enabled.
+	if (RelationRow != InvalidRowHandle)
+	{
+		if (const UScriptStruct* SubjectChanged = Environment->GetRelationAdapter().GetSubjectChangedColumn(Type))
+		{
+			AddColumn(Subject, SubjectChanged);
+		}
+		if (const UScriptStruct* ObjectChanged = Environment->GetRelationAdapter().GetObjectChangedColumn(Type))
+		{
+			AddColumn(Object, ObjectChanged);
+		}
+	}
+
+	return RelationRow;
+}
+
+void UEditorDataStorage::BatchCreateRelations(
+	UE::Editor::DataStorage::RelationTypeHandle Type,
+	TConstArrayView<UE::Editor::DataStorage::RowHandle> Subjects,
+	TConstArrayView<UE::Editor::DataStorage::RowHandle> Objects,
+	TArray<UE::Editor::DataStorage::RowHandle>* OutRelationRows)
+{
+	using namespace UE::Editor::DataStorage;
+
+	if (!ActiveEditorEntityManager || Subjects.IsEmpty() || Objects.IsEmpty()
+		|| !Environment->GetRelationAdapter().IsValidRelationType(Type))
+	{
+		if (OutRelationRows)
+		{
+			OutRelationRows->Reset();
+		}
+		return;
+	}
+
+	const UE::Mass::FTypeHandle MassHandle = Environment->GetRelationAdapter().GetMassHandle(Type);
+
+	// Convert RowHandle arrays to FMassEntityHandle arrays
+	TArray<FMassEntityHandle> MassSubjects;
+	TArray<FMassEntityHandle> MassObjects;
+	MassSubjects.Reserve(Subjects.Num());
+	MassObjects.Reserve(Objects.Num());
+
+	for (RowHandle Row : Subjects)
+	{
+		MassSubjects.Add(FMassEntityHandle::FromNumber(Row));
+	}
+	for (RowHandle Row : Objects)
+	{
+		MassObjects.Add(FMassEntityHandle::FromNumber(Row));
+	}
+
+	// Ensure Mass's RelationsDataMap is populated (see CreateRelation for explanation)
+	if (!Environment->GetRelationAdapter().HasMassRelationData(Type))
+	{
+		ActiveEditorEntityManager->OnNewTypeRegistered(MassHandle);
+		Environment->GetRelationAdapter().SetHasMassRelationData(Type);
+	}
+
+	// Call FRelationManager::CreateRelationInstances directly (C3) to get handles back
+	TArray<FMassEntityHandle> MassRelationEntities =
+		ActiveEditorEntityManager->GetRelationManager().CreateRelationInstances(
+			MassHandle, MassSubjects, MassObjects);
+
+	// Initialize hierarchy metadata if needed
+	const FTedsRelationTraits* Traits = Environment->GetRelationAdapter().GetTraits(Type);
+	const bool bIsHierarchical = Traits && Traits->HierarchyMode != EHierarchyMode::Disabled;
+	const UScriptStruct* MetadataStruct = nullptr;
+	if (bIsHierarchical)
+	{
+		MetadataStruct = (Traits->HierarchyMode == EHierarchyMode::IntervalEncoded)
+			? FIntervalEncodedHierarchyMetadata::StaticStruct()
+			: FWalkOnlyHierarchyMetadata::StaticStruct();
+	}
+
+	if (OutRelationRows)
+	{
+		OutRelationRows->Reset(MassRelationEntities.Num());
+	}
+
+	for (int32 Index = 0; Index < MassRelationEntities.Num(); ++Index)
+	{
+		const RowHandle RelationRow = MassRelationEntities[Index].AsNumber();
+		if (OutRelationRows)
+		{
+			OutRelationRows->Add(RelationRow);
+		}
+
+		if (bIsHierarchical && RelationRow != InvalidRowHandle)
+		{
+			AddColumn(RelationRow, MetadataStruct);
+			// Match Mass's clamping behavior: FMath::Min(Array.Num()-1, PairIndex)
+			const RowHandle SubjectRow = Subjects[FMath::Min(Subjects.Num() - 1, Index)];
+			const RowHandle ObjectRow = Objects[FMath::Min(Objects.Num() - 1, Index)];
+			Environment->GetHierarchicalRelationManager().InitializeHierarchyMetadata(*this, Type, RelationRow, SubjectRow, ObjectRow);
+		}
+	}
+}
+
+bool UEditorDataStorage::DestroyRelation(
+	UE::Editor::DataStorage::RelationTypeHandle Type,
+	UE::Editor::DataStorage::RowHandle Subject,
+	UE::Editor::DataStorage::RowHandle Object)
+{
+	if (!ActiveEditorEntityManager || !Environment->GetRelationAdapter().HasMassRelationData(Type))
+	{
+		return false;
+	}
+
+	const UE::Mass::FTypeHandle MassHandle = Environment->GetRelationAdapter().GetMassHandle(Type);
+
+	// Clean up hierarchy metadata mapping before destroying the relation
+	if (Environment->GetRelationAdapter().IsHierarchical(Type))
+	{
+		Environment->GetHierarchicalRelationManager().RemoveHierarchyMetadata(Type, Subject);
+	}
+
+	const bool bDestroyed = ActiveEditorEntityManager->GetRelationManager().DestroyRelationInstance(
+		MassHandle, FMassEntityHandle::FromNumber(Subject), FMassEntityHandle::FromNumber(Object));
+
+	// Stamp change-notification columns on participant rows if the relation was actually destroyed.
+	if (bDestroyed)
+	{
+		if (const UScriptStruct* SubjectChanged = Environment->GetRelationAdapter().GetSubjectChangedColumn(Type))
+		{
+			if (IsRowAssigned(Subject))
+			{
+				AddColumn(Subject, SubjectChanged);
+			}
+		}
+		if (const UScriptStruct* ObjectChanged = Environment->GetRelationAdapter().GetObjectChangedColumn(Type))
+		{
+			if (IsRowAssigned(Object))
+			{
+				AddColumn(Object, ObjectChanged);
+			}
+		}
+	}
+
+	return bDestroyed;
+}
+
+bool UEditorDataStorage::HasRelation(
+	UE::Editor::DataStorage::RelationTypeHandle Type,
+	UE::Editor::DataStorage::RowHandle Subject,
+	UE::Editor::DataStorage::RowHandle Object) const
+{
+	if (!ActiveEditorEntityManager)
+	{
+		return false;
+	}
+
+	// Guard: Mass populates RelationData lazily on first CreateRelationInstance.
+	if (!Environment->GetRelationAdapter().HasMassRelationData(Type))
+	{
+		return false;
+	}
+	const UE::Mass::FTypeHandle MassHandle = Environment->GetRelationAdapter().GetMassHandle(Type);
+	return ActiveEditorEntityManager->GetRelationManager().IsSubjectOfRelation(
+		MassHandle, FMassEntityHandle::FromNumber(Subject), FMassEntityHandle::FromNumber(Object));
+}
+
+bool UEditorDataStorage::HasRelationObject(
+	UE::Editor::DataStorage::RelationTypeHandle Type,
+	UE::Editor::DataStorage::RowHandle Subject) const
+{
+	using namespace UE::Editor::DataStorage;
+	if (!ActiveEditorEntityManager || !Environment->GetRelationAdapter().HasMassRelationData(Type))
+	{
+		return false;
+	}
+	const UE::Mass::FTypeHandle MassHandle = Environment->GetRelationAdapter().GetMassHandle(Type);
+	return ActiveEditorEntityManager->GetRelationManager().HasRelationObjects(
+		MassHandle, FMassEntityHandle::FromNumber(Subject));
+}
+
+bool UEditorDataStorage::HasRelationSubject(
+	UE::Editor::DataStorage::RelationTypeHandle Type,
+	UE::Editor::DataStorage::RowHandle Object) const
+{
+	using namespace UE::Editor::DataStorage;
+	if (!ActiveEditorEntityManager || !Environment->GetRelationAdapter().HasMassRelationData(Type))
+	{
+		return false;
+	}
+	const UE::Mass::FTypeHandle MassHandle = Environment->GetRelationAdapter().GetMassHandle(Type);
+	return ActiveEditorEntityManager->GetRelationManager().HasRelationSubjects(
+		MassHandle, FMassEntityHandle::FromNumber(Object));
+}
+
+void UEditorDataStorage::GetRelationObjects(
+	UE::Editor::DataStorage::RelationTypeHandle Type,
+	UE::Editor::DataStorage::RowHandle Subject,
+	TArray<UE::Editor::DataStorage::RowHandle>& OutObjects) const
+{
+	using namespace UE::Editor::DataStorage;
+	OutObjects.Reset();
+
+	if (!ActiveEditorEntityManager)
+	{
+		return;
+	}
+
+	if (!Environment->GetRelationAdapter().HasMassRelationData(Type))
+	{
+		return;
+	}
+	const UE::Mass::FTypeHandle MassHandle = Environment->GetRelationAdapter().GetMassHandle(Type);
+	TArray<FMassEntityHandle> MassObjects = ActiveEditorEntityManager->GetRelationManager().GetRelationObjects(
+		MassHandle, FMassEntityHandle::FromNumber(Subject));
+
+	OutObjects.Reserve(MassObjects.Num());
+	for (const FMassEntityHandle& Handle : MassObjects)
+	{
+		OutObjects.Add(Handle.AsNumber());
+	}
+}
+
+void UEditorDataStorage::GetRelationSubjects(
+	UE::Editor::DataStorage::RelationTypeHandle Type,
+	UE::Editor::DataStorage::RowHandle Object,
+	TArray<UE::Editor::DataStorage::RowHandle>& OutSubjects) const
+{
+	using namespace UE::Editor::DataStorage;
+	OutSubjects.Reset();
+
+	if (!ActiveEditorEntityManager)
+	{
+		return;
+	}
+
+	if (!Environment->GetRelationAdapter().HasMassRelationData(Type))
+	{
+		return;
+	}
+	const UE::Mass::FTypeHandle MassHandle = Environment->GetRelationAdapter().GetMassHandle(Type);
+	TArray<FMassEntityHandle> MassSubjects = ActiveEditorEntityManager->GetRelationManager().GetRelationSubjects(
+		MassHandle, FMassEntityHandle::FromNumber(Object));
+
+	OutSubjects.Reserve(MassSubjects.Num());
+	for (const FMassEntityHandle& Handle : MassSubjects)
+	{
+		OutSubjects.Add(Handle.AsNumber());
+	}
+}
+
+UE::Editor::DataStorage::RowHandle UEditorDataStorage::GetRelationObject(
+	UE::Editor::DataStorage::RelationTypeHandle Type,
+	UE::Editor::DataStorage::RowHandle Subject) const
+{
+	using namespace UE::Editor::DataStorage;
+
+	TArray<RowHandle> Objects;
+	GetRelationObjects(Type, Subject, Objects);
+
+	// For relations with exclusive Object, there's at most one object
+	return Objects.Num() > 0 ? Objects[0] : InvalidRowHandle;
+}
+
+UE::Editor::DataStorage::RowHandle UEditorDataStorage::GetRelationSubject(
+	UE::Editor::DataStorage::RelationTypeHandle Type,
+	UE::Editor::DataStorage::RowHandle Object) const
+{
+	using namespace UE::Editor::DataStorage;
+
+	TArray<RowHandle> Subjects;
+	GetRelationSubjects(Type, Object, Subjects);
+
+	// For relations with exclusive Subject, there's at most one subject
+	return Subjects.Num() > 0 ? Subjects[0] : InvalidRowHandle;
+}
+
+bool UEditorDataStorage::IsDescendantOf(
+	UE::Editor::DataStorage::RelationTypeHandle Type,
+	UE::Editor::DataStorage::RowHandle Descendant,
+	UE::Editor::DataStorage::RowHandle Ancestor,
+	bool bIncludeSelf) const
+{
+	if (!ActiveEditorEntityManager)
+	{
+		return false;
+	}
+
+	return Environment->GetHierarchicalRelationManager().IsDescendantOf(*this, Type, Descendant, Ancestor, bIncludeSelf);
+}
+
+UE::Editor::DataStorage::RowHandle UEditorDataStorage::GetHierarchyRoot(
+	UE::Editor::DataStorage::RelationTypeHandle Type,
+	UE::Editor::DataStorage::RowHandle Row) const
+{
+	if (!ActiveEditorEntityManager)
+	{
+		return UE::Editor::DataStorage::InvalidRowHandle;
+	}
+
+	return Environment->GetHierarchicalRelationManager().GetHierarchyRoot(*this, Type, Row);
+}
+
+int32 UEditorDataStorage::GetHierarchyDepth(
+	UE::Editor::DataStorage::RelationTypeHandle Type,
+	UE::Editor::DataStorage::RowHandle Row) const
+{
+	if (!ActiveEditorEntityManager)
+	{
+		return 0;
+	}
+
+	return Environment->GetHierarchicalRelationManager().GetHierarchyDepth(*this, Type, Row);
+}
+
+void UEditorDataStorage::GetDescendants(
+	UE::Editor::DataStorage::RelationTypeHandle Type,
+	UE::Editor::DataStorage::RowHandle Ancestor,
+	TArray<UE::Editor::DataStorage::RowHandle>& OutDescendants) const
+{
+	if (!ActiveEditorEntityManager)
+	{
+		OutDescendants.Reset();
+		return;
+	}
+
+	Environment->GetHierarchicalRelationManager().GetDescendants(*this, Type, Ancestor, OutDescendants);
+}
+
+void UEditorDataStorage::GetAncestors(
+	UE::Editor::DataStorage::RelationTypeHandle Type,
+	UE::Editor::DataStorage::RowHandle Descendant,
+	TArray<UE::Editor::DataStorage::RowHandle>& OutAncestors) const
+{
+	if (!ActiveEditorEntityManager)
+	{
+		OutAncestors.Reset();
+		return;
+	}
+
+	Environment->GetHierarchicalRelationManager().GetAncestors(*this, Type, Descendant, OutAncestors);
+}
+
+void UEditorDataStorage::TraverseDescendants(
+	UE::Editor::DataStorage::RelationTypeHandle Type,
+	UE::Editor::DataStorage::RowHandle StartRow,
+	FRelationTraversalCallback Callback,
+	ETraversalOrder Order,
+	int32 MaxDepth) const
+{
+	if (!ActiveEditorEntityManager)
+	{
+		return;
+	}
+
+	Environment->GetHierarchicalRelationManager().TraverseDescendants(*this, Type, StartRow, Callback, Order, MaxDepth);
+}
+
+void UEditorDataStorage::ForEachRelationObject(
+	UE::Editor::DataStorage::RelationTypeHandle Type,
+	UE::Editor::DataStorage::RowHandle Subject,
+	TFunctionRef<void(UE::Editor::DataStorage::RowHandle)> Callback) const
+{
+	using namespace UE::Editor::DataStorage;
+
+	if (!ActiveEditorEntityManager)
+	{
+		return;
+	}
+
+	if (!Environment->GetRelationAdapter().HasMassRelationData(Type))
+	{
+		return;
+	}
+
+	const UE::Mass::FTypeHandle MassHandle = Environment->GetRelationAdapter().GetMassHandle(Type);
+	ActiveEditorEntityManager->GetRelationManager().ForEachRelationObject(
+		MassHandle, FMassEntityHandle::FromNumber(Subject),
+		[&Callback](FMassEntityHandle MassHandle)
+		{
+			Callback(MassHandle.AsNumber());
+		});
+}
+
+void UEditorDataStorage::ForEachRelationSubject(
+	UE::Editor::DataStorage::RelationTypeHandle Type,
+	UE::Editor::DataStorage::RowHandle Object,
+	TFunctionRef<void(UE::Editor::DataStorage::RowHandle)> Callback) const
+{
+	using namespace UE::Editor::DataStorage;
+
+	if (!ActiveEditorEntityManager)
+	{
+		return;
+	}
+
+	if (!Environment->GetRelationAdapter().HasMassRelationData(Type))
+	{
+		return;
+	}
+
+	const UE::Mass::FTypeHandle MassHandle = Environment->GetRelationAdapter().GetMassHandle(Type);
+	ActiveEditorEntityManager->GetRelationManager().ForEachRelationSubject(
+		MassHandle, FMassEntityHandle::FromNumber(Object),
+		[&Callback](FMassEntityHandle MassHandle)
+		{
+			Callback(MassHandle.AsNumber());
+		});
+}
+
+void UEditorDataStorage::TraverseAncestors(
+	UE::Editor::DataStorage::RelationTypeHandle Type,
+	UE::Editor::DataStorage::RowHandle Descendant,
+	TFunctionRef<bool(UE::Editor::DataStorage::RowHandle, int32)> Callback) const
+{
+	using namespace UE::Editor::DataStorage;
+
+	constexpr int32 MaxAncestorWalkDepth = 65536;
+	RowHandle Current = GetRelationObject(Type, Descendant);
+	int32 Depth = 1;
+	while (Current != InvalidRowHandle && Depth < MaxAncestorWalkDepth)
+	{
+		if (!Callback(Current, Depth))
+		{
+			break;
+		}
+		Current = GetRelationObject(Type, Current);
+		++Depth;
+	}
+}
+
+int32 UEditorDataStorage::ComputeDescendantCount(
+	UE::Editor::DataStorage::RelationTypeHandle Type,
+	UE::Editor::DataStorage::RowHandle Ancestor) const
+{
+	int32 Count = 0;
+	TraverseDescendants(Type, Ancestor,
+		[&Count](UE::Editor::DataStorage::RowHandle, UE::Editor::DataStorage::RowHandle, int32)
+		{
+			++Count;
+		});
+	return Count;
+}
+
+bool UEditorDataStorage::HasRelationSubjects(
+	UE::Editor::DataStorage::RelationTypeHandle Type,
+	UE::Editor::DataStorage::RowHandle Object) const
+{
+	using namespace UE::Editor::DataStorage;
+
+	if (!ActiveEditorEntityManager || !Environment->GetRelationAdapter().HasMassRelationData(Type))
+	{
+		return false;
+	}
+
+	const UE::Mass::FTypeHandle MassHandle = Environment->GetRelationAdapter().GetMassHandle(Type);
+	return ActiveEditorEntityManager->GetRelationManager().HasRelationSubjects(
+		MassHandle, FMassEntityHandle::FromNumber(Object));
+}
+
 void UEditorDataStorage::RegisterTickGroup(
 	FName GroupName, UE::Editor::DataStorage::EQueryTickPhase Phase, FName BeforeGroup, FName AfterGroup, UE::Editor::DataStorage::EExecutionMode ExecutionMode)
 {
@@ -1674,13 +2340,20 @@ void UEditorDataStorage::FinalizePhase(UE::Editor::DataStorage::EQueryTickPhase
 		// TODO: This requires Mass to tell TEDS it's about to flush its deferred commands. Right now this gets called after the
 		// deferred commands are run, which require exclusive access.
 		//FGlobalLock::InternalSharedUnlock();
-		
+
 		// The preamble queries are all run on the game thread. While this is true it's safe to take a global write lock.
 		// If there's a performance loss because this lock is held too long, the work in RunPhasePostambleQueries can be split
 		// into a step that runs the queries and uses a shared lock and one that executes the command buffer with an exclusive lock.
 		FScopedExclusiveLock Lock(EGlobalLockScope::Public);
 		Environment->GetQueryStore().RunPhasePostambleQueries(*ActiveEditorEntityManager, *Environment, Phase, DeltaTime);
 
+		// Process interval rebalancing for hierarchical relations at the end of each frame.
+		// This is lightweight when there are no tagged rows (just a query that returns empty).
+		if (Phase == EQueryTickPhase::FrameEnd )
+		{
+			Environment->GetHierarchicalRelationManager().ProcessRebalancing(*this);
+		}
+
 		// Flush any custom commands that have queued up in between phases.
 		ActiveCallsCount--;
 		checkf(ActiveCallsCount == 0, TEXT("During phase finalization there should be no TEDS running."));
