[TEDS] Fixed the result counter in FQueryResult not being atomic.

The result counter in FQueryResult is not atomic which can cause issues when processing is distributed over multiple threads.

#Synced-CL 52688201

diff --git a/Engine/Plugins/Experimental/EditorDataStorage/Source/TedsCore/Private/Processors/TypedElementProcessorAdaptors.cpp b/Engine/Plugins/Experimental/EditorDataStorage/Source/TedsCore/Private/Processors/TypedElementProcessorAdaptors.cpp
--- a/Engine/Plugins/Experimental/EditorDataStorage/Source/TedsCore/Private/Processors/TypedElementProcessorAdaptors.cpp
+++ b/Engine/Plugins/Experimental/EditorDataStorage/Source/TedsCore/Private/Processors/TypedElementProcessorAdaptors.cpp
@@ -1927,8 +1927,10 @@ UE::Editor::DataStorage::FQueryResult FTypedElementQueryProcessorData::Execute(
 	if (!QueryContext_ProcessorApi::SupportsCapabilities(Callback.CapabilityMask))
 	{
 		QueryContext_ProcessorApi::StaticReportCompatibility(Callback.CapabilityMask);
-		return FQueryResult{ .Completed = FQueryResult::ECompletion::Unsupported };
-	}	
+		FQueryResult UnsupportedResult;
+		UnsupportedResult.Completed = FQueryResult::ECompletion::Unsupported;
+		return UnsupportedResult;
+	}
 
 	if (EnumHasAnyFlags(ExecutionFlags, ERunQueryFlags::AllowBoundQueries) || !Description.Callback.Function)
 	{
@@ -1979,7 +1981,9 @@ UE::Editor::DataStorage::FQueryResult FTypedElementQueryProcessorData::Execute(
 	if (!QueryContext_ProcessorApi::SupportsCapabilities(Callback.CapabilityMask))
 	{
 		QueryContext_ProcessorApi::StaticReportCompatibility(Callback.CapabilityMask);
-		return FQueryResult{ .Completed = FQueryResult::ECompletion::Unsupported };
+		FQueryResult UnsupportedResult;
+		UnsupportedResult.Completed = FQueryResult::ECompletion::Unsupported;
+		return UnsupportedResult;
 	}
 
 	FQueryResult Result;
diff --git a/Engine/Plugins/Experimental/EditorDataStorageFeatures/Source/TedsDebugger/Private/Widgets/QueryEditor/TedsQueryEditorResultsView.cpp b/Engine/Plugins/Experimental/EditorDataStorageFeatures/Source/TedsDebugger/Private/Widgets/QueryEditor/TedsQueryEditorResultsView.cpp
--- a/Engine/Plugins/Experimental/EditorDataStorageFeatures/Source/TedsDebugger/Private/Widgets/QueryEditor/TedsQueryEditorResultsView.cpp
+++ b/Engine/Plugins/Experimental/EditorDataStorageFeatures/Source/TedsDebugger/Private/Widgets/QueryEditor/TedsQueryEditorResultsView.cpp
@@ -103,7 +103,7 @@ namespace UE::Editor::DataStorage::Debug::QueryEditor
 					FQueryResult QueryResult = TedsInterface.RunQuery(CountQueryHandle);
 					if (QueryResult.Completed == FQueryResult::ECompletion::Fully)
 					{
-						FString String = FString::Printf(TEXT("Element Count: %u"), QueryResult.Count);
+						FString String = FString::Printf(TEXT("Element Count: %u"), QueryResult.Count.load());
 						return FText::FromString(String);
 					}
 					else
diff --git a/Engine/Source/Runtime/TypedElementFramework/Public/DataStorage/Queries/Types.h b/Engine/Source/Runtime/TypedElementFramework/Public/DataStorage/Queries/Types.h
--- a/Engine/Source/Runtime/TypedElementFramework/Public/DataStorage/Queries/Types.h
+++ b/Engine/Source/Runtime/TypedElementFramework/Public/DataStorage/Queries/Types.h
@@ -4,6 +4,7 @@
 
 #include "HAL/Platform.h"
 #include "Misc/EnumClassFlags.h"
+#include <atomic>
 
 namespace UE::Editor::DataStorage
 {
@@ -135,7 +136,33 @@ namespace UE::Editor::DataStorage
 			MissingDependency
 		};
 
-		uint32 Count{ 0 }; /** The number of rows were processed. */
+		FQueryResult() = default;
+		FQueryResult(const FQueryResult& Other)
+			: Count(Other.Count.load())
+			, Completed(Other.Completed)
+		{}
+		
+		FQueryResult(FQueryResult&& Other) noexcept
+			: Count(Other.Count.load())
+			, Completed(Other.Completed)
+		{}
+		
+		FQueryResult& operator=(const FQueryResult& Other)
+		{
+			Count.store(Other.Count.load());
+			Completed = Other.Completed;
+			return *this;
+		}
+		
+		FQueryResult& operator=(FQueryResult&& Other) noexcept
+		{
+			Count.store(Other.Count.load());
+			Completed = Other.Completed;
+			return *this;
+		}
+
+		/** The number of rows that were processed. */
+		std::atomic<uint32> Count{ 0 };
 		ECompletion Completed{ ECompletion::Unavailable };
 	};
 } // namespace UE::Editor::DataStorage
